Initial implementation of the Chaos Visual Debugger plugin

This CL contains the plugin which adds the Window tool for visualizing recorded physics simulations and an engine module that defines the data format (although for this first submission is just a placeholder).

The recording implementation is not included as it is pretty much a placeholder like the data format itself and will change shortly.

The tool can be open from the Menu Tools->Debug->Chaos Visual Debugger

Included functionality:

- 3D Viewport to see the recorded elements of the simulation (Independent of the Editor). The current placeholder data format only contains the Position/Rotation of particles so the visualization is limited to that for now.
- Simple playback controls
- Instances of World Outliner, Details panel, and Output log
- Ability to load and reproduce cvd files (format defined in this same CL)

Current issues and limitations:
- Outliner/Viewport and details panel are still using global selection events so what is selected in the tool also affects the Editor. This will be addressed in the next commits when we implement derived versions of these elements and change to a local selection event.
- All the UI icons are placeholders, copies from base engine icons (some slightly modified).
- Settings button in the UI is not implemented

Everything is pretty much in flux. The next focus and likely next CL will be one dedicated to implementing a proper first version of the recording data format + the recording logic and after that refocus on the visualization tool again. Also, it is likely that the data format types will be moved into the Chaos module to be able o re-use existing types.

diff --git a/Engine/Plugins/Experimental/ChaosVD/ChaosVD.uplugin b/Engine/Plugins/Experimental/ChaosVD/ChaosVD.uplugin
new file mode 100644
--- /dev/null
+++ b/Engine/Plugins/Experimental/ChaosVD/ChaosVD.uplugin
@@ -0,0 +1,25 @@
+{
+	"FileVersion": 3,
+	"Version": 1,
+	"VersionName": "0.1",
+	"FriendlyName": "Chaos Visual Debugger",
+	"Description": "Enables support for Visual debugging of Chaos Physics simulations",
+	"Category": "Physics",
+	"CreatedBy" : "Epic Games, Inc.",
+	"CreatedByURL" : "https://epicgames.com",
+	"DocsURL": "",
+	"MarketplaceURL": "",
+	"SupportURL": "",
+	"EnabledByDefault": false,
+	"CanContainContent": true,
+	"IsBetaVersion": false,
+	"IsExperimentalVersion": true,
+	"Installed": false,
+	"Modules": [
+		{
+			"Name": "ChaosVD",
+			"Type": "Editor",
+			"LoadingPhase": "Default"
+		}
+	]
+}
\ No newline at end of file
diff --git a/Engine/Plugins/Experimental/ChaosVD/Config/Engine.ini b/Engine/Plugins/Experimental/ChaosVD/Config/Engine.ini
new file mode 100644
--- /dev/null
+++ b/Engine/Plugins/Experimental/ChaosVD/Config/Engine.ini
@@ -0,0 +1,3 @@
+﻿[/Script/ChaosVD.ChaosVDEditorSettings]
+DebugMesh=/ChaosVD/SM_1M_Cube.SM_1M_Cube
+BasePhysicsVDWorld=/ChaosVD/ChaosVDBaseWorld.ChaosVDBaseWorld
\ No newline at end of file
diff --git a/Engine/Plugins/Experimental/ChaosVD/Source/ChaosVD/ChaosVD.Build.cs b/Engine/Plugins/Experimental/ChaosVD/Source/ChaosVD/ChaosVD.Build.cs
new file mode 100644
--- /dev/null
+++ b/Engine/Plugins/Experimental/ChaosVD/Source/ChaosVD/ChaosVD.Build.cs
@@ -0,0 +1,35 @@
+// Copyright Epic Games, Inc. All Rights Reserved.
+
+using UnrealBuildTool;
+
+public class ChaosVD : ModuleRules
+{
+	public ChaosVD(ReadOnlyTargetRules Target) : base(Target)
+	{
+		PublicDependencyModuleNames.AddRange(
+			new string[]
+			{
+				"Core",
+				"CoreUObject",
+				"Engine",
+				"ChaosVDRuntime",
+			});
+
+		PrivateDependencyModuleNames.AddRange(
+			new string[]
+			{
+				"DesktopPlatform",
+				"Projects",
+				"InputCore",
+				"EditorFramework",
+				"UnrealEd",
+				"ToolMenus",
+				"Slate",
+				"SlateCore", 
+				"OutputLog", 
+				"SceneOutliner", 
+				"WorkspaceMenuStructure",
+			}
+			);
+	}
+}
diff --git a/Engine/Plugins/Experimental/ChaosVD/Source/ChaosVD/Private/ChaosVDEngine.cpp b/Engine/Plugins/Experimental/ChaosVD/Source/ChaosVD/Private/ChaosVDEngine.cpp
new file mode 100644
--- /dev/null
+++ b/Engine/Plugins/Experimental/ChaosVD/Source/ChaosVD/Private/ChaosVDEngine.cpp
@@ -0,0 +1,37 @@
+﻿// Copyright Epic Games, Inc. All Rights Reserved.
+
+#include "ChaosVD/Public/ChaosVDEngine.h"
+
+#include "ChaosVDPlaybackController.h"
+#include "ChaosVDScene.h"
+
+void FChaosVDEngine::Initialize()
+{
+	if (bIsInitialized)
+	{
+		return;
+	}
+
+	// Create an Empty Scene
+	//TODO: Handle multiple scenes. We will need it to represent multiple worlds
+	CurrentScene = MakeShared<FChaosVDScene>();
+	CurrentScene->Initialize();
+
+	PlaybackController = MakeShared<FChaosVDPlaybackController>(CurrentScene);
+
+	bIsInitialized = true;
+}
+
+void FChaosVDEngine::DeInitialize()
+{
+	if (!bIsInitialized)
+	{
+		return;
+	}
+
+	CurrentScene->DeInitialize();
+}
+bool FChaosVDEngine::Tick(float DeltaTime)
+{
+	return true;
+}
diff --git a/Engine/Plugins/Experimental/ChaosVD/Source/ChaosVD/Private/ChaosVDModule.cpp b/Engine/Plugins/Experimental/ChaosVD/Source/ChaosVD/Private/ChaosVDModule.cpp
new file mode 100644
--- /dev/null
+++ b/Engine/Plugins/Experimental/ChaosVD/Source/ChaosVD/Private/ChaosVDModule.cpp
@@ -0,0 +1,70 @@
+// Copyright Epic Games, Inc. All Rights Reserved.
+
+#include "ChaosVDModule.h"
+#include "ChaosVDStyle.h"
+#include "ChaosVDParticleActorCustomization.h"
+#include "ChaosVDTabsIDs.h"
+#include "Widgets/Docking/SDockTab.h"
+#include "Widgets/SChaosVDMainTab.h"
+#include "PropertyEditorModule.h"
+#include "ChaosVDEngine.h"
+#include "WorkspaceMenuStructure.h"
+#include "WorkspaceMenuStructureModule.h"
+
+#define LOCTEXT_NAMESPACE "ChaosVisualDebugger"
+
+void FChaosVDModule::StartupModule()
+{	
+	FChaosVDStyle::Initialize();
+	FChaosVDStyle::ReloadTextures();
+
+	RegisterClassesCustomDetails();
+
+	FGlobalTabmanager::Get()->RegisterNomadTabSpawner(FChaosVDTabID::ChaosVisualDebuggerTab, FOnSpawnTab::CreateRaw(this, &FChaosVDModule::SpawnMainTab))
+								.SetDisplayName(LOCTEXT("VisualDebuggerTabTitle", "Chaos Visual Debugger"))
+								.SetTooltipText(LOCTEXT("VisualDebuggerTabDesc", "Opens the Chaos Visual Debugger window"))
+								//TODO: Hook up the final icon
+								.SetIcon(FSlateIcon(FAppStyle::GetAppStyleSetName(), "CollisionAnalyzer.TabIcon"))
+								.SetGroup(WorkspaceMenu::GetMenuStructure().GetDeveloperToolsDebugCategory());
+}
+
+void FChaosVDModule::ShutdownModule()
+{
+	FChaosVDStyle::Shutdown();
+
+	FGlobalTabmanager::Get()->UnregisterNomadTabSpawner(FChaosVDTabID::ChaosVisualDebuggerTab);
+}
+
+void FChaosVDModule::RegisterClassesCustomDetails() const
+{
+	FPropertyEditorModule& PropertyModule = FModuleManager::GetModuleChecked<FPropertyEditorModule>("PropertyEditor");
+	PropertyModule.RegisterCustomClassLayout("ChaosVDParticleActor", FOnGetDetailCustomizationInstance::CreateStatic(&FChaosVDParticleActorCustomization::MakeInstance));
+}
+
+TSharedRef<SDockTab> FChaosVDModule::SpawnMainTab(const FSpawnTabArgs& Args) const
+{
+	TSharedRef<SDockTab> MainTabInstance =
+		SNew(SDockTab)
+		.TabRole(ETabRole::MajorTab)
+		.Label(LOCTEXT("MainTabLabel", "Chaos Visual Debugger"))
+		.ToolTipText(LOCTEXT("MainTabToolTip", "The Chaos Visual debugger is under development"));
+
+	// Initialize the Chaos VD Engine instance this tab will represent
+	// For now its lifetime will be controlled by this tab
+	const TSharedPtr<FChaosVDEngine> ChaosVDEngineInstance = MakeShared<FChaosVDEngine>();
+	ChaosVDEngineInstance->Initialize();
+
+	MainTabInstance->SetContent
+	(
+		SNew(SChaosVDMainTab, ChaosVDEngineInstance)
+			.OwnerTab(MainTabInstance.ToSharedPtr())
+	);
+
+	MainTabInstance->SetTabIcon(FChaosVDStyle::Get().GetBrush("TabIconPlaybackViewport"));
+
+	return MainTabInstance;
+}
+
+#undef LOCTEXT_NAMESPACE
+	
+IMPLEMENT_MODULE(FChaosVDModule, ChaosVD)
\ No newline at end of file
diff --git a/Engine/Plugins/Experimental/ChaosVD/Source/ChaosVD/Private/ChaosVDObjectDetailsTab.cpp b/Engine/Plugins/Experimental/ChaosVD/Source/ChaosVD/Private/ChaosVDObjectDetailsTab.cpp
new file mode 100644
--- /dev/null
+++ b/Engine/Plugins/Experimental/ChaosVD/Source/ChaosVD/Private/ChaosVDObjectDetailsTab.cpp
@@ -0,0 +1,65 @@
+﻿// Copyright Epic Games, Inc. All Rights Reserved.
+
+#include "ChaosVD/Public/ChaosVDObjectDetailsTab.h"
+
+#include "ChaosVDStyle.h"
+#include "ChaosVDTabsIDs.h"
+#include "Modules/ModuleManager.h"
+#include "PropertyEditorModule.h"
+#include "Selection.h"
+#include "Templates/SharedPointer.h"
+#include "UnrealEd/Public/Editor.h"
+#include "Widgets/Docking/SDockTab.h"
+
+#define LOCTEXT_NAMESPACE "ChaosVisualDebugger"
+
+FChaosVDObjectDetailsTab::~FChaosVDObjectDetailsTab()
+{
+	if (GEditor == nullptr || GEditor->GetSelectedActors() == nullptr)
+	{
+		return;
+	}
+
+	GEditor->GetSelectedActors()->SelectionChangedEvent.Remove(SelectionDelegateHandle);
+}
+
+TSharedRef<SDockTab> FChaosVDObjectDetailsTab::HandleTabSpawned(const FSpawnTabArgs& Args)
+{
+	FPropertyEditorModule& PropertyEditorModule = FModuleManager::GetModuleChecked<FPropertyEditorModule>("PropertyEditor");
+	FDetailsViewArgs DetailsViewArgs;
+	DetailsViewArgs.bAllowSearch = false;
+	DetailsViewArgs.NameAreaSettings = FDetailsViewArgs::HideNameArea;
+
+	DetailsPanel = PropertyEditorModule.CreateDetailView(DetailsViewArgs);
+
+	//TODO: This is just for testing. We will not use global selection events in the final version as these affect the entire editor
+	SelectionDelegateHandle = GEditor->GetSelectedActors()->SelectionChangedEvent.AddLambda(
+	[this](UObject* Object)
+	{
+		TArray<AActor*> SelectedActors;
+		GEditor->GetSelectedActors()->GetSelectedObjects<AActor>(SelectedActors);
+
+		UpdateSelectedObject(SelectedActors.Num() > 0 ? SelectedActors[0] : nullptr);
+	} );
+
+	TSharedRef<SDockTab> DetailsPanelTab =
+		SNew(SDockTab)
+		.TabRole(ETabRole::MajorTab)
+		.Label(LOCTEXT("DetailsPanel", "Details"))
+		.ToolTipText(LOCTEXT("DetailsPanelToolTip", "See the details of the selected object"));
+	DetailsPanelTab->SetContent
+	(
+		DetailsPanel.ToSharedRef()
+	);
+
+	DetailsPanelTab->SetTabIcon(FChaosVDStyle::Get().GetBrush("TabIconDetailsPanel"));
+
+	return DetailsPanelTab;
+}
+
+void FChaosVDObjectDetailsTab::UpdateSelectedObject(AActor* NewObject) const
+{
+	DetailsPanel->SetObject(NewObject, true);
+}
+
+#undef LOCTEXT_NAMESPACE
diff --git a/Engine/Plugins/Experimental/ChaosVD/Source/ChaosVD/Private/ChaosVDOutputLogTab.cpp b/Engine/Plugins/Experimental/ChaosVD/Source/ChaosVD/Private/ChaosVDOutputLogTab.cpp
new file mode 100644
--- /dev/null
+++ b/Engine/Plugins/Experimental/ChaosVD/Source/ChaosVD/Private/ChaosVDOutputLogTab.cpp
@@ -0,0 +1,33 @@
+﻿// Copyright Epic Games, Inc. All Rights Reserved.
+
+#include "ChaosVDOutputLogTab.h"
+
+#include "OutputLogCreationParams.h"
+#include "OutputLogModule.h"
+#include "Framework/Docking/TabManager.h"
+#include "Widgets/Docking/SDockTab.h"
+
+#define LOCTEXT_NAMESPACE "ChaosVisualDebugger"
+
+TSharedRef<SDockTab> FChaosVDOutputLogTab::HandleTabSpawned(const FSpawnTabArgs& Args)
+{
+	FOutputLogCreationParams Params;
+	Params.bCreateDockInLayoutButton = true;
+	Params.SettingsMenuCreationFlags = EOutputLogSettingsMenuFlags::SkipClearOnPie
+		| EOutputLogSettingsMenuFlags::SkipOpenSourceButton
+		| EOutputLogSettingsMenuFlags::SkipEnableWordWrapping; // Checkbox relies on saving an editor config file and does not work correctly
+	
+	TSharedRef<SDockTab> OutputLogTab =
+		SNew(SDockTab)
+		.TabRole(ETabRole::MajorTab)
+		.Label(LOCTEXT("OutputLogTabLabel", "Output Log"));
+	
+	OutputLogTab->SetContent
+	(
+		FOutputLogModule::Get().MakeOutputLogWidget(Params)
+	);
+
+	return OutputLogTab;
+}
+
+#undef LOCTEXT_NAMESPACE
diff --git a/Engine/Plugins/Experimental/ChaosVD/Source/ChaosVD/Private/ChaosVDParticleActor.cpp b/Engine/Plugins/Experimental/ChaosVD/Source/ChaosVD/Private/ChaosVDParticleActor.cpp
new file mode 100644
--- /dev/null
+++ b/Engine/Plugins/Experimental/ChaosVD/Source/ChaosVD/Private/ChaosVDParticleActor.cpp
@@ -0,0 +1,43 @@
+﻿// Copyright Epic Games, Inc. All Rights Reserved.
+
+#include "ChaosVDParticleActor.h"
+#include "ChaosVDRecording.h"
+#include "Components/SceneComponent.h"
+#include "Components/StaticMeshComponent.h"
+#include "Engine/CollisionProfile.h"
+#include "Engine/StaticMesh.h"
+
+AChaosVDParticleActor::AChaosVDParticleActor(const FObjectInitializer& ObjectInitializer) : Super(ObjectInitializer)
+{
+	RootComponent = CreateDefaultSubobject<USceneComponent>(TEXT("SceneComponent0"));;
+	
+	//Note: This static mesh component is just a place holder for debugging purposes 
+	StaticMeshComponent = CreateDefaultSubobject<UStaticMeshComponent>(TEXT("StaticMeshComponent0"));
+	StaticMeshComponent->SetCollisionProfileName(UCollisionProfile::NoCollision_ProfileName);
+	StaticMeshComponent->SetGenerateOverlapEvents(false);
+	StaticMeshComponent->SetupAttachment(RootComponent);
+
+	StaticMeshComponent->Mobility = EComponentMobility::Movable;
+	StaticMeshComponent->bUseDefaultCollision = false;
+	StaticMeshComponent->bCastStaticShadow = false;
+	StaticMeshComponent->bCastDynamicShadow = false;
+	StaticMeshComponent->bSelectable = true;
+}
+
+void AChaosVDParticleActor::UpdateFromRecordedData(const FChaosVDParticleDebugData& RecordedData)
+{
+	SetActorLocation(RecordedData.Position);
+	SetActorRotation(RecordedData.Rotation);
+
+	// TODO: We should store a ptr to the data and in our custom details panel draw it, but the current data format is just a test
+	// So we should do this once it is final
+	RecordedDebugData = RecordedData;
+}
+
+void AChaosVDParticleActor::SetStaticMesh(UStaticMesh* NewStaticMesh)
+{
+	if (ensure(StaticMeshComponent))
+	{
+		StaticMeshComponent->SetStaticMesh(NewStaticMesh);
+	}
+}
diff --git a/Engine/Plugins/Experimental/ChaosVD/Source/ChaosVD/Private/ChaosVDParticleActorCustomization.cpp b/Engine/Plugins/Experimental/ChaosVD/Source/ChaosVD/Private/ChaosVDParticleActorCustomization.cpp
new file mode 100644
--- /dev/null
+++ b/Engine/Plugins/Experimental/ChaosVD/Source/ChaosVD/Private/ChaosVDParticleActorCustomization.cpp
@@ -0,0 +1,28 @@
+﻿// Copyright Epic Games, Inc. All Rights Reserved.
+
+#include "ChaosVDParticleActorCustomization.h"
+
+#include "DetailCategoryBuilder.h"
+#include "DetailLayoutBuilder.h"
+
+
+TSharedRef<IDetailCustomization> FChaosVDParticleActorCustomization::MakeInstance()
+{
+	return MakeShareable( new FChaosVDParticleActorCustomization );
+}
+
+void FChaosVDParticleActorCustomization::CustomizeDetails(IDetailLayoutBuilder& DetailBuilder)
+{
+	// Hide everything as the only thing we want to show in these actors is the Recorded debug data
+	TArray<FName> CurrentCategoryNames;
+	DetailBuilder.GetCategoryNames(CurrentCategoryNames);
+	for (const FName& CategoryToHide : CurrentCategoryNames)
+	{
+		if (CategoryToHide != ChaosVDCategoryName)
+		{
+			DetailBuilder.HideCategory(CategoryToHide);
+		}
+	}
+
+	DetailBuilder.EditCategory(ChaosVDCategoryName).InitiallyCollapsed(false);
+}
diff --git a/Engine/Plugins/Experimental/ChaosVD/Source/ChaosVD/Private/ChaosVDPlaybackController.cpp b/Engine/Plugins/Experimental/ChaosVD/Source/ChaosVD/Private/ChaosVDPlaybackController.cpp
new file mode 100644
--- /dev/null
+++ b/Engine/Plugins/Experimental/ChaosVD/Source/ChaosVD/Private/ChaosVDPlaybackController.cpp
@@ -0,0 +1,125 @@
+﻿// Copyright Epic Games, Inc. All Rights Reserved.
+
+#include "ChaosVD/Public/ChaosVDPlaybackController.h"
+
+#include "ChaosVDRecording.h"
+#include "ChaosVDScene.h"
+#include "Misc/FileHelper.h"
+#include "Serialization/MemoryReader.h"
+
+FChaosVDPlaybackController::FChaosVDPlaybackController(TWeakPtr<FChaosVDScene> InSceneToControl)
+	: CurrentFrame(0),
+	CurrentStep(0)
+{
+	SceneToControl = InSceneToControl;
+}
+
+bool FChaosVDPlaybackController::LoadChaosVDRecording(const FString& RecordingPath)
+{
+	if (RecordingPath.IsEmpty())
+	{
+		return false;
+	}
+
+	if (LoadedRecording.IsValid())
+	{
+		UnloadCurrentRecording();
+	}
+
+	//TODO: Show a notification or loading bar. The plan is support streaming of the file so maybe we don't need this at all
+
+	TArray<uint8> SerializedDebugData;
+	if (!ensure(FFileHelper::LoadFileToArray(SerializedDebugData, *RecordingPath)))
+	{
+		return false;
+	}
+
+	LoadedRecording = MakeShared<FChaosVDRecording>();
+
+	FMemoryReader SerializedDataReader(SerializedDebugData, false);					
+	FChaosVDRecording::StaticStruct()->SerializeBin(SerializedDataReader, LoadedRecording.Get());
+
+	if (!ensure(!SerializedDataReader.IsError()))
+	{
+		//TODO: Pop up an error message informing the user
+		LoadedRecording.Reset();
+		return false;
+	}
+
+	GoToRecordedStep(0,0);
+
+	OnControllerUpdated().ExecuteIfBound(this);
+
+	return true;
+}
+
+void FChaosVDPlaybackController::UnloadCurrentRecording()
+{
+	CurrentFrame = 0;
+	CurrentStep = 0;
+	LoadedRecording.Reset();
+	
+	if (const TSharedPtr<FChaosVDScene> SceneToControlSharedPtr = SceneToControl.Pin())
+	{
+		SceneToControlSharedPtr->CleanUpScene();
+	}
+	
+	OnControllerUpdated().ExecuteIfBound(this);
+}
+
+void FChaosVDPlaybackController::GoToRecordedStep(int32 FrameNumber, int32 Step)
+{
+	if (const TSharedPtr<FChaosVDScene> SceneToControlSharedPtr = SceneToControl.Pin())
+	{
+		if (ensure(LoadedRecording.IsValid()))
+		{
+			if (!LoadedRecording->RecordedFramesData.IsValidIndex(FrameNumber))
+			{
+				return;
+			}
+
+			//TODO: Handle different steps per solver if that is something that can happen
+			for (TPair<EChaosVDSolverType, FChaosVDSolverData>& SolverData : LoadedRecording->RecordedFramesData[FrameNumber].RecordedSolvers)
+			{
+				SceneToControlSharedPtr->UpdateFromRecordedStepData(SolverData.Key, SolverData.Value.SolverSteps[Step]);
+			}
+
+			CurrentStep = Step;
+			CurrentFrame = FrameNumber;
+		}
+	}
+	else
+	{
+		ensureMsgf(false, TEXT("GoToRecordedStep Called without a valid scene to control"));	
+	}
+}
+
+int32 FChaosVDPlaybackController::GetStepsForFrame(int32 FrameNumber) const
+{
+	if (LoadedRecording.IsValid())
+	{
+		if (!LoadedRecording->RecordedFramesData.IsValidIndex(FrameNumber))
+		{
+			return INDEX_NONE;
+		}
+
+		for (const TPair<EChaosVDSolverType, FChaosVDSolverData>& SolverData : LoadedRecording->RecordedFramesData[FrameNumber].RecordedSolvers)
+		{
+			//TODO: Handle different steps per solver if that is something that can happen
+			return SolverData.Value.SolverSteps.Num();
+		}
+	}
+
+	return	INDEX_NONE;
+}
+
+int32 FChaosVDPlaybackController::GetAvailableFramesNumber() const
+{
+	if (!LoadedRecording.IsValid())
+	{
+		return INDEX_NONE;
+	}
+
+	return LoadedRecording->RecordedFramesData.Num();
+}
+
diff --git a/Engine/Plugins/Experimental/ChaosVD/Source/ChaosVD/Private/ChaosVDPlaybackViewportTab.cpp b/Engine/Plugins/Experimental/ChaosVD/Source/ChaosVD/Private/ChaosVDPlaybackViewportTab.cpp
new file mode 100644
--- /dev/null
+++ b/Engine/Plugins/Experimental/ChaosVD/Source/ChaosVD/Private/ChaosVDPlaybackViewportTab.cpp
@@ -0,0 +1,33 @@
+﻿// Copyright Epic Games, Inc. All Rights Reserved.
+
+#include "ChaosVD/Public/ChaosVDPlaybackViewportTab.h"
+
+#include "ChaosVDEngine.h"
+#include "ChaosVDStyle.h"
+#include "ChaosVD/Public/Widgets/SChaosVDPlaybackViewport.h"
+#include "ChaosVDTabsIDs.h"
+#include "Widgets/Docking/SDockTab.h"
+#include "Widgets/SChaosVDMainTab.h"
+
+#define LOCTEXT_NAMESPACE "ChaosVisualDebugger"
+
+TSharedRef<SDockTab> FChaosVDPlaybackViewportTab::HandleTabSpawned(const FSpawnTabArgs& Args)
+{
+	TSharedRef<SDockTab> ViewportTab =
+		SNew(SDockTab)
+		.TabRole(ETabRole::PanelTab)
+		.Label(LOCTEXT("ViewportTabLabel", "Playback Viewport"))
+		.ToolTipText(LOCTEXT("ViewportTabToolTip", "The Chaos Visual debugger Viewport is under development"));
+	
+	ViewportTab->SetContent
+	(
+		//TODO: Handle Null cases to not crash the Editor
+		SAssignNew(PlaybackViewportWidget, SChaosVDPlaybackViewport, GetChaosVDWorld(), OwningTabWidget->GetChaosVDEngineInstance()->GetPlaybackController())
+	);
+
+	ViewportTab->SetTabIcon(FChaosVDStyle::Get().GetBrush("TabIconPlaybackViewport"));
+
+	return ViewportTab;
+}
+
+#undef LOCTEXT_NAMESPACE
\ No newline at end of file
diff --git a/Engine/Plugins/Experimental/ChaosVD/Source/ChaosVD/Private/ChaosVDScene.cpp b/Engine/Plugins/Experimental/ChaosVD/Source/ChaosVD/Private/ChaosVDScene.cpp
new file mode 100644
--- /dev/null
+++ b/Engine/Plugins/Experimental/ChaosVD/Source/ChaosVD/Private/ChaosVDScene.cpp
@@ -0,0 +1,191 @@
+﻿// Copyright Epic Games, Inc. All Rights Reserved.
+
+#include "ChaosVD/Public/ChaosVDScene.h"
+
+#include "ChaosVDEditorSettings.h"
+#include "ChaosVDModule.h"
+#include "ChaosVDParticleActor.h"
+#include "ChaosVDRecording.h"
+#include "Engine/Engine.h"
+#include "Engine/LevelStreamingDynamic.h"
+#include "Engine/World.h"
+#include "UObject/Package.h"
+
+void FChaosVDScene::Initialize()
+{
+	if (!ensure(!bIsInitialized))
+	{
+		return;
+	}
+
+	PhysicsVDWorld = CreatePhysicsVDWorld();
+
+	bIsInitialized = true;
+}
+
+void FChaosVDScene::DeInitialize()
+{
+	if (!ensure(bIsInitialized))
+	{
+		return;
+	}
+
+	if (PhysicsVDWorld)
+	{
+		PhysicsVDWorld->DestroyWorld(true);
+		PhysicsVDWorld->MarkAsGarbage();
+	}
+	
+	CollectGarbage(GARBAGE_COLLECTION_KEEPFLAGS);
+
+	bIsInitialized = false;
+}
+
+void FChaosVDScene::AddReferencedObjects(FReferenceCollector& Collector)
+{
+	Collector.AddReferencedObject(PhysicsVDWorld);
+	Collector.AddStableReferenceMap(ParticleVDInstancesByID);
+}
+
+void FChaosVDScene::UpdateFromRecordedStepData(EChaosVDSolverType SolverID, const FChaosVDStepData& InRecordedStepData)
+{
+	TSet<uint32> ParticlesIDsInRecordedStepData;
+	ParticlesIDsInRecordedStepData.Reserve(InRecordedStepData.RecordedParticles.Num());
+
+	// Go over existing Particle VD Instances and update them or create them if needed 
+	for (const FChaosVDParticleDebugData& Particle : InRecordedStepData.RecordedParticles)
+	{
+		const uint32 ParticleVDInstanceID = GetIDForRecordedParticleData(Particle);
+		ParticlesIDsInRecordedStepData.Add(ParticleVDInstanceID);
+
+		if (AChaosVDParticleActor** ExistingParticleVDInstancePtrPtr = ParticleVDInstancesByID.Find(ParticleVDInstanceID))
+		{
+			if (AChaosVDParticleActor* ExistingParticleVDInstancePtr = *ExistingParticleVDInstancePtrPtr)
+			{
+				ExistingParticleVDInstancePtr->UpdateFromRecordedData(Particle);
+			}
+			else
+			{
+				//TODO: Handle this error
+				ensure(false);
+			}
+		}
+		else
+		{
+			if (AChaosVDParticleActor* NewParticleVDInstance = SpawnParticleFromRecordedData(Particle))
+			{
+				FName FolderPath = *FPaths::Combine(TEXT("Solver ID : ") + UEnum::GetDisplayValueAsText(SolverID).ToString(), UEnum::GetDisplayValueAsText(Particle.ParticleType).ToString());
+				NewParticleVDInstance->SetFolderPath(FolderPath);
+
+				// TODO: Precalculate the max num of entries we would see in the loaded file, and use that number to pre-allocate this map
+				ParticleVDInstancesByID.Add(ParticleVDInstanceID, NewParticleVDInstance);
+			}
+			else
+			{
+				//TODO: Handle this error
+				ensure(false);
+			}
+		}
+	}
+
+	int32 AmountRemoved = 0;
+	for (TMap<uint32, AChaosVDParticleActor*>::TIterator RemoveIterator = ParticleVDInstancesByID.CreateIterator(); RemoveIterator; ++RemoveIterator)
+	{
+		if (!ParticlesIDsInRecordedStepData.Contains(RemoveIterator.Key()))
+		{
+			if (AChaosVDParticleActor* ActorToRemove = RemoveIterator.Value())
+			{
+				PhysicsVDWorld->DestroyActor(ActorToRemove);
+			}
+
+			RemoveIterator.RemoveCurrent();
+
+			AmountRemoved++;
+		}
+	}
+	
+	OnSceneUpdated().Broadcast();
+}
+
+void FChaosVDScene::CleanUpScene()
+{
+	for (const TPair<uint32, AChaosVDParticleActor*>& ParticleVDInstanceWithID : ParticleVDInstancesByID)
+	{
+		PhysicsVDWorld->DestroyActor(ParticleVDInstanceWithID.Value);
+	}
+
+	ParticleVDInstancesByID.Reset();
+
+	CollectGarbage(GARBAGE_COLLECTION_KEEPFLAGS);
+}
+
+AChaosVDParticleActor* FChaosVDScene::SpawnParticleFromRecordedData(const FChaosVDParticleDebugData& InParticleData)
+{
+	// Temp code until we have geometry data
+	const UChaosVDEditorSettings* Settings = GetDefault<UChaosVDEditorSettings>();
+	UStaticMesh* DebugMesh = Settings->DebugMesh.LoadSynchronous();
+
+	FActorSpawnParameters Params;
+	Params.Name = *InParticleData.DebugName;
+	Params.NameMode = FActorSpawnParameters::ESpawnActorNameMode::Requested;
+	if (AChaosVDParticleActor* NewActor = PhysicsVDWorld->SpawnActor<AChaosVDParticleActor>(Params))
+	{
+		NewActor->UpdateFromRecordedData(InParticleData);
+		NewActor->SetActorLabel(InParticleData.DebugName);
+		NewActor->SetStaticMesh(DebugMesh);
+
+		return NewActor;
+	}
+
+	return nullptr;
+}
+
+uint32 FChaosVDScene::GetIDForRecordedParticleData(const FChaosVDParticleDebugData& InParticleData) const
+{
+	return InParticleData.ParticleIndex;
+}
+
+UWorld* FChaosVDScene::CreatePhysicsVDWorld() const
+{
+	const FName UniqueWorldName = FName(FGuid::NewGuid().ToString());
+	UWorld* NewWorld = NewWorld = NewObject<UWorld>( GetTransientPackage(), UniqueWorldName );
+	
+	NewWorld->WorldType = EWorldType::EditorPreview;
+
+	FWorldContext& WorldContext = GEngine->CreateNewWorldContext( NewWorld->WorldType );
+	WorldContext.SetCurrentWorld(NewWorld);
+
+	NewWorld->InitializeNewWorld( UWorld::InitializationValues()
+										  .AllowAudioPlayback( false )
+										  .CreatePhysicsScene( false )
+										  .RequiresHitProxies( false )
+										  .CreateNavigation( false )
+										  .CreateAISystem( false )
+										  .ShouldSimulatePhysics( false )
+										  .SetTransactional( false )
+	);
+
+	// Add the base content as a sublevel
+	const UChaosVDEditorSettings* Settings = GetDefault<UChaosVDEditorSettings>();
+
+	ULevelStreamingDynamic* StreamedInLevel = NewObject<ULevelStreamingDynamic>(NewWorld);
+	StreamedInLevel->SetWorldAssetByPackageName(FName(Settings->BasePhysicsVDWorld.GetLongPackageName()));
+
+	StreamedInLevel->PackageNameToLoad = FName(Settings->BasePhysicsVDWorld.GetLongPackageName());
+
+	StreamedInLevel->SetShouldBeLoaded(true);
+	StreamedInLevel->bShouldBlockOnLoad = true;
+	StreamedInLevel->bInitiallyLoaded = true;
+
+	StreamedInLevel->SetShouldBeVisible(true);
+	StreamedInLevel->bInitiallyVisible = true;
+	StreamedInLevel->bLocked = true;
+
+	NewWorld->AddStreamingLevel(StreamedInLevel);
+
+	NewWorld->FlushLevelStreaming(EFlushLevelStreamingType::Full);
+
+	return NewWorld;
+}
+
+
diff --git a/Engine/Plugins/Experimental/ChaosVD/Source/ChaosVD/Private/ChaosVDStyle.cpp b/Engine/Plugins/Experimental/ChaosVD/Source/ChaosVD/Private/ChaosVDStyle.cpp
new file mode 100644
--- /dev/null
+++ b/Engine/Plugins/Experimental/ChaosVD/Source/ChaosVD/Private/ChaosVDStyle.cpp
@@ -0,0 +1,72 @@
+// Copyright Epic Games, Inc. All Rights Reserved.
+
+#include "ChaosVDStyle.h"
+#include "Framework/Application/SlateApplication.h"
+#include "Interfaces/IPluginManager.h"
+#include "Slate/SlateGameResources.h"
+#include "Styling/SlateStyle.h"
+#include "Styling/SlateStyleMacros.h"
+#include "Styling/SlateStyleRegistry.h"
+
+#define RootToContentDir Style->RootToContentDir
+
+TSharedPtr<FSlateStyleSet> FChaosVDStyle::StyleInstance = nullptr;
+
+void FChaosVDStyle::Initialize()
+{
+	if (!StyleInstance.IsValid())
+	{
+		StyleInstance = Create();
+		FSlateStyleRegistry::RegisterSlateStyle(*StyleInstance);
+	}
+}
+
+void FChaosVDStyle::Shutdown()
+{
+	FSlateStyleRegistry::UnRegisterSlateStyle(*StyleInstance);
+	ensure(StyleInstance.IsUnique());
+	StyleInstance.Reset();
+}
+
+FName FChaosVDStyle::GetStyleSetName()
+{
+	static FName StyleSetName(TEXT("ChaosVDStyle"));
+	return StyleSetName;
+}
+
+const FVector2D Icon16x16(16.0f, 16.0f);
+const FVector2D Icon20x20(20.0f, 20.0f);
+const FVector2D Icon50x50(50.0f, 50.0f);
+
+TSharedRef< FSlateStyleSet > FChaosVDStyle::Create()
+{
+	TSharedRef< FSlateStyleSet > Style = MakeShareable(new FSlateStyleSet("ChaosVDStyle"));
+	Style->SetContentRoot(IPluginManager::Get().FindPlugin("ChaosVD")->GetBaseDir() / TEXT("Resources"));
+
+	Style->Set("ChaosVD.OpenPluginWindow", new IMAGE_BRUSH_SVG(TEXT("PlaceholderButtonIcon"), Icon20x20));
+
+	Style->Set("PlayIcon", new IMAGE_BRUSH(TEXT("Play_16x"), Icon16x16));
+	Style->Set("StopIcon",  new IMAGE_BRUSH(TEXT("Stop_16x"), Icon16x16));
+	Style->Set("NextIcon",  new IMAGE_BRUSH(TEXT("Next_16x"), Icon16x16));
+	Style->Set("PrevIcon",  new IMAGE_BRUSH( TEXT("Prev_16x"), Icon16x16));
+	Style->Set("PauseIcon",  new IMAGE_BRUSH( TEXT("generic_pause_16x"), Icon16x16));
+	
+	Style->Set("TabIconDetailsPanel",  new IMAGE_BRUSH_SVG(TEXT("Details"), Icon16x16));
+	Style->Set("TabIconPlaybackViewport",  new IMAGE_BRUSH(TEXT("icon_tab_PlaybackViewport_16x"), Icon16x16));
+	Style->Set("TabIconWorldOutliner", new IMAGE_BRUSH_SVG(TEXT("WorldOutliner"), Icon16x16));
+
+	return Style;
+}
+
+void FChaosVDStyle::ReloadTextures()
+{
+	if (FSlateApplication::IsInitialized())
+	{
+		FSlateApplication::Get().GetRenderer()->ReloadTextureResources();
+	}
+}
+
+const ISlateStyle& FChaosVDStyle::Get()
+{
+	return *StyleInstance;
+}
diff --git a/Engine/Plugins/Experimental/ChaosVD/Source/ChaosVD/Private/ChaosVDTabSpawnerBase.cpp b/Engine/Plugins/Experimental/ChaosVD/Source/ChaosVD/Private/ChaosVDTabSpawnerBase.cpp
new file mode 100644
--- /dev/null
+++ b/Engine/Plugins/Experimental/ChaosVD/Source/ChaosVD/Private/ChaosVDTabSpawnerBase.cpp
@@ -0,0 +1,23 @@
+﻿// Copyright Epic Games, Inc. All Rights Reserved.
+
+#include "ChaosVDTabSpawnerBase.h"
+
+#include "ChaosVDEngine.h"
+#include "ChaosVDScene.h"
+#include "Widgets/SChaosVDMainTab.h"
+
+FChaosVDTabSpawnerBase::FChaosVDTabSpawnerBase(const FName& InTabID, TSharedPtr<FTabManager> InTabManager, SChaosVDMainTab* InOwningTabWidget)
+{
+	OwningTabWidget = InOwningTabWidget;
+	InTabManager->RegisterTabSpawner(InTabID, FOnSpawnTab::CreateRaw(this, &FChaosVDTabSpawnerBase::HandleTabSpawned));
+}
+
+UWorld* FChaosVDTabSpawnerBase::GetChaosVDWorld() const
+{
+	if (ensure(OwningTabWidget) && OwningTabWidget->GetChaosVDEngineInstance()->GetCurrentScene().IsValid())
+	{
+		return OwningTabWidget->GetChaosVDEngineInstance()->GetCurrentScene()->GetUnderlyingWorld();
+	}
+
+	return nullptr;
+}
diff --git a/Engine/Plugins/Experimental/ChaosVD/Source/ChaosVD/Private/ChaosVDTabsIDs.cpp b/Engine/Plugins/Experimental/ChaosVD/Source/ChaosVD/Private/ChaosVDTabsIDs.cpp
new file mode 100644
--- /dev/null
+++ b/Engine/Plugins/Experimental/ChaosVD/Source/ChaosVD/Private/ChaosVDTabsIDs.cpp
@@ -0,0 +1,10 @@
+﻿// Copyright Epic Games, Inc. All Rights Reserved.
+
+#include "ChaosVDTabsIDs.h"
+#include "UObject/NameTypes.h"
+
+const FName FChaosVDTabID::ChaosVisualDebuggerTab = TEXT("ChaosVisualDebuggerTab");
+const FName FChaosVDTabID::PlaybackViewport = TEXT("PlaybackViewport");
+const FName FChaosVDTabID::WorldOutliner = TEXT("WorldOutliner");
+const FName FChaosVDTabID::DetailsPanel = TEXT("DetailsPanel");
+const FName FChaosVDTabID::OutputLog = TEXT("ChaosVDOutputLog");
\ No newline at end of file
diff --git a/Engine/Plugins/Experimental/ChaosVD/Source/ChaosVD/Private/ChaosVDWorldOutlinerTab.cpp b/Engine/Plugins/Experimental/ChaosVD/Source/ChaosVD/Private/ChaosVDWorldOutlinerTab.cpp
new file mode 100644
--- /dev/null
+++ b/Engine/Plugins/Experimental/ChaosVD/Source/ChaosVD/Private/ChaosVDWorldOutlinerTab.cpp
@@ -0,0 +1,46 @@
+﻿// Copyright Epic Games, Inc. All Rights Reserved.
+
+#include "ChaosVDWorldOutlinerTab.h"
+
+#include "ChaosVDStyle.h"
+#include "ChaosVDTabsIDs.h"
+#include "Modules/ModuleManager.h"
+#include "SceneOutlinerModule.h"
+#include "SceneOutlinerPublicTypes.h"
+#include "Widgets/Docking/SDockTab.h"
+#include "Engine/World.h"
+
+#define LOCTEXT_NAMESPACE "ChaosVisualDebugger"
+
+TSharedRef<SDockTab> FChaosVDWorldOutlinerTab::HandleTabSpawned(const FSpawnTabArgs& Args)
+{
+	FSceneOutlinerInitializationOptions InitOptions;
+	InitOptions.bShowTransient = true;
+
+	InitOptions.OutlinerIdentifier = TEXT("ChaosVDOutliner");
+	InitOptions.FilterBarOptions.bHasFilterBar = true;
+
+	InitOptions.FilterBarOptions.bUseSharedSettings = false;
+
+	const FSceneOutlinerModule& SceneOutlinerModule = FModuleManager::Get().LoadModuleChecked<FSceneOutlinerModule>("SceneOutliner");
+
+	//TODO: See if we can change this to an Actor Picker instead of Actor Browser as it is simpler
+	SceneOutlinerWidget = SceneOutlinerModule.CreateActorBrowser(InitOptions, TWeakObjectPtr<UWorld>(GetChaosVDWorld()));
+
+	TSharedRef<SDockTab> OutlinerTab =
+		SNew(SDockTab)
+		.TabRole(ETabRole::MajorTab)
+		.Label(LOCTEXT("Physics World Outliner", "Physics World Outliner"))
+		.ToolTipText(LOCTEXT("PhysicsWorldOutlinerTabToolTip", "Hierachy view of the physics objects by category"));
+	
+	OutlinerTab->SetContent
+	(
+		SceneOutlinerWidget.ToSharedRef()
+	);
+
+	OutlinerTab->SetTabIcon(FChaosVDStyle::Get().GetBrush("TabIconWorldOutliner"));
+
+	return OutlinerTab;
+}
+
+#undef LOCTEXT_NAMESPACE
diff --git a/Engine/Plugins/Experimental/ChaosVD/Source/ChaosVD/Private/Widgets/SChaosVDMainTab.cpp b/Engine/Plugins/Experimental/ChaosVD/Source/ChaosVD/Private/Widgets/SChaosVDMainTab.cpp
new file mode 100644
--- /dev/null
+++ b/Engine/Plugins/Experimental/ChaosVD/Source/ChaosVD/Private/Widgets/SChaosVDMainTab.cpp
@@ -0,0 +1,257 @@
+﻿// Copyright Epic Games, Inc. All Rights Reserved.
+
+#include "ChaosVD/Public/Widgets/SChaosVDMainTab.h"
+
+#include "ChaosVDEngine.h"
+#include "ChaosVDObjectDetailsTab.h"
+#include "ChaosVDOutputLogTab.h"
+#include "ChaosVDPlaybackController.h"
+#include "ChaosVDPlaybackViewportTab.h"
+#include "ChaosVDTabsIDs.h"
+#include "ChaosVDWorldOutlinerTab.h"
+#include "DesktopPlatformModule.h"
+#include "Framework/Application/SlateApplication.h"
+#include "Framework/Application/SWindowTitleBar.h"
+#include "Framework/Commands/UICommandList.h"
+#include "Framework/MultiBox/MultiBoxBuilder.h"
+#include "IDesktopPlatform.h"
+#include "Styling/StyleColors.h"
+#include "Widgets/Docking/SDockTab.h"
+#include "Widgets/Input/SComboButton.h"
+#include "Widgets/SToolTip.h"
+
+#define LOCTEXT_NAMESPACE "ChaosVisualDebugger"
+
+void SChaosVDMainTab::Construct(const FArguments& InArgs, TSharedPtr<FChaosVDEngine> InChaosVDEngine)
+{
+	ChaosVDEngine = InChaosVDEngine;
+
+	TabManager = FGlobalTabmanager::Get()->NewTabManager(InArgs._OwnerTab.ToSharedRef()).ToSharedPtr();
+
+	// Create the UI Tabs Handlers
+	//TODO: Pass self as a weakPtr
+	WorldOutlinerTab = MakeShared<FChaosVDWorldOutlinerTab>(FChaosVDTabID::WorldOutliner, TabManager, this);
+	ObjectDetailsTab = MakeShared<FChaosVDObjectDetailsTab>(FChaosVDTabID::DetailsPanel, TabManager, this);
+	OutputLogTab = MakeShared<FChaosVDOutputLogTab>(FChaosVDTabID::OutputLog, TabManager, this);
+	PlaybackViewportTab = MakeShared<FChaosVDPlaybackViewportTab>(FChaosVDTabID::PlaybackViewport, TabManager, this);
+
+	GenerateMainWindowMenu();
+
+	ChildSlot
+	[
+		// Row between the tab and main content 
+		SNew(SVerticalBox)
+		+ SVerticalBox::Slot()
+		.AutoHeight()
+		[
+			SNew( SBorder )
+			.BorderImage(FAppStyle::Get().GetBrush("Brushes.Panel"))
+			.Padding(FMargin(0.f, 0.f, 0.f, 0.f))
+			[
+				// Open Button
+				SNew(SHorizontalBox)
+				+SHorizontalBox::Slot()
+				.VAlign(VAlign_Center)
+				.HAlign(HAlign_Left)
+				.Padding(FMargin(12, 7, 18, 7))
+				.AutoWidth()
+				[
+					SNew(SButton)
+					.ToolTip(SNew(SToolTip).Text(LOCTEXT("OpenFileDesc", "Click here to open a Chaos Visual Debugger file.")))
+					.ContentPadding(FMargin(0, 5.f, 0, 4.f))
+					.OnClicked_Lambda([this]()
+					{
+						BrowseAndOpenChaosVDFile();
+						return FReply::Handled();
+					})
+					.Content()
+					[
+						SNew(SHorizontalBox)
+						+ SHorizontalBox::Slot()
+						.HAlign(HAlign_Center)
+						.VAlign(VAlign_Center)
+						[
+							SNew(SImage)
+							.Image(FAppStyle::Get().GetBrush("Icons.Plus"))
+							.ColorAndOpacity(FStyleColors::AccentGreen)
+						]
+						+ SHorizontalBox::Slot()
+						.Padding(FMargin(3, 0, 0, 0))
+						.VAlign(VAlign_Center)
+						.AutoWidth()
+						[
+							SNew(STextBlock)
+							.TextStyle(FAppStyle::Get(), "SmallButtonText")
+							.Text(LOCTEXT("OpenFile", "Open File"))
+						]
+					]
+				]
+
+				+ SHorizontalBox::Slot()
+				.Padding(0, 7, 0, 7)
+
+				// Settings button
+				+ SHorizontalBox::Slot()
+				.HAlign(HAlign_Right)
+				.VAlign(VAlign_Center)
+				.AutoWidth()
+				.Padding(FMargin(14.f, 0.f, 14.f, 0.f))
+				[
+					SNew(SComboButton)
+					.ContentPadding(0)
+					.HasDownArrow(false)
+					.ForegroundColor(FSlateColor::UseForeground())
+					.ComboButtonStyle(FAppStyle::Get(), "SimpleComboButton")
+					.AddMetaData<FTagMetaData>(FTagMetaData(TEXT("ViewOptions")))
+					/*.MenuContent()
+					[
+						//TODO: Implement Settings menu
+					]*/
+					.ButtonContent()
+					[
+						SNew(SHorizontalBox)
+						+ SHorizontalBox::Slot()
+						.AutoWidth()
+						.HAlign(HAlign_Center)
+						[
+							SNew(SImage)
+							.Image(FAppStyle::Get().GetBrush("Icons.Toolbar.Settings"))
+							.ColorAndOpacity(FSlateColor::UseForeground())
+						]
+						+ SHorizontalBox::Slot()
+						.Padding(FMargin(5, 0, 0, 0))
+						.VAlign(VAlign_Center)
+						.AutoWidth()
+						[
+							SNew(STextBlock)
+							.TextStyle(FAppStyle::Get(), "NormalText")
+							.Text(LOCTEXT("SettingsLabel", "Settings"))
+							.ColorAndOpacity(FSlateColor::UseForeground())
+						]
+					]
+				]
+			]
+				
+		]
+		// Main Visual Debugger Interface content
+		+SVerticalBox::Slot()
+		.Padding(FMargin(0.0f,5.0f,0.0f,0.0f))
+		[
+			TabManager->RestoreFrom(GenerateMainLayout(), TabManager->GetOwnerTab()->GetParentWindow()).ToSharedRef()
+		]
+	];
+}
+
+TSharedRef<FTabManager::FLayout> SChaosVDMainTab::GenerateMainLayout()
+{
+	return FTabManager::NewLayout("ChaosVisualDebugger_Layout")
+		->AddArea
+		(
+			FTabManager::NewPrimaryArea()
+			->SetOrientation(Orient_Horizontal)
+			->SetExtensionId("TopLevelArea")
+			->Split
+			(
+				FTabManager::NewSplitter()
+				->SetOrientation(Orient_Vertical)
+				->SetSizeCoefficient(0.8f)
+				->Split
+				(
+					FTabManager::NewStack()
+					->SetSizeCoefficient(0.7f)
+					->AddTab(FChaosVDTabID::PlaybackViewport, ETabState::OpenedTab)
+					->SetHideTabWell(true)
+				)
+				->Split
+				(
+					FTabManager::NewStack()
+					->SetSizeCoefficient(0.3f)
+					->AddTab(FChaosVDTabID::OutputLog, ETabState::OpenedTab)
+				)
+			)
+			->Split
+			(
+				FTabManager::NewSplitter()
+				->SetOrientation(Orient_Vertical)
+				->SetSizeCoefficient(0.15f)
+				->Split
+				(
+					FTabManager::NewStack()
+					->SetSizeCoefficient(0.7f)
+					->AddTab(FChaosVDTabID::WorldOutliner, ETabState::OpenedTab)
+				)
+				->Split
+				(
+					FTabManager::NewStack()
+					->SetSizeCoefficient(0.3f)
+					->AddTab(FChaosVDTabID::DetailsPanel, ETabState::OpenedTab)
+				)
+			)
+		);
+}
+
+SChaosVDMainTab::~SChaosVDMainTab()
+{
+	if (ChaosVDEngine.IsValid())
+	{
+		ChaosVDEngine->DeInitialize();
+	}
+}
+
+void SChaosVDMainTab::GenerateMainWindowMenu()
+{
+	FMenuBarBuilder MenuBarBuilder = FMenuBarBuilder(TSharedPtr<FUICommandList >());
+	MenuBarBuilder.AddPullDownMenu(
+		LOCTEXT("FileMenuLabel", "File"),
+		FText::GetEmpty(),
+		FNewMenuDelegate::CreateLambda([this](FMenuBuilder& MenuBuilder)
+		{
+			MenuBuilder.AddMenuEntry(LOCTEXT("OpenFileMenuLabel", "OpenFile"), FText(), FSlateIcon(),
+			FUIAction(
+				FExecuteAction::CreateLambda([this]()
+				{
+					BrowseAndOpenChaosVDFile();
+				})));
+		}),
+		"File"
+	);
+	MenuBarBuilder.AddPullDownMenu(
+		LOCTEXT("WindowMenuLabel", "Window"),
+		FText::GetEmpty(),
+		FNewMenuDelegate::CreateLambda([this](FMenuBuilder& MenuBuilder)
+		{
+			TabManager->PopulateLocalTabSpawnerMenu(MenuBuilder);
+		}),
+		"Window"
+	);
+
+	TabManager->SetAllowWindowMenuBar(true);
+	TabManager->SetMenuMultiBox(MenuBarBuilder.GetMultiBox(), MenuBarBuilder.MakeWidget());
+}
+
+void SChaosVDMainTab::BrowseAndOpenChaosVDFile()
+{
+	TArray<FString> OutOpenFilenames;
+	if (IDesktopPlatform* DesktopPlatform = FDesktopPlatformModule::Get())
+	{
+		DesktopPlatform->OpenFileDialog(
+			FSlateApplication::Get().FindBestParentWindowHandleForDialogs(AsShared()),
+			LOCTEXT("OpenDialogTitle", "Open Chaos Visual Debug File").ToString(),
+			TEXT(""),
+			TEXT(""),
+			TEXT("Chaos Visual Debugger Files|*.cvd"),
+			EFileDialogFlags::None,
+			OutOpenFilenames
+		);
+	}
+
+	if (OutOpenFilenames.Num() > 0)
+	{
+		if (ensure(ChaosVDEngine->GetPlaybackController()))
+		{
+			ChaosVDEngine->GetPlaybackController()->LoadChaosVDRecording(OutOpenFilenames[0]);
+		}	
+	}
+}
+
+#undef LOCTEXT_NAMESPACE
diff --git a/Engine/Plugins/Experimental/ChaosVD/Source/ChaosVD/Private/Widgets/SChaosVDPlaybackViewport.cpp b/Engine/Plugins/Experimental/ChaosVD/Source/ChaosVD/Private/Widgets/SChaosVDPlaybackViewport.cpp
new file mode 100644
--- /dev/null
+++ b/Engine/Plugins/Experimental/ChaosVD/Source/ChaosVD/Private/Widgets/SChaosVDPlaybackViewport.cpp
@@ -0,0 +1,181 @@
+﻿// Copyright Epic Games, Inc. All Rights Reserved.
+
+#include "Widgets/SChaosVDPlaybackViewport.h"
+
+#include "ChaosVDPlaybackController.h"
+#include "Framework/Application/SlateApplication.h"
+#include "LevelEditorViewport.h"
+#include "Slate/SceneViewport.h"
+#include "Widgets/SChaosVDTimelineWidget.h"
+#include "Widgets/SViewport.h"
+#include "Widgets/Text/STextBlock.h"
+
+
+#define LOCTEXT_NAMESPACE "ChaosVisualDebugger"
+
+SChaosVDPlaybackViewport::~SChaosVDPlaybackViewport()
+{
+	LevelViewportClient->Viewport = nullptr;
+	LevelViewportClient.Reset();
+}
+
+TSharedPtr<FLevelEditorViewportClient> SChaosVDPlaybackViewport::CreateViewportClient() const
+{
+	TSharedPtr<FLevelEditorViewportClient> NewViewport = MakeShareable(new FLevelEditorViewportClient(TSharedPtr<class SLevelViewport>()));
+
+	NewViewport->SetAllowCinematicControl(false);
+	
+	NewViewport->bSetListenerPosition = false;
+	NewViewport->EngineShowFlags = FEngineShowFlags(ESFIM_Editor);
+	NewViewport->LastEngineShowFlags = FEngineShowFlags(ESFIM_Editor);
+	NewViewport->ViewportType = LVT_Perspective;
+	NewViewport->bDrawAxes = true;
+	NewViewport->bDisableInput = false;
+	NewViewport->VisibilityDelegate.BindLambda([] {return true; });
+
+	return NewViewport;
+}
+
+void SChaosVDPlaybackViewport::Construct(const FArguments& InArgs, const UWorld* DefaultWorld, TWeakPtr<FChaosVDPlaybackController> InPlaybackController)
+{
+	ensure(DefaultWorld);
+	ensure(InPlaybackController.IsValid());
+
+	PlaybackController = InPlaybackController;
+
+	LevelViewportClient = CreateViewportClient();
+
+	ViewportWidget = SNew(SViewport)
+		.RenderDirectlyToWindow(false)
+		.IsEnabled(FSlateApplication::Get().GetNormalExecutionAttribute())
+		.EnableGammaCorrection(false)
+		.EnableBlending(false);
+
+	SceneViewport = MakeShareable(new FSceneViewport(LevelViewportClient.Get(), ViewportWidget));
+
+	LevelViewportClient->Viewport = SceneViewport.Get();
+
+	ViewportWidget->SetViewportInterface(SceneViewport.ToSharedRef());
+	
+	// Default to the base map
+	LevelViewportClient->SetReferenceToWorldContext(*GEngine->GetWorldContextFromWorld(DefaultWorld));
+
+	ChildSlot
+	[
+		// 3D Viewport
+		SNew(SVerticalBox)
+		+SVerticalBox::Slot()
+		.FillHeight(0.9f)
+		[
+			ViewportWidget.ToSharedRef()
+		]
+		// Playback controls
+		// TODO: Now that the tool is In-Editor, see if we can/is worth use the Sequencer widgets
+		// instead of these custom ones
+		+SVerticalBox::Slot()
+		.Padding(16.0f, 16.0f, 16.0f, 16.0f)
+		.FillHeight(0.1f)
+		[
+			SNew(SHorizontalBox)
+			+SHorizontalBox::Slot()
+			.FillWidth(0.7f)
+			[
+			SNew(SVerticalBox)
+				+SVerticalBox::Slot()
+				.AutoHeight()
+				.Padding(0.0f, 0.0f, 0.0f, 2.0f)
+				[
+					SNew(STextBlock)
+					.Justification(ETextJustify::Center)
+					.Text(LOCTEXT("PlaybackViewportWidgetPhysicsFramesLabel", "Physics Frames" ))
+				]
+				+SVerticalBox::Slot()
+				[
+					SAssignNew(FramesTimelineWidget, SChaosVDTimelineWidget)
+						.HidePlayStopButtons(false)
+						.OnFrameChanged_Raw(this, &SChaosVDPlaybackViewport::OnFrameSelectionUpdated)
+						.MaxFrames(0)
+				]
+			]
+			+SHorizontalBox::Slot()
+			.FillWidth(0.3f)
+			[
+				SNew(SVerticalBox)
+				+SVerticalBox::Slot()
+				.AutoHeight()
+				.Padding(0.0f, 0.0f, 0.0f, 2.0f)
+				[
+					SNew(STextBlock)
+					.Justification(ETextJustify::Center)
+					.Text(LOCTEXT("PlaybackViewportWidgetStepsLabel", "Solver Steps" ))
+				]
+				+SVerticalBox::Slot()
+				[
+					SAssignNew(StepsTimelineWidget, SChaosVDTimelineWidget)
+					.HidePlayStopButtons(true)
+					.OnFrameChanged_Raw(this, &SChaosVDPlaybackViewport::OnStepSelectionUpdated)
+					.MaxFrames(0)
+				]	
+			]
+		]
+	];
+
+	if (const TSharedPtr<FChaosVDPlaybackController> PlaybackControllerPtr = PlaybackController.Pin())
+	{
+		PlaybackControllerPtr->OnControllerUpdated().BindRaw(this, &SChaosVDPlaybackViewport::OnPlaybackControllerUpdated);
+	}
+}
+
+void SChaosVDPlaybackViewport::OnPlaybackControllerUpdated(FChaosVDPlaybackController* Controller) const
+{
+	if (!ensure(Controller))
+	{
+		return;
+	}
+
+	if (TSharedPtr<FChaosVDRecording> LoadedRecording = Controller->GetCurrentRecording().Pin())
+	{
+		const int32 AvailableFrames = Controller->GetAvailableFramesNumber();
+		const int32 AvailableSteps = Controller->GetStepsForFrame(Controller->GetCurrentFrame());
+
+		// Max is inclusive and we use this to request as the index on the recorded frames/steps arrays so we need to -1 to the available frames/steps
+		FramesTimelineWidget->UpdateMinMaxValue(0, AvailableFrames != INDEX_NONE ? AvailableFrames -1  : 0);
+
+		//TODO: This will show steps 0/0 if only one step is recorded, we need to add a way to override that functionality
+		// or just set the slider to start from 1 and handle the offset later 
+		StepsTimelineWidget->UpdateMinMaxValue(0,AvailableSteps != INDEX_NONE ? AvailableSteps -1 : 0);
+	}
+	else
+	{
+		FramesTimelineWidget->UpdateMinMaxValue(0, 0);
+		FramesTimelineWidget->ResetTimeline();
+		StepsTimelineWidget->UpdateMinMaxValue(0,0);
+		StepsTimelineWidget->ResetTimeline();
+	}
+
+	LevelViewportClient->bNeedsRedraw = true;
+}
+
+void SChaosVDPlaybackViewport::OnFrameSelectionUpdated(int32 NewFrameIndex) const
+{
+	if (const TSharedPtr<FChaosVDPlaybackController> PlaybackControllerPtr = PlaybackController.Pin())
+	{
+		// On Frame updates, always use Step 0
+		PlaybackControllerPtr->GoToRecordedStep(NewFrameIndex, 0);
+
+		LevelViewportClient->bNeedsRedraw = true;
+	}
+}
+
+void SChaosVDPlaybackViewport::OnStepSelectionUpdated(int32 NewStepIndex) const
+{
+	if (const TSharedPtr<FChaosVDPlaybackController> PlaybackControllerPtr = PlaybackController.Pin())
+	{
+		// On Steps updates. Always use the current Frame
+		PlaybackControllerPtr->GoToRecordedStep(PlaybackControllerPtr->GetCurrentFrame(), NewStepIndex);
+	
+		LevelViewportClient->bNeedsRedraw = true;
+	}
+}
+
+#undef LOCTEXT_NAMESPACE
diff --git a/Engine/Plugins/Experimental/ChaosVD/Source/ChaosVD/Private/Widgets/SChaosVDTimelineWidget.cpp b/Engine/Plugins/Experimental/ChaosVD/Source/ChaosVD/Private/Widgets/SChaosVDTimelineWidget.cpp
new file mode 100644
--- /dev/null
+++ b/Engine/Plugins/Experimental/ChaosVD/Source/ChaosVD/Private/Widgets/SChaosVDTimelineWidget.cpp
@@ -0,0 +1,231 @@
+﻿// Copyright Epic Games, Inc. All Rights Reserved.
+
+#include "ChaosVD/Public/Widgets/SChaosVDTimelineWidget.h"
+
+#include "ChaosVDStyle.h"
+#include "Input/Reply.h"
+#include "Widgets/Input/SButton.h"
+#include "Widgets/Input/SSlider.h"
+#include "Widgets/Images/SImage.h"
+#include "Widgets/Text/STextBlock.h"
+#include "Widgets/SBoxPanel.h"
+
+#define LOCTEXT_NAMESPACE "ChaosVisualDebugger"
+
+void SChaosVDTimelineWidget::Construct(const FArguments& InArgs)
+{
+	MaxFrames = InArgs._MaxFrames;
+	FrameChangedDelegate = InArgs._OnFrameChanged;
+
+	SetCanTick(false);
+
+	ChildSlot
+	[
+		SNew(SHorizontalBox)
+			+SHorizontalBox::Slot()
+			.HAlign(HAlign_Center)
+			.VAlign(VAlign_Center)
+			.FillWidth(0.2)
+			[
+				SNew(SHorizontalBox)
+				+SHorizontalBox::Slot()
+				[
+					SNew(SButton)
+					.Visibility(InArgs._HidePlayStopButtons ? EVisibility::Collapsed : EVisibility::Visible)
+					.OnClicked( FOnClicked::CreateRaw(this, &SChaosVDTimelineWidget::Play))
+					.ContentPadding( 2.0f )
+					.ForegroundColor( FSlateColor::UseForeground() )
+					.IsFocusable( false )
+					[
+						SNew( SImage )
+						.Image_Raw(this, &SChaosVDTimelineWidget::GetPlayOrPauseIcon)
+						.DesiredSizeOverride(FVector2D(16.0f,16.0f))
+						.ColorAndOpacity( FSlateColor::UseForeground() )
+					]
+				]
+				+SHorizontalBox::Slot()
+				[
+					SNew(SButton)
+					.Visibility(InArgs._HidePlayStopButtons ? EVisibility::Collapsed : EVisibility::Visible)
+					.OnClicked( FOnClicked::CreateRaw(this, &SChaosVDTimelineWidget::Stop))
+					.ContentPadding( 2.0f )
+					.ForegroundColor( FSlateColor::UseForeground() )
+					.IsFocusable( false )
+					[
+						SNew( SImage )
+						.Image(FChaosVDStyle::Get().GetBrush("StopIcon"))
+						.DesiredSizeOverride(FVector2D(16.0f,16.0f))
+						.ColorAndOpacity( FSlateColor::UseForeground() )
+					]
+				]
+				+SHorizontalBox::Slot()
+				[
+					SNew(SButton)
+					.OnClicked( FOnClicked::CreateRaw(this, &SChaosVDTimelineWidget::Prev))
+					.ContentPadding( 2.0f )
+					.ForegroundColor( FSlateColor::UseForeground() )
+					.IsFocusable( false )
+					[
+						SNew( SImage )
+						.DesiredSizeOverride(FVector2D(16.0f,16.0f))
+						.Image(FChaosVDStyle::Get().GetBrush("PrevIcon"))
+						.ColorAndOpacity( FSlateColor::UseForeground() )
+					]
+				]
+				+SHorizontalBox::Slot()
+				[
+					SNew(SButton)
+					.OnClicked( FOnClicked::CreateRaw(this, &SChaosVDTimelineWidget::Next))
+					.ContentPadding( 2.0f )
+					.ForegroundColor( FSlateColor::UseForeground() )
+					.IsFocusable( false )
+					[
+						SNew( SImage )
+						.DesiredSizeOverride(FVector2D(16.0f,16.0f))
+						.Image(FChaosVDStyle::Get().GetBrush("NextIcon"))
+						.ColorAndOpacity( FSlateColor::UseForeground() )
+					]
+				]
+			]
+			+SHorizontalBox::Slot()
+			.VAlign(VAlign_Center)
+			.FillWidth(0.75)
+			[
+			  SAssignNew(TimelineSlider, SSlider)
+			  .ToolTipText_Lambda([this]()-> FText{ return FText::AsNumber(CurrentFrame); })
+			  .Value(CurrentFrame)
+			  .OnValueChanged_Raw(this, &SChaosVDTimelineWidget::SetCurrentTimelineFrame)
+			  .StepSize(1)
+			  .MaxValue(MaxFrames)
+			  .MinValue(0)
+			]
+			+SHorizontalBox::Slot()
+			.FillWidth(0.05)
+			.HAlign(HAlign_Center)
+			.VAlign(VAlign_Center)
+			[
+				SNew(STextBlock)
+				.ColorAndOpacity(FLinearColor::White)
+				.Text_Lambda([this]()->FText{ return FText::Format(LOCTEXT("FramesCounter","{0} / {1}"), CurrentFrame, MaxFrames);})
+			]
+	];
+}
+
+void SChaosVDTimelineWidget::UpdateMinMaxValue(float NewMin, float NewMax)
+{
+	if (!TimelineSlider.IsValid())
+	{
+		return;
+	}
+
+	TimelineSlider->SetMinAndMaxValues(NewMin, NewMax);
+
+	MinFrames = NewMin;
+	MaxFrames = NewMax;
+	CurrentFrame = NewMin;
+}
+
+void SChaosVDTimelineWidget::ResetTimeline()
+{
+	TimelineSlider->SetValue(MinFrames);
+	CurrentFrame = MinFrames;
+}
+
+void SChaosVDTimelineWidget::SetCurrentTimelineFrame(float FrameNumber)
+{
+	CurrentFrame = static_cast<int32>(FrameNumber);
+	FrameChangedDelegate.ExecuteIfBound(CurrentFrame);
+}
+
+FReply SChaosVDTimelineWidget::Play()
+{
+	if (bIsPlaying)
+	{
+		bIsPlaying = false;
+		SetCanTick(false);
+	}
+	else
+	{
+		bIsPlaying = true;
+		SetCanTick(true);
+	}
+
+	return FReply::Handled();
+}
+
+FReply SChaosVDTimelineWidget::Stop()
+{
+	CurrentFrame = 0;
+	CurrentPlaybackTime = 0.0f;
+	bIsPlaying = false;
+	TimelineSlider->SetValue(CurrentFrame);
+
+	SetCurrentTimelineFrame(CurrentFrame);
+
+	SetCanTick(false);
+
+	return FReply::Handled();
+}
+
+FReply SChaosVDTimelineWidget::Next()
+{
+	if (CurrentFrame >= MaxFrames)
+	{
+		CurrentFrame = MaxFrames;
+		return FReply::Handled();
+	}
+
+	CurrentFrame++;
+	TimelineSlider->SetValue(CurrentFrame);
+	
+	SetCurrentTimelineFrame(CurrentFrame);
+
+	return FReply::Handled();
+}
+
+FReply SChaosVDTimelineWidget::Prev()
+{
+	if (CurrentFrame == 0)
+	{
+		return FReply::Handled();
+	}
+
+	CurrentFrame--;
+	TimelineSlider->SetValue(CurrentFrame);
+
+	SetCurrentTimelineFrame(CurrentFrame);
+
+	return FReply::Handled();
+}
+
+const FSlateBrush* SChaosVDTimelineWidget::GetPlayOrPauseIcon() const
+{
+	return bIsPlaying ? FChaosVDStyle::Get().GetBrush("PauseIcon") : FChaosVDStyle::Get().GetBrush("PlayIcon");
+}
+
+
+void SChaosVDTimelineWidget::Tick(const FGeometry& AllottedGeometry, const double InCurrentTime, const float InDeltaTime)
+{
+	SCompoundWidget::Tick(AllottedGeometry, InCurrentTime, InDeltaTime);
+	
+	if (bIsPlaying)
+	{
+		if (CurrentFrame == MaxFrames)
+		{
+			Stop();
+		}
+
+		//TODO: Allow to customize the playback frametimes so it can be made an option or read from the recorded file
+		constexpr float PlaybackFrameTime = 0.016f;
+		
+		CurrentPlaybackTime += InDeltaTime;
+
+		if (CurrentPlaybackTime > PlaybackFrameTime)
+		{
+			CurrentPlaybackTime = 0.0f;
+			Next();
+		}
+	}
+}
+
+#undef LOCTEXT_NAMESPACE
diff --git a/Engine/Plugins/Experimental/ChaosVD/Source/ChaosVD/Public/ChaosVDEditorSettings.h b/Engine/Plugins/Experimental/ChaosVD/Source/ChaosVD/Public/ChaosVDEditorSettings.h
new file mode 100644
--- /dev/null
+++ b/Engine/Plugins/Experimental/ChaosVD/Source/ChaosVD/Public/ChaosVDEditorSettings.h
@@ -0,0 +1,23 @@
+﻿// Copyright Epic Games, Inc. All Rights Reserved.
+
+#pragma once
+
+#include "Engine/World.h"
+#include "Engine/StaticMesh.h"
+#include "UObject/Object.h"
+#include "UObject/SoftObjectPtr.h"
+
+#include "ChaosVDEditorSettings.generated.h"
+
+UCLASS(config = Engine)
+class UChaosVDEditorSettings : public UObject
+{
+	GENERATED_BODY()
+public:
+
+	UPROPERTY(Config, EditAnywhere, Category = Debug)
+	TSoftObjectPtr<UStaticMesh> DebugMesh;
+
+	UPROPERTY(Config, EditAnywhere, Category = Debug)
+	TSoftObjectPtr<UWorld> BasePhysicsVDWorld;
+};
\ No newline at end of file
diff --git a/Engine/Plugins/Experimental/ChaosVD/Source/ChaosVD/Public/ChaosVDEngine.h b/Engine/Plugins/Experimental/ChaosVD/Source/ChaosVD/Public/ChaosVDEngine.h
new file mode 100644
--- /dev/null
+++ b/Engine/Plugins/Experimental/ChaosVD/Source/ChaosVD/Public/ChaosVDEngine.h
@@ -0,0 +1,32 @@
+﻿// Copyright Epic Games, Inc. All Rights Reserved.
+
+#pragma once
+
+#include "Templates/SharedPointer.h"
+#include "Containers/Ticker.h"
+
+class FChaosVDPlaybackController;
+class FChaosVDScene;
+class FChaosVisualDebuggerMainUI;
+
+/** Core Implementation of the visual debugger - Owns the systems that are not UI */
+class FChaosVDEngine : public FTSTickerObjectBase, public TSharedFromThis<FChaosVDEngine>
+{
+public:
+
+	void Initialize();
+	
+	void DeInitialize();
+	
+	virtual bool Tick(float DeltaTime) override;
+	
+	TSharedPtr<FChaosVDScene>& GetCurrentScene() { return CurrentScene; };
+	TSharedPtr<FChaosVDPlaybackController>& GetPlaybackController() { return PlaybackController; };
+	
+private:
+
+	TSharedPtr<FChaosVDScene> CurrentScene;
+	TSharedPtr<FChaosVDPlaybackController> PlaybackController;
+	
+	bool bIsInitialized = false;
+};
diff --git a/Engine/Plugins/Experimental/ChaosVD/Source/ChaosVD/Public/ChaosVDModule.h b/Engine/Plugins/Experimental/ChaosVD/Source/ChaosVD/Public/ChaosVDModule.h
new file mode 100644
--- /dev/null
+++ b/Engine/Plugins/Experimental/ChaosVD/Source/ChaosVD/Public/ChaosVDModule.h
@@ -0,0 +1,24 @@
+// Copyright Epic Games, Inc. All Rights Reserved.
+
+#pragma once
+
+#include "Templates/SharedPointer.h"
+#include "Modules/ModuleManager.h"
+
+class SDockTab;
+class FSpawnTabArgs;
+
+class FChaosVDModule : public IModuleInterface
+{
+public:
+
+	/** IModuleInterface implementation */
+	virtual void StartupModule() override;
+	virtual void ShutdownModule() override;
+
+private:
+
+	void RegisterClassesCustomDetails() const;
+
+	TSharedRef<SDockTab> SpawnMainTab(const FSpawnTabArgs& Args) const;
+};
diff --git a/Engine/Plugins/Experimental/ChaosVD/Source/ChaosVD/Public/ChaosVDObjectDetailsTab.h b/Engine/Plugins/Experimental/ChaosVD/Source/ChaosVD/Public/ChaosVDObjectDetailsTab.h
new file mode 100644
--- /dev/null
+++ b/Engine/Plugins/Experimental/ChaosVD/Source/ChaosVD/Public/ChaosVDObjectDetailsTab.h
@@ -0,0 +1,38 @@
+﻿// Copyright Epic Games, Inc. All Rights Reserved.
+
+#pragma once
+
+#include "Templates/SharedPointer.h"
+#include "ChaosVDTabSpawnerBase.h"
+#include "Delegates/IDelegateInstance.h"
+
+struct FChaosVDParticleDebugData;
+
+class AActor;
+class IDetailsView;
+class FName;
+class FSpawnTabArgs;
+class FTabManager;
+class SDockTab;
+
+
+/** Spawns and handles and instance for the visual debugger details panel */
+class FChaosVDObjectDetailsTab : public FChaosVDTabSpawnerBase, public TSharedFromThis<FChaosVDObjectDetailsTab>
+{
+public:
+
+	FChaosVDObjectDetailsTab(const FName& InTabID, TSharedPtr<FTabManager> InTabManager, SChaosVDMainTab* InOwningTabWidget) : FChaosVDTabSpawnerBase(InTabID, InTabManager, InOwningTabWidget)
+	{
+	}
+
+	virtual ~FChaosVDObjectDetailsTab() override;
+
+protected:
+	
+	virtual TSharedRef<SDockTab> HandleTabSpawned(const FSpawnTabArgs& Args) override;
+	
+	void UpdateSelectedObject(AActor* NewObject) const;
+
+	FDelegateHandle SelectionDelegateHandle;
+	TSharedPtr<IDetailsView> DetailsPanel;
+};
diff --git a/Engine/Plugins/Experimental/ChaosVD/Source/ChaosVD/Public/ChaosVDOutputLogTab.h b/Engine/Plugins/Experimental/ChaosVD/Source/ChaosVD/Public/ChaosVDOutputLogTab.h
new file mode 100644
--- /dev/null
+++ b/Engine/Plugins/Experimental/ChaosVD/Source/ChaosVD/Public/ChaosVDOutputLogTab.h
@@ -0,0 +1,26 @@
+﻿// Copyright Epic Games, Inc. All Rights Reserved.
+
+#pragma once
+
+#include "Templates/SharedPointer.h"
+#include "ChaosVDTabSpawnerBase.h"
+
+class SOutputLog;
+class FOutputLogHistory;
+class FName;
+
+/** Spawns and handles and instance for the visual debugger Output Panel */
+class FChaosVDOutputLogTab : public FChaosVDTabSpawnerBase, public TSharedFromThis<FChaosVDOutputLogTab>
+{
+public:
+
+	FChaosVDOutputLogTab(const FName& InTabID, TSharedPtr<FTabManager> InTabManager, SChaosVDMainTab* InOwningTabWidget) : FChaosVDTabSpawnerBase(InTabID, InTabManager, InOwningTabWidget)
+	{
+	}
+
+	virtual TSharedRef<SDockTab> HandleTabSpawned(const FSpawnTabArgs& Args) override;
+
+private:
+	TSharedPtr<FOutputLogHistory> OutputLogHistory;
+	TSharedPtr<SOutputLog> OutputLog;
+};
diff --git a/Engine/Plugins/Experimental/ChaosVD/Source/ChaosVD/Public/ChaosVDParticleActor.h b/Engine/Plugins/Experimental/ChaosVD/Source/ChaosVD/Public/ChaosVDParticleActor.h
new file mode 100644
--- /dev/null
+++ b/Engine/Plugins/Experimental/ChaosVD/Source/ChaosVD/Public/ChaosVDParticleActor.h
@@ -0,0 +1,33 @@
+﻿// Copyright Epic Games, Inc. All Rights Reserved.
+
+#pragma once
+#include "ChaosVDRecording.h"
+#include "GameFramework/Actor.h"
+
+#include "ChaosVDParticleActor.generated.h"
+
+struct FChaosVDParticleDebugData;
+class USceneComponent;
+class UStaticMeshComponent;
+class UStaticMesh;
+
+/** Actor used to represent a Chaos Particle in the Visual Debugger's world */
+UCLASS(HideCategories=("Transform"))
+class AChaosVDParticleActor : public AActor
+{
+	GENERATED_BODY()
+public:
+	AChaosVDParticleActor(const FObjectInitializer& ObjectInitializer);
+
+	void UpdateFromRecordedData(const FChaosVDParticleDebugData& RecordedData);
+
+	void SetStaticMesh(UStaticMesh* NewStaticMesh);
+
+protected:
+
+	UPROPERTY(EditDefaultsOnly, Category= Debug)
+	TObjectPtr<UStaticMeshComponent> StaticMeshComponent;
+
+	UPROPERTY(VisibleAnywhere, Category="Chaos Visual Debugger Data")
+	FChaosVDParticleDebugData RecordedDebugData;
+};
diff --git a/Engine/Plugins/Experimental/ChaosVD/Source/ChaosVD/Public/ChaosVDParticleActorCustomization.h b/Engine/Plugins/Experimental/ChaosVD/Source/ChaosVD/Public/ChaosVDParticleActorCustomization.h
new file mode 100644
--- /dev/null
+++ b/Engine/Plugins/Experimental/ChaosVD/Source/ChaosVD/Public/ChaosVDParticleActorCustomization.h
@@ -0,0 +1,18 @@
+﻿// Copyright Epic Games, Inc. All Rights Reserved.
+
+#pragma once
+
+#include "Templates/SharedPointer.h"
+#include "IDetailCustomization.h"
+
+/** Custom details panel for the ChaosVD Particle Actor */
+class FChaosVDParticleActorCustomization : public IDetailCustomization
+{
+public:
+	
+	inline static FName ChaosVDCategoryName = FName("Chaos Visual Debugger Data");
+
+	static TSharedRef<IDetailCustomization> MakeInstance();
+
+	virtual void CustomizeDetails(IDetailLayoutBuilder& DetailBuilder) override;
+};
diff --git a/Engine/Plugins/Experimental/ChaosVD/Source/ChaosVD/Public/ChaosVDPlaybackController.h b/Engine/Plugins/Experimental/ChaosVD/Source/ChaosVD/Public/ChaosVDPlaybackController.h
new file mode 100644
--- /dev/null
+++ b/Engine/Plugins/Experimental/ChaosVD/Source/ChaosVD/Public/ChaosVDPlaybackController.h
@@ -0,0 +1,51 @@
+﻿// Copyright Epic Games, Inc. All Rights Reserved.
+
+#pragma once
+
+#include "Templates/SharedPointer.h"
+#include "Delegates/DelegateCombinations.h"
+#include "Delegates/Delegate.h"
+
+class FChaosVDScene;
+struct FChaosVDRecording;
+class FChaosVDPlaybackController;
+class FString;
+
+DECLARE_DELEGATE_OneParam(FChaosVDPlaybackControllerUpdated, FChaosVDPlaybackController*)
+
+/** Loads,unloads and owns a Chaos VD recording file */
+class FChaosVDPlaybackController
+{
+public:
+
+	FChaosVDPlaybackController(TWeakPtr<FChaosVDScene> InSceneToControl);
+
+	bool LoadChaosVDRecording(const FString& RecordingPath);
+	
+	void UnloadCurrentRecording();
+
+	void GoToRecordedStep(int32 FrameNumber, int32 Step);
+
+	TWeakPtr<FChaosVDRecording> GetCurrentRecording() { return LoadedRecording; }
+
+	int32 GetStepsForFrame(int32 FrameNumber) const;
+	int32 GetAvailableFramesNumber() const;
+
+	int32 GetCurrentFrame() const { return CurrentFrame; }
+	int32 GetCurrentStep() const { return CurrentStep; }
+
+	TWeakPtr<FChaosVDScene> GetControllerScene() { return SceneToControl; }
+
+	FChaosVDPlaybackControllerUpdated& OnControllerUpdated() { return ControllerUpdatedDelegate; }
+
+protected:
+
+	int32 CurrentFrame;
+	int32 CurrentStep;
+	
+	TSharedPtr<FChaosVDRecording> LoadedRecording;
+
+	TWeakPtr<FChaosVDScene> SceneToControl;
+
+	FChaosVDPlaybackControllerUpdated ControllerUpdatedDelegate;
+};
diff --git a/Engine/Plugins/Experimental/ChaosVD/Source/ChaosVD/Public/ChaosVDPlaybackViewportTab.h b/Engine/Plugins/Experimental/ChaosVD/Source/ChaosVD/Public/ChaosVDPlaybackViewportTab.h
new file mode 100644
--- /dev/null
+++ b/Engine/Plugins/Experimental/ChaosVD/Source/ChaosVD/Public/ChaosVDPlaybackViewportTab.h
@@ -0,0 +1,25 @@
+﻿// Copyright Epic Games, Inc. All Rights Reserved.
+
+#pragma once
+
+#include "Templates/SharedPointer.h"
+#include "ChaosVDTabSpawnerBase.h"
+
+class FChaosVDScene;
+class SChaosVDPlaybackViewport;
+
+/** Spawns and handles and instance for the visual debugger Playback viewport tab
+ * It contains the Playback controls and the 3D viewport
+ */
+class FChaosVDPlaybackViewportTab : public FChaosVDTabSpawnerBase, public TSharedFromThis<FChaosVDPlaybackViewportTab>
+{
+public:
+	FChaosVDPlaybackViewportTab(const FName& InTabID, TSharedPtr<FTabManager> InTabManager, SChaosVDMainTab* InOwningTabWidget) : FChaosVDTabSpawnerBase(InTabID, InTabManager, InOwningTabWidget)
+	{	
+	}
+	
+	virtual TSharedRef<SDockTab> HandleTabSpawned(const FSpawnTabArgs& Args) override;
+
+protected:
+	TSharedPtr<SChaosVDPlaybackViewport> PlaybackViewportWidget;
+};
diff --git a/Engine/Plugins/Experimental/ChaosVD/Source/ChaosVD/Public/ChaosVDScene.h b/Engine/Plugins/Experimental/ChaosVD/Source/ChaosVD/Public/ChaosVDScene.h
new file mode 100644
--- /dev/null
+++ b/Engine/Plugins/Experimental/ChaosVD/Source/ChaosVD/Public/ChaosVDScene.h
@@ -0,0 +1,63 @@
+﻿// Copyright Epic Games, Inc. All Rights Reserved.
+
+#pragma once
+
+#include "ChaosVDRecording.h"
+#include "Containers/UnrealString.h"
+#include "Containers/Map.h"
+#include "UObject/ObjectMacros.h"
+#include "UObject/GCObject.h"
+
+class AChaosVDParticleActor;
+class FReferenceCollector;
+class UWorld;
+
+DECLARE_MULTICAST_DELEGATE(FChaosVDSceneUpdatedDelegate)
+
+/** Recreates a UWorld from a recorded Chaos VD Frame */
+class FChaosVDScene : public FGCObject
+{
+public:
+
+	void Initialize();
+	void DeInitialize();
+
+	virtual void AddReferencedObjects(FReferenceCollector& Collector) override;
+	virtual FString GetReferencerName() const override
+	{
+		return TEXT("FChaosVDScene");
+	}
+
+	/** Called each time this Scene is modified */
+	FChaosVDSceneUpdatedDelegate& OnSceneUpdated() { return SceneUpdatedDelegate; }
+
+	/** Updates, Adds and Remove actors to match the provided Step Data */
+	void UpdateFromRecordedStepData(EChaosVDSolverType SolverID, const FChaosVDStepData& InRecordedStepData);
+
+	/** Deletes all actors of the Scene and underlying UWorld */
+	void CleanUpScene();
+
+	/** Returns the a ptr to the UWorld used to represent the current recorded frame data */
+	UWorld* GetUnderlyingWorld() const { return PhysicsVDWorld; };
+
+private:
+
+	/** Creates an ChaosVDParticle actor for the Provided recorded Particle Data */
+	AChaosVDParticleActor* SpawnParticleFromRecordedData(const FChaosVDParticleDebugData& InParticleData);
+
+	/** Returns the ID used to track this recorded particle data */
+	uint32 GetIDForRecordedParticleData(const FChaosVDParticleDebugData& InParticleData) const;
+
+	/** Creates the instance of the World which will be used the recorded data*/
+	UWorld* CreatePhysicsVDWorld() const;
+
+	/** Map of ID-ChaosVDParticle Actor. Used to keep track of actor instances and be able to modify them as needed*/
+	TMap<uint32, AChaosVDParticleActor*> ParticleVDInstancesByID;
+
+	/** UWorld instance used to represent the recorded debug data */
+	UWorld* PhysicsVDWorld = nullptr;
+
+	FChaosVDSceneUpdatedDelegate SceneUpdatedDelegate;
+
+	bool bIsInitialized = false;
+};
diff --git a/Engine/Plugins/Experimental/ChaosVD/Source/ChaosVD/Public/ChaosVDStyle.h b/Engine/Plugins/Experimental/ChaosVD/Source/ChaosVD/Public/ChaosVDStyle.h
new file mode 100644
--- /dev/null
+++ b/Engine/Plugins/Experimental/ChaosVD/Source/ChaosVD/Public/ChaosVDStyle.h
@@ -0,0 +1,30 @@
+// Copyright Epic Games, Inc. All Rights Reserved.
+
+#pragma once
+
+#include "Templates/SharedPointer.h"
+
+class FName;
+class FSlateStyleSet;
+class ISlateStyle;
+
+class FChaosVDStyle
+{
+public:
+
+	static void Initialize();
+
+	static void Shutdown();
+
+	static void ReloadTextures();
+
+	static const ISlateStyle& Get();
+
+	static FName GetStyleSetName();
+
+private:
+
+	static TSharedRef<FSlateStyleSet> Create();
+
+	static TSharedPtr<FSlateStyleSet> StyleInstance;
+};
\ No newline at end of file
diff --git a/Engine/Plugins/Experimental/ChaosVD/Source/ChaosVD/Public/ChaosVDTabSpawnerBase.h b/Engine/Plugins/Experimental/ChaosVD/Source/ChaosVD/Public/ChaosVDTabSpawnerBase.h
new file mode 100644
--- /dev/null
+++ b/Engine/Plugins/Experimental/ChaosVD/Source/ChaosVD/Public/ChaosVDTabSpawnerBase.h
@@ -0,0 +1,30 @@
+﻿// Copyright Epic Games, Inc. All Rights Reserved.
+
+#pragma once
+
+#include "Templates/SharedPointer.h"
+#include "UObject/NameTypes.h"
+
+class SChaosVDMainTab;
+class FName;
+class FTabManager;
+class SDockTab;
+class FSpawnTabArgs;
+
+/** Base class for any tab of the Chaos Visual Debugger tool*/
+class FChaosVDTabSpawnerBase
+{
+public:
+	virtual ~FChaosVDTabSpawnerBase() = default;
+
+	FChaosVDTabSpawnerBase(const FName& InTabID, TSharedPtr<FTabManager> InTabManager, SChaosVDMainTab* InOwningTabWidget);
+
+	virtual TSharedRef<SDockTab> HandleTabSpawned(const FSpawnTabArgs& Args) = 0;
+
+protected:
+
+	/** Ptr to the main tab of the owning visual debugger tool instance */
+	SChaosVDMainTab* OwningTabWidget;
+
+	UWorld* GetChaosVDWorld() const;
+};
diff --git a/Engine/Plugins/Experimental/ChaosVD/Source/ChaosVD/Public/ChaosVDTabsIDs.h b/Engine/Plugins/Experimental/ChaosVD/Source/ChaosVD/Public/ChaosVDTabsIDs.h
new file mode 100644
--- /dev/null
+++ b/Engine/Plugins/Experimental/ChaosVD/Source/ChaosVD/Public/ChaosVDTabsIDs.h
@@ -0,0 +1,16 @@
+﻿// Copyright Epic Games, Inc. All Rights Reserved.
+
+#pragma once
+
+class FName;
+
+class FChaosVDTabID
+{
+public:
+
+	static const FName ChaosVisualDebuggerTab;
+	static const FName PlaybackViewport;
+	static const FName WorldOutliner;
+	static const FName DetailsPanel;
+	static const FName OutputLog;
+};
diff --git a/Engine/Plugins/Experimental/ChaosVD/Source/ChaosVD/Public/ChaosVDWorldOutlinerTab.h b/Engine/Plugins/Experimental/ChaosVD/Source/ChaosVD/Public/ChaosVDWorldOutlinerTab.h
new file mode 100644
--- /dev/null
+++ b/Engine/Plugins/Experimental/ChaosVD/Source/ChaosVD/Public/ChaosVDWorldOutlinerTab.h
@@ -0,0 +1,41 @@
+﻿// Copyright Epic Games, Inc. All Rights Reserved.
+
+#pragma once
+
+#include "ChaosVDTabSpawnerBase.h"
+#include "Delegates/DelegateCombinations.h"
+#include "Delegates/Delegate.h"
+#include "Templates/SharedPointer.h"
+
+struct FChaosVDOutlinerItem;
+
+class AActor;
+class ISceneOutliner;
+class FChaosVDScene;
+class FSpawnTabArgs;
+class FTabManager;
+class SChaosVDSceneOutliner;
+class SDockTab;
+
+enum class EChaosVDParticleType : uint8;
+
+DECLARE_DELEGATE_OneParam(FChaosVDSelectionChanged, AActor* SelectedActor)
+
+/** Spawns and handles and instance for the visual debugger World Outliner */
+class FChaosVDWorldOutlinerTab : public FChaosVDTabSpawnerBase, public TSharedFromThis<FChaosVDWorldOutlinerTab>
+{
+public:
+	FChaosVDWorldOutlinerTab(const FName& InTabID, TSharedPtr<FTabManager> InTabManager, SChaosVDMainTab* InOwningTabWidget) : FChaosVDTabSpawnerBase(InTabID, InTabManager, InOwningTabWidget)
+	{	
+	}
+
+	virtual TSharedRef<SDockTab> HandleTabSpawned(const FSpawnTabArgs& Args) override;
+
+	FChaosVDSelectionChanged& OnSelectionChanged() { return ChaosVDSelectionChangedDelegate; }
+
+private:
+
+	TSharedPtr<ISceneOutliner> SceneOutlinerWidget;
+
+	FChaosVDSelectionChanged ChaosVDSelectionChangedDelegate;
+};
diff --git a/Engine/Plugins/Experimental/ChaosVD/Source/ChaosVD/Public/Widgets/SChaosVDMainTab.h b/Engine/Plugins/Experimental/ChaosVD/Source/ChaosVD/Public/Widgets/SChaosVDMainTab.h
new file mode 100644
--- /dev/null
+++ b/Engine/Plugins/Experimental/ChaosVD/Source/ChaosVD/Public/Widgets/SChaosVDMainTab.h
@@ -0,0 +1,46 @@
+﻿// Copyright Epic Games, Inc. All Rights Reserved.
+
+#pragma once
+
+#include "Framework/Docking/TabManager.h"
+#include "Templates/SharedPointer.h"
+#include "Widgets/SCompoundWidget.h"
+
+class FChaosVDEngine;
+class FChaosVDOutputLogTab;
+class FChaosVDPlaybackViewportTab;
+class FChaosVDObjectDetailsTab;
+class FChaosVDWorldOutlinerTab;
+class SDockTab;
+
+/** The main widget containing the Chaos Visual Debugger interface */
+class SChaosVDMainTab : public SCompoundWidget
+{
+public:
+
+	SLATE_BEGIN_ARGS(SChaosVDMainTab) {}
+		SLATE_ARGUMENT(TSharedPtr<SDockTab>, OwnerTab)
+	SLATE_END_ARGS()
+
+	virtual ~SChaosVDMainTab() override;
+	void Construct(const FArguments& InArgs, TSharedPtr<FChaosVDEngine> InChaosVDEngine);
+
+	TSharedRef<FChaosVDEngine> GetChaosVDEngineInstance() const { return ChaosVDEngine.ToSharedRef(); };
+
+private:
+
+	TSharedRef<FTabManager::FLayout> GenerateMainLayout();
+
+	void GenerateMainWindowMenu();
+
+	void BrowseAndOpenChaosVDFile();
+
+	TSharedPtr<FChaosVDEngine> ChaosVDEngine;
+
+	TSharedPtr<FChaosVDWorldOutlinerTab> WorldOutlinerTab;
+	TSharedPtr<FChaosVDObjectDetailsTab> ObjectDetailsTab;
+	TSharedPtr<FChaosVDPlaybackViewportTab> PlaybackViewportTab;
+	TSharedPtr<FChaosVDOutputLogTab> OutputLogTab;
+
+	TSharedPtr<FTabManager> TabManager;
+};
diff --git a/Engine/Plugins/Experimental/ChaosVD/Source/ChaosVD/Public/Widgets/SChaosVDPlaybackViewport.h b/Engine/Plugins/Experimental/ChaosVD/Source/ChaosVD/Public/Widgets/SChaosVDPlaybackViewport.h
new file mode 100644
--- /dev/null
+++ b/Engine/Plugins/Experimental/ChaosVD/Source/ChaosVD/Public/Widgets/SChaosVDPlaybackViewport.h
@@ -0,0 +1,45 @@
+﻿// Copyright Epic Games, Inc. All Rights Reserved.
+
+#pragma once
+
+#include "Templates/SharedPointer.h"
+#include "Widgets/SCompoundWidget.h"
+
+class FChaosVDPlaybackController;
+class SChaosVDTimelineWidget;
+struct FChaosVDRecording;
+class FChaosVDScene;
+class FLevelEditorViewportClient;
+class FSceneViewport;
+class SViewport;
+
+/* Widget that contains the 3D viewport and playback controls */
+class SChaosVDPlaybackViewport : public SCompoundWidget
+{
+public:
+
+	SLATE_BEGIN_ARGS( SChaosVDPlaybackViewport ){}
+	SLATE_END_ARGS()
+
+	virtual ~SChaosVDPlaybackViewport() override;
+
+	void Construct(const FArguments& InArgs, const UWorld* DefaultWorld, TWeakPtr<FChaosVDPlaybackController> InPlaybackController);
+
+protected:
+
+	TSharedPtr<FLevelEditorViewportClient> CreateViewportClient() const;
+
+	void OnPlaybackControllerUpdated(FChaosVDPlaybackController* Controller) const;
+
+	void OnFrameSelectionUpdated(int32 NewFrameIndex) const;
+	void OnStepSelectionUpdated(int32 NewStepIndex) const;
+	
+	TSharedPtr<SChaosVDTimelineWidget> FramesTimelineWidget;
+	TSharedPtr<SChaosVDTimelineWidget> StepsTimelineWidget;
+
+	TSharedPtr<FLevelEditorViewportClient> LevelViewportClient;
+	TSharedPtr<SViewport> ViewportWidget;
+	TSharedPtr<FSceneViewport> SceneViewport;
+
+	TWeakPtr<FChaosVDPlaybackController> PlaybackController;
+};
diff --git a/Engine/Plugins/Experimental/ChaosVD/Source/ChaosVD/Public/Widgets/SChaosVDTimelineWidget.h b/Engine/Plugins/Experimental/ChaosVD/Source/ChaosVD/Public/Widgets/SChaosVDTimelineWidget.h
new file mode 100644
--- /dev/null
+++ b/Engine/Plugins/Experimental/ChaosVD/Source/ChaosVD/Public/Widgets/SChaosVDTimelineWidget.h
@@ -0,0 +1,58 @@
+﻿// Copyright Epic Games, Inc. All Rights Reserved.
+
+#pragma once
+
+#include "Templates/SharedPointer.h"
+#include "Widgets/SCompoundWidget.h"
+
+class SSlider;
+class FReply;
+struct FSlateBrush;
+
+DECLARE_DELEGATE_OneParam(FChaosVDFrameChangedDelegate, int32)
+
+/** Simple timeline control widget */
+class SChaosVDTimelineWidget : public SCompoundWidget
+{
+public:
+
+	SLATE_BEGIN_ARGS( SChaosVDTimelineWidget ){}
+		SLATE_ARGUMENT(int32, MaxFrames)
+		SLATE_EVENT(FChaosVDFrameChangedDelegate, OnFrameChanged)
+		SLATE_ARGUMENT(bool, HidePlayStopButtons)
+	SLATE_END_ARGS()
+
+	void Construct(const FArguments& InArgs);
+
+	void UpdateMinMaxValue(float NewMin, float NewMax);
+
+	/** Brings back the state of the timeline to its original state*/
+	void ResetTimeline();
+
+	/** Called when a new frame is manually selected or auto-updated during playback */
+	FChaosVDFrameChangedDelegate& OnFrameChanged() { return FrameChangedDelegate; }
+
+protected:
+
+	void SetCurrentTimelineFrame(float FrameNumber);
+
+	FReply  Play();
+	FReply  Stop();
+	FReply  Next();
+	FReply  Prev();
+
+	const FSlateBrush* GetPlayOrPauseIcon() const;
+
+	virtual void Tick(const FGeometry& AllottedGeometry, const double InCurrentTime, const float InDeltaTime) override;
+
+	TSharedPtr<SSlider> TimelineSlider;
+
+	int32 CurrentFrame = 0;
+	int32 MinFrames = 0;
+	int32 MaxFrames = 1000;
+
+	bool bIsPlaying = false;
+	float CurrentPlaybackTime = 0.0f;
+
+	FChaosVDFrameChangedDelegate FrameChangedDelegate;
+};
diff --git a/Engine/Source/Runtime/Experimental/ChaosVisualDebugger/ChaosVDRuntime.Build.cs b/Engine/Source/Runtime/Experimental/ChaosVisualDebugger/ChaosVDRuntime.Build.cs
new file mode 100644
--- /dev/null
+++ b/Engine/Source/Runtime/Experimental/ChaosVisualDebugger/ChaosVDRuntime.Build.cs
@@ -0,0 +1,20 @@
+﻿// Copyright Epic Games, Inc. All Rights Reserved.
+
+using System.Linq;
+using System.Collections.Generic;
+
+namespace UnrealBuildTool.Rules
+{
+	public class ChaosVDRuntime : ModuleRules
+	{
+		public ChaosVDRuntime(ReadOnlyTargetRules Target) : base(Target)
+		{
+			PublicDependencyModuleNames.AddRange(
+				new string[] {
+					"Core",
+					"CoreUObject"
+				}
+			);
+		}
+	}
+}
\ No newline at end of file
diff --git a/Engine/Source/Runtime/Experimental/ChaosVisualDebugger/Private/ChaosVDRecording.cpp b/Engine/Source/Runtime/Experimental/ChaosVisualDebugger/Private/ChaosVDRecording.cpp
new file mode 100644
--- /dev/null
+++ b/Engine/Source/Runtime/Experimental/ChaosVisualDebugger/Private/ChaosVDRecording.cpp
@@ -0,0 +1,10 @@
+﻿// Copyright Epic Games, Inc. All Rights Reserved.
+
+#include "ChaosVisualDebugger/Public/ChaosVDRecording.h"
+
+FChaosVDFrameData& FChaosVDRecording::AddFrame()
+{
+	RecordedFramesData.Add(FChaosVDFrameData());
+
+	return RecordedFramesData.Last();
+}
diff --git a/Engine/Source/Runtime/Experimental/ChaosVisualDebugger/Private/ChaosVDRuntimeModule.cpp b/Engine/Source/Runtime/Experimental/ChaosVisualDebugger/Private/ChaosVDRuntimeModule.cpp
new file mode 100644
--- /dev/null
+++ b/Engine/Source/Runtime/Experimental/ChaosVisualDebugger/Private/ChaosVDRuntimeModule.cpp
@@ -0,0 +1,93 @@
+﻿// Copyright Epic Games, Inc. All Rights Reserved.
+
+#include "ChaosVisualDebugger/Public/ChaosVDRuntimeModule.h"
+
+#include "ChaosVDRecording.h"
+#include "Containers/Array.h"
+#include "HAL/IConsoleManager.h"
+#include "Misc/FileHelper.h"
+#include "Misc/Paths.h"
+#include "Serialization/MemoryWriter.h"
+#include "Modules/ModuleManager.h"
+
+IMPLEMENT_MODULE(FChaosVDRuntimeModule, ChaosVDRuntime);
+
+FAutoConsoleCommand ChaosVDStartRecordingCommand(
+	TEXT("p.Chaos.StartVDRecording"),
+	TEXT("Turn on the recording of debugging data"),
+	FConsoleCommandWithArgsDelegate::CreateLambda([](const TArray<FString>& Args)
+	{
+		FChaosVDRuntimeModule::Get().StartRecording();
+	})
+);
+
+FAutoConsoleCommand StopVDStartRecordingCommand(
+	TEXT("p.Chaos.StopVDRecording"),
+	TEXT("Turn off the recording of debugging data"),
+	FConsoleCommandWithArgsDelegate::CreateLambda([](const TArray<FString>& Args)
+	{
+		FChaosVDRuntimeModule::Get().StopRecording();
+	})
+);
+
+FChaosVDRuntimeModule& FChaosVDRuntimeModule::Get()
+{
+	return FModuleManager::Get().LoadModuleChecked<FChaosVDRuntimeModule>(TEXT("ChaosVDRuntime"));
+}
+
+void FChaosVDRuntimeModule::StartupModule()
+{
+}
+
+void FChaosVDRuntimeModule::ShutdownModule()
+{
+}
+
+void FChaosVDRuntimeModule::StartRecording()
+{
+	// Note: The code in this method is pretty much place holder code for testing purposes
+
+	if (bIsRecording)
+	{
+		return;
+	}
+
+	bIsRecording = true;
+
+	CurrentRecording = MakeShared<FChaosVDRecording>();
+}
+
+void FChaosVDRuntimeModule::StopRecording()
+{
+	// Note: The code in this method is pretty much place holder code for testing purposes
+
+	if (!ensure(bIsRecording))
+	{
+		return;
+	}
+
+	if (!ensure(CurrentRecording.IsValid()))
+	{
+		return;
+	}
+
+	const TCHAR* FilePrefix = TEXT("ChaosVD");
+	const FString FullPathPrefix = FPaths::ProfilingDir() / FilePrefix;
+
+	int32 Tries = 0;
+	FString UseFileName;
+	do
+	{
+		UseFileName = FString::Printf(TEXT("%s_%d.cvd"), *FullPathPrefix, Tries++);
+	} while (IFileManager::Get().FileExists(*UseFileName));
+
+
+	TArray<uint8> SerializedRecordedData;
+	FMemoryWriter RecordedDataWriter(SerializedRecordedData, false);
+
+	FChaosVDRecording::StaticStruct()->SerializeBin(RecordedDataWriter,CurrentRecording.Get());
+
+	FFileHelper::SaveArrayToFile(SerializedRecordedData, *UseFileName);
+
+	bIsRecording = false;
+}
diff --git a/Engine/Source/Runtime/Experimental/ChaosVisualDebugger/Public/ChaosVDRecording.h b/Engine/Source/Runtime/Experimental/ChaosVisualDebugger/Public/ChaosVDRecording.h
new file mode 100644
--- /dev/null
+++ b/Engine/Source/Runtime/Experimental/ChaosVisualDebugger/Public/ChaosVDRecording.h
@@ -0,0 +1,121 @@
+﻿// Copyright Epic Games, Inc. All Rights Reserved.
+
+#pragma once
+#include "UObject/ObjectMacros.h"
+#include "Containers/UnrealString.h"
+
+#include "ChaosVDRecording.generated.h"
+
+UENUM()
+enum class EChaosVDParticleType : uint8
+{
+	Static,
+	Kinematic,
+	Rigid,
+	Clustered,
+	StaticMesh,
+	SkeletalMesh,
+	GeometryCollection,
+	Unknown
+};
+
+UENUM()
+enum class EChaosVDSolverType
+{
+	Rigid
+};
+
+USTRUCT()
+struct FChaosVDParticleDebugData
+{
+	GENERATED_BODY()
+
+	UPROPERTY(VisibleAnywhere, Category= "Chaos Visual Debugger Data")
+	EChaosVDParticleType ParticleType;
+
+	UPROPERTY(VisibleAnywhere, Category= "Chaos Visual Debugger Data")
+	FString DebugName;
+
+	UPROPERTY(VisibleAnywhere, Category= "Chaos Visual Debugger Data")
+	int32 ParticleIndex;
+
+	UPROPERTY(VisibleAnywhere, Category= "Chaos Visual Debugger Data")
+	FVector Position;
+
+	UPROPERTY(VisibleAnywhere, Category= "Chaos Visual Debugger Data")
+	FQuat Rotation;
+};
+
+USTRUCT()
+struct FChaosVDStepData
+{
+	GENERATED_BODY()
+
+	UPROPERTY()
+	TArray<FChaosVDParticleDebugData> RecordedParticles;
+};
+
+USTRUCT()
+struct FChaosVDSolverData
+{
+	GENERATED_BODY()
+
+	UPROPERTY()
+	EChaosVDSolverType SolverType;
+
+	UPROPERTY()
+	TArray<FChaosVDStepData> SolverSteps;
+};
+
+USTRUCT()
+struct FChaosVDEventData
+{
+	GENERATED_BODY()
+
+	UPROPERTY()
+	FString EventID;
+
+	UPROPERTY()
+	TArray<uint8> SerializedEventData;
+};
+
+USTRUCT()
+struct CHAOSVDRUNTIME_API FChaosVDFrameData
+{
+	GENERATED_BODY()
+
+	UPROPERTY()
+	TMap<EChaosVDSolverType, FChaosVDSolverData> RecordedSolvers;
+
+	UPROPERTY()
+	TArray<FChaosVDEventData> RecordedEvents;
+
+	FChaosVDSolverData& GetSolverData(EChaosVDSolverType SolverType) { return RecordedSolvers.Contains(SolverType) ? RecordedSolvers[SolverType] : RecordedSolvers.Add(SolverType,FChaosVDSolverData()); }
+
+	//FChaosVDSolverData& AddSolverData(EChaosVDSolverType SolverType) { return RecordedSolvers.Contains(SolverType) ? RecordedSolvers[SolverType] : RecordedSolvers.Add(SolverType,FChaosVDSolverData()); }
+};
+
+USTRUCT()
+struct FChaosVDRecordingHeader
+{
+	GENERATED_BODY()
+
+	UPROPERTY()
+	bool bIsRealTimeRecording = false;
+};
+
+USTRUCT()
+struct CHAOSVDRUNTIME_API FChaosVDRecording 
+{
+	GENERATED_BODY()
+
+	UPROPERTY()
+	FChaosVDRecordingHeader RecordingHeader;
+
+	UPROPERTY()
+	TArray<FChaosVDFrameData> RecordedFramesData;
+
+	FChaosVDFrameData& GetCurrentFrame() { return RecordedFramesData.Num() ? RecordedFramesData.Last() : AddFrame(); }
+
+	FChaosVDFrameData& AddFrame();
+};
diff --git a/Engine/Source/Runtime/Experimental/ChaosVisualDebugger/Public/ChaosVDRuntimeModule.h b/Engine/Source/Runtime/Experimental/ChaosVisualDebugger/Public/ChaosVDRuntimeModule.h
new file mode 100644
--- /dev/null
+++ b/Engine/Source/Runtime/Experimental/ChaosVisualDebugger/Public/ChaosVDRuntimeModule.h
@@ -0,0 +1,31 @@
+﻿// Copyright Epic Games, Inc. All Rights Reserved.
+
+#pragma once
+
+#include "Templates/SharedPointer.h"
+#include "Modules/ModuleInterface.h"
+
+struct FChaosVDRecording;
+
+class CHAOSVDRUNTIME_API FChaosVDRuntimeModule : public IModuleInterface
+{
+public:
+
+	static FChaosVDRuntimeModule& Get();
+	
+	virtual void StartupModule() override;
+	virtual void ShutdownModule() override;
+
+	TSharedPtr<FChaosVDRecording>& GetRecording() { return CurrentRecording; }
+
+	void StartRecording();
+	void StopRecording();
+
+	bool IsRecording() const { return bIsRecording; }
+
+private:
+
+	bool bIsRecording = false;
+	TSharedPtr<FChaosVDRecording> CurrentRecording;
+
+};
