Fix CIS issues in ChaosVDObjectDetailsTab.cpp

(removed module names from include paths)

diff --git a/Engine/Plugins/Experimental/ChaosVD/Source/ChaosVD/Private/ChaosVDObjectDetailsTab.cpp b/Engine/Plugins/Experimental/ChaosVD/Source/ChaosVD/Private/ChaosVDObjectDetailsTab.cpp
--- a/Engine/Plugins/Experimental/ChaosVD/Source/ChaosVD/Private/ChaosVDObjectDetailsTab.cpp
+++ b/Engine/Plugins/Experimental/ChaosVD/Source/ChaosVD/Private/ChaosVDObjectDetailsTab.cpp
@@ -8,7 +8,7 @@
 #include "PropertyEditorModule.h"
 #include "Selection.h"
 #include "Templates/SharedPointer.h"
-#include "UnrealEd/Public/Editor.h"
+#include "Editor.h"
 #include "Widgets/Docking/SDockTab.h"
 
 #define LOCTEXT_NAMESPACE "ChaosVisualDebugger"
