[Chaos VD]

- Replaced custom Serialization with Trace
- Hooked up UI playback controls for sub-steps
- Initial work to support recording multiple solvers

#rnx

diff --git a/Engine/Plugins/Experimental/ChaosVD/Source/ChaosVD/ChaosVD.Build.cs b/Engine/Plugins/Experimental/ChaosVD/Source/ChaosVD/ChaosVD.Build.cs
--- a/Engine/Plugins/Experimental/ChaosVD/Source/ChaosVD/ChaosVD.Build.cs
+++ b/Engine/Plugins/Experimental/ChaosVD/Source/ChaosVD/ChaosVD.Build.cs
@@ -12,7 +12,8 @@ public class ChaosVD : ModuleRules
 				"Core",
 				"CoreUObject",
 				"Engine",
-				"ChaosVDRuntime",
+				"ChaosVDRuntime", 
+				"TraceServices",
 			});
 
 		PrivateDependencyModuleNames.AddRange(
@@ -28,7 +29,10 @@ public class ChaosVD : ModuleRules
 				"SlateCore", 
 				"OutputLog", 
 				"SceneOutliner", 
-				"WorkspaceMenuStructure",
+				"WorkspaceMenuStructure", 
+				"TraceAnalysis", 
+				"TraceInsights",
+				"TraceLog",
 			}
 			);
 	}
diff --git a/Engine/Plugins/Experimental/ChaosVD/Source/ChaosVD/Private/ChaosVDEngine.cpp b/Engine/Plugins/Experimental/ChaosVD/Source/ChaosVD/Private/ChaosVDEngine.cpp
--- a/Engine/Plugins/Experimental/ChaosVD/Source/ChaosVD/Private/ChaosVDEngine.cpp
+++ b/Engine/Plugins/Experimental/ChaosVD/Source/ChaosVD/Private/ChaosVDEngine.cpp
@@ -2,8 +2,10 @@
 
 #include "ChaosVDEngine.h"
 
+#include "ChaosVDModule.h"
 #include "ChaosVDPlaybackController.h"
 #include "ChaosVDScene.h"
+#include "Trace/ChaosVDTraceManager.h"
 
 void FChaosVDEngine::Initialize()
 {
@@ -30,7 +32,17 @@ void FChaosVDEngine::DeInitialize()
 	}
 
 	CurrentScene->DeInitialize();
+	CurrentScene.Reset();
+	PlaybackController.Reset();
 }
+
+void FChaosVDEngine::LoadRecording(const FString& FilePath)
+{
+	CurrentSessionName = FChaosVDModule::Get().GetTraceManager()->LoadTraceFile(FilePath);
+
+	PlaybackController->LoadChaosVDRecordingFromTraceSession(CurrentSessionName);
+}
+
 bool FChaosVDEngine::Tick(float DeltaTime)
 {
 	return true;
diff --git a/Engine/Plugins/Experimental/ChaosVD/Source/ChaosVD/Private/ChaosVDModule.cpp b/Engine/Plugins/Experimental/ChaosVD/Source/ChaosVD/Private/ChaosVDModule.cpp
--- a/Engine/Plugins/Experimental/ChaosVD/Source/ChaosVD/Private/ChaosVDModule.cpp
+++ b/Engine/Plugins/Experimental/ChaosVD/Source/ChaosVD/Private/ChaosVDModule.cpp
@@ -10,9 +10,16 @@
 #include "ChaosVDEngine.h"
 #include "WorkspaceMenuStructure.h"
 #include "WorkspaceMenuStructureModule.h"
+#include "Trace/ChaosVDTraceManager.h"
+#include "Misc/Guid.h"
 
 #define LOCTEXT_NAMESPACE "ChaosVisualDebugger"
 
+FChaosVDModule& FChaosVDModule::Get()
+{
+	return FModuleManager::Get().LoadModuleChecked<FChaosVDModule>(TEXT("ChaosVD"));
+}
+
 void FChaosVDModule::StartupModule()
 {	
 	FChaosVDStyle::Initialize();
@@ -26,6 +33,8 @@ void FChaosVDModule::StartupModule()
 								//TODO: Hook up the final icon
 								.SetIcon(FSlateIcon(FAppStyle::GetAppStyleSetName(), "CollisionAnalyzer.TabIcon"))
 								.SetGroup(WorkspaceMenu::GetMenuStructure().GetDeveloperToolsDebugCategory());
+
+	ChaosVDTraceManager = MakeShared<FChaosVDTraceManager>();
 }
 
 void FChaosVDModule::ShutdownModule()
@@ -41,7 +50,7 @@ void FChaosVDModule::RegisterClassesCustomDetails() const
 	PropertyModule.RegisterCustomClassLayout("ChaosVDParticleActor", FOnGetDetailCustomizationInstance::CreateStatic(&FChaosVDParticleActorCustomization::MakeInstance));
 }
 
-TSharedRef<SDockTab> FChaosVDModule::SpawnMainTab(const FSpawnTabArgs& Args) const
+TSharedRef<SDockTab> FChaosVDModule::SpawnMainTab(const FSpawnTabArgs& Args)
 {
 	TSharedRef<SDockTab> MainTabInstance =
 		SNew(SDockTab)
@@ -62,9 +71,44 @@ TSharedRef<SDockTab> FChaosVDModule::SpawnMainTab(const FSpawnTabArgs& Args) con
 
 	MainTabInstance->SetTabIcon(FChaosVDStyle::Get().GetBrush("TabIconPlaybackViewport"));
 
+	const FGuid InstanceGuid = ChaosVDEngineInstance->GetInstanceGuid();
+	RegisterChaosVDInstance(InstanceGuid, ChaosVDEngineInstance);
+
+	// Workaround. Currently the ChaosVD Engine instance determines the lifetime of the Editor world and other objects
+	// Some widgets, like UE Level viewport tries to iterate on these objects on destruction
+	// For now we can avoid any crashes by just de-initializing ChaosVD Engine on the next frame but that is not the real fix.
+	
+	//TODO: Ensure that systems that uses the Editor World we create know beforehand when it is about to be Destroyed and GC'd
+	MainTabInstance->SetOnTabClosed(SDockTab::FOnTabClosedCallback::CreateLambda( [this, InstanceGuid](TSharedRef<SDockTab>)
+	{
+		FTSTicker::GetCoreTicker().AddTicker(FTickerDelegate::CreateLambda([this, InstanceGuid](float DeltaTime)->bool
+		{
+			DeregisterChaosVDInstance(InstanceGuid);
+			return false;
+		}));
+	}));
+
 	return MainTabInstance;
 }
 
+void FChaosVDModule::RegisterChaosVDInstance(const FGuid& InstanceGuid, TSharedPtr<FChaosVDEngine> Instance)
+{
+	ActiveChaosVDInstances.Add(InstanceGuid, Instance);
+}
+
+void FChaosVDModule::DeregisterChaosVDInstance(const FGuid& InstanceGuid)
+{
+	if (TSharedPtr<FChaosVDEngine>* InstancePtrPtr = ActiveChaosVDInstances.Find(InstanceGuid))
+	{
+		if (TSharedPtr<FChaosVDEngine> InstancePtr = *InstancePtrPtr)
+		{
+			InstancePtr->DeInitialize();
+		}
+	
+		ActiveChaosVDInstances.Remove(InstanceGuid);
+	}	
+}
+
 #undef LOCTEXT_NAMESPACE
 	
 IMPLEMENT_MODULE(FChaosVDModule, ChaosVD)
\ No newline at end of file
diff --git a/Engine/Plugins/Experimental/ChaosVD/Source/ChaosVD/Private/ChaosVDPlaybackController.cpp b/Engine/Plugins/Experimental/ChaosVD/Source/ChaosVD/Private/ChaosVDPlaybackController.cpp
--- a/Engine/Plugins/Experimental/ChaosVD/Source/ChaosVD/Private/ChaosVDPlaybackController.cpp
+++ b/Engine/Plugins/Experimental/ChaosVD/Source/ChaosVD/Private/ChaosVDPlaybackController.cpp
@@ -2,21 +2,28 @@
 
 #include "ChaosVDPlaybackController.h"
 
+#include "ChaosVDModule.h"
 #include "ChaosVDRecording.h"
 #include "ChaosVDScene.h"
-#include "Misc/FileHelper.h"
-#include "Serialization/MemoryReader.h"
+#include "Trace/ChaosVDTraceManager.h"
+#include "Trace/ChaosVDTraceProvider.h"
+#include "TraceServices/Model/AnalysisSession.h"
 
-FChaosVDPlaybackController::FChaosVDPlaybackController(TWeakPtr<FChaosVDScene> InSceneToControl)
-	: CurrentFrame(0),
-	CurrentStep(0)
+
+FChaosVDPlaybackController::FChaosVDPlaybackController(const TWeakPtr<FChaosVDScene>& InSceneToControl)
 {
 	SceneToControl = InSceneToControl;
 }
 
-bool FChaosVDPlaybackController::LoadChaosVDRecording(const FString& RecordingPath)
+FChaosVDPlaybackController::~FChaosVDPlaybackController()
 {
-	if (RecordingPath.IsEmpty())
+	constexpr bool bBroadcastControllerUpdate = false;
+	UnloadCurrentRecording(bBroadcastControllerUpdate);
+}
+
+bool FChaosVDPlaybackController::LoadChaosVDRecordingFromTraceSession(const FString& InSessionName)
+{
+	if (InSessionName.IsEmpty())
 	{
 		return false;
 	}
@@ -26,66 +33,117 @@ bool FChaosVDPlaybackController::LoadChaosVDRecording(const FString& RecordingPa
 		UnloadCurrentRecording();
 	}
 
-	//TODO: Show a notification or loading bar. The plan is support streaming of the file so maybe we don't need this at all
-
-	TArray<uint8> SerializedDebugData;
-	if (!ensure(FFileHelper::LoadFileToArray(SerializedDebugData, *RecordingPath)))
+	if (const TSharedPtr<const TraceServices::IAnalysisSession> TraceSession = FChaosVDModule::Get().GetTraceManager()->GetSession(InSessionName))
 	{
-		return false;
+		if (const FChaosVDTraceProvider* ChaosVDProvider = TraceSession->ReadProvider<FChaosVDTraceProvider>(FChaosVDTraceProvider::ProviderName))
+		{
+			LoadedRecording = ChaosVDProvider->GetRecordingForSession();
+		}
 	}
 
-	LoadedRecording = MakeShared<FChaosVDRecording>();
-
-	FMemoryReader SerializedDataReader(SerializedDebugData, false);					
-	FChaosVDRecording::StaticStruct()->SerializeBin(SerializedDataReader, LoadedRecording.Get());
-
-	if (!ensure(!SerializedDataReader.IsError()))
+	if (!LoadedRecording.IsValid())
 	{
-		//TODO: Pop up an error message informing the user
-		LoadedRecording.Reset();
 		return false;
 	}
 
-	GoToRecordedStep(0,0);
+	LoadedRecording->OnRecordingUpdated().AddRaw(this, &FChaosVDPlaybackController::HandleCurrentRecordingUpdated);
 
-	OnControllerUpdated().ExecuteIfBound(this);
+	for (TMap<int32, TArray<FChaosVDSolverFrameData>>::TConstIterator SolversIterator = LoadedRecording->GetAvailableSolvers().CreateConstIterator(); SolversIterator; ++SolversIterator)
+	{
+		GoToRecordedStep(SolversIterator->Key, 0, 0);
+	}
+	
+	OnControllerUpdated().ExecuteIfBound(AsWeak());
 
 	return true;
 }
 
-void FChaosVDPlaybackController::UnloadCurrentRecording()
+void FChaosVDPlaybackController::UnloadCurrentRecording(const bool bBroadcastUpdate)
 {
-	CurrentFrame = 0;
-	CurrentStep = 0;
-	LoadedRecording.Reset();
-	
+	CurrentFramePerTrack.Reset();
+	CurrentStepPerTrack.Reset();
+
+	if (LoadedRecording.IsValid())
+	{
+		LoadedRecording->OnRecordingUpdated().RemoveAll(this);
+
+		if (const TSharedPtr<const TraceServices::IAnalysisSession> TraceSession = FChaosVDModule::Get().GetTraceManager()->GetSession(LoadedRecording->SessionName))
+		{
+			TraceServices::IAnalysisSession* TraceSessionPtr = const_cast<TraceServices::IAnalysisSession*>(TraceSession.Get());
+			if (FChaosVDTraceProvider* ChaosVDProvider = TraceSessionPtr->EditProvider<FChaosVDTraceProvider>(FChaosVDTraceProvider::ProviderName))
+			{
+				TraceServices::FAnalysisSessionEditScope SessionEditScope(*TraceSessionPtr);
+				ChaosVDProvider->DeleteRecordingInstanceForSession();
+			}
+		}
+
+		LoadedRecording.Reset();
+	}
+
 	if (const TSharedPtr<FChaosVDScene> SceneToControlSharedPtr = SceneToControl.Pin())
 	{
-		SceneToControlSharedPtr->CleanUpScene();
+		if (SceneToControlSharedPtr->IsInitialized())
+		{
+			SceneToControlSharedPtr->CleanUpScene();
+		}	
+	}
+
+	if (bBroadcastUpdate)
+	{
+		const TWeakPtr<FChaosVDPlaybackController> ThisWeakPtr = DoesSharedInstanceExist() ? AsWeak() : nullptr;
+		OnControllerUpdated().ExecuteIfBound(ThisWeakPtr);
 	}
-	
-	OnControllerUpdated().ExecuteIfBound(this);
 }
 
-void FChaosVDPlaybackController::GoToRecordedStep(int32 FrameNumber, int32 Step)
+void FChaosVDPlaybackController::GoToRecordedStep(const int32 InTrackID, const int32 FrameNumber, const int32 Step)
 {
 	if (const TSharedPtr<FChaosVDScene> SceneToControlSharedPtr = SceneToControl.Pin())
 	{
 		if (ensure(LoadedRecording.IsValid()))
 		{
-			if (!LoadedRecording->RecordedFramesData.IsValidIndex(FrameNumber))
+			const EChaosVDFrameLoadState FrameState = LoadedRecording->GetFrameState(InTrackID, FrameNumber);
+			if (FrameState == EChaosVDFrameLoadState::Unknown)
 			{
+				// Invalid Frame Number
 				return;
 			}
 
-			//TODO: Handle different steps per solver if that is something that can happen
-			for (TPair<EChaosVDSolverType, FChaosVDSolverData>& SolverData : LoadedRecording->RecordedFramesData[FrameNumber].RecordedSolvers)
 			{
-				SceneToControlSharedPtr->UpdateFromRecordedStepData(SolverData.Key, SolverData.Value.SolverSteps[Step]);
+				const TSharedPtr<const TraceServices::IAnalysisSession> TraceSession = FChaosVDModule::Get().GetTraceManager()->GetSession(LoadedRecording->SessionName);
+				if (!ensure(TraceSession))
+				{
+					return;	
+				}
+
+				// TODO: This will not cover a future case where the recording is not own/populated by Trace
+				// If in the future we decide implement the CVD format standalone with streaming support again,
+				// we will need to add a lock to the file. A feature that might need that is Recording Clips
+				TraceServices::FAnalysisSessionReadScope SessionReadScope(*TraceSession);
+
+				const FChaosVDSolverFrameData* SolverFrameData = LoadedRecording->GetFrameForSolver(InTrackID, FrameNumber);
+				if (SolverFrameData && ensure(SolverFrameData->SolverSteps.IsValidIndex(Step)))
+				{
+					SceneToControlSharedPtr->UpdateFromRecordedStepData(InTrackID, SolverFrameData->DebugName, SolverFrameData->SolverSteps[Step]);
+				}
+			}
+
+			if (int32* CurrentFrameNumber = CurrentFramePerTrack.Find(InTrackID))
+			{
+				 *CurrentFrameNumber = FrameNumber;
+			}
+			else
+			{
+				CurrentFramePerTrack.Add(InTrackID, FrameNumber);
 			}
 
-			CurrentStep = Step;
-			CurrentFrame = FrameNumber;
+			if (int32* CurrentStepNumber = CurrentStepPerTrack.Find(InTrackID))
+			{
+				*CurrentStepNumber = Step;
+			}
+			else
+			{
+				CurrentStepPerTrack.Add(InTrackID, Step);
+			}
 		}
 	}
 	else
@@ -94,32 +152,77 @@ void FChaosVDPlaybackController::GoToRecordedStep(int32 FrameNumber, int32 Step)
 	}
 }
 
-int32 FChaosVDPlaybackController::GetStepsForFrame(int32 FrameNumber) const
+int32 FChaosVDPlaybackController::GetStepsForFrame(const int32 InTrackID, const int32 FrameNumber) const
 {
 	if (LoadedRecording.IsValid())
 	{
-		if (!LoadedRecording->RecordedFramesData.IsValidIndex(FrameNumber))
+		if (const FChaosVDSolverFrameData* FrameData = LoadedRecording->GetFrameForSolver(InTrackID, FrameNumber))
 		{
-			return INDEX_NONE;
+			return FrameData->SolverSteps.Num();
 		}
-
-		for (const TPair<EChaosVDSolverType, FChaosVDSolverData>& SolverData : LoadedRecording->RecordedFramesData[FrameNumber].RecordedSolvers)
+		else
 		{
-			//TODO: Handle different steps per solver if that is something that can happen
-			return SolverData.Value.SolverSteps.Num();
+			return INDEX_NONE;
 		}
 	}
 
 	return	INDEX_NONE;
 }
 
-int32 FChaosVDPlaybackController::GetAvailableFramesNumber() const
+int32 FChaosVDPlaybackController::GetAvailableFramesNumber(const int32 InTrackID) const
 {
 	if (!LoadedRecording.IsValid())
 	{
 		return INDEX_NONE;
 	}
 
-	return LoadedRecording->RecordedFramesData.Num();
+	return LoadedRecording->GetAvailableFramesNumber(InTrackID);
+}
+
+int32 FChaosVDPlaybackController::GetAvailableSolversNumber() const
+{
+	if (!LoadedRecording.IsValid())
+	{
+		return INDEX_NONE;
+	}
+
+	return LoadedRecording->GetAvailableSolvers().Num();
+}
+
+int32 FChaosVDPlaybackController::GetActiveSolverTrackID() const
+{
+	//For Now use the Last Solver if any
+	int32 SolverID = INDEX_NONE;
+	for (const TPair<int32, TArray<FChaosVDSolverFrameData>>& SolverData : LoadedRecording->GetAvailableSolvers())
+	{
+		SolverID = SolverData.Key;
+	}
+
+	return SolverID;
+}
+
+int32 FChaosVDPlaybackController::GetCurrentFrame(const int32 InTrackID) const
+{
+	if (const int32* FrameNumber = CurrentFramePerTrack.Find(InTrackID))
+	{
+		return *FrameNumber;
+	}
+
+	return INDEX_NONE;
+}
+
+int32 FChaosVDPlaybackController::GetCurrentStep(const int32 InTrackID) const
+{
+	if (const int32* StepNumber = CurrentStepPerTrack.Find(InTrackID))
+	{
+		return *StepNumber;
+	}
+
+	return INDEX_NONE;
+}
+
+void FChaosVDPlaybackController::HandleCurrentRecordingUpdated()
+{
+	OnControllerUpdated().ExecuteIfBound(AsWeak());
 }
 
diff --git a/Engine/Plugins/Experimental/ChaosVD/Source/ChaosVD/Private/ChaosVDScene.cpp b/Engine/Plugins/Experimental/ChaosVD/Source/ChaosVD/Private/ChaosVDScene.cpp
--- a/Engine/Plugins/Experimental/ChaosVD/Source/ChaosVD/Private/ChaosVDScene.cpp
+++ b/Engine/Plugins/Experimental/ChaosVD/Source/ChaosVD/Private/ChaosVDScene.cpp
@@ -47,7 +47,7 @@ void FChaosVDScene::AddReferencedObjects(FReferenceCollector& Collector)
 	Collector.AddStableReferenceMap(ParticleVDInstancesByID);
 }
 
-void FChaosVDScene::UpdateFromRecordedStepData(EChaosVDSolverType SolverID, const FChaosVDStepData& InRecordedStepData)
+void FChaosVDScene::UpdateFromRecordedStepData(const int32 SolverID, const FString& SolverName, const FChaosVDStepData& InRecordedStepData)
 {
 	TSet<uint32> ParticlesIDsInRecordedStepData;
 	ParticlesIDsInRecordedStepData.Reserve(InRecordedStepData.RecordedParticles.Num());
@@ -74,7 +74,9 @@ void FChaosVDScene::UpdateFromRecordedStepData(EChaosVDSolverType SolverID, cons
 		{
 			if (AChaosVDParticleActor* NewParticleVDInstance = SpawnParticleFromRecordedData(Particle))
 			{
-				FName FolderPath = *FPaths::Combine(TEXT("Solver ID : ") + UEnum::GetDisplayValueAsText(SolverID).ToString(), UEnum::GetDisplayValueAsText(Particle.ParticleType).ToString());
+				FStringFormatOrderedArguments Args {SolverName, FString::FromInt(SolverID)};
+				const FName FolderPath = *FPaths::Combine(FString::Format(TEXT("Solver {0} | ID {1}"), Args), UEnum::GetDisplayValueAsText(Particle.ParticleType).ToString());
+
 				NewParticleVDInstance->SetFolderPath(FolderPath);
 
 				// TODO: Precalculate the max num of entries we would see in the loaded file, and use that number to pre-allocate this map
@@ -109,9 +111,12 @@ void FChaosVDScene::UpdateFromRecordedStepData(EChaosVDSolverType SolverID, cons
 
 void FChaosVDScene::CleanUpScene()
 {
-	for (const TPair<uint32, AChaosVDParticleActor*>& ParticleVDInstanceWithID : ParticleVDInstancesByID)
+	if (PhysicsVDWorld)
 	{
-		PhysicsVDWorld->DestroyActor(ParticleVDInstanceWithID.Value);
+		for (const TPair<uint32, AChaosVDParticleActor*>& ParticleVDInstanceWithID : ParticleVDInstancesByID)
+		{
+			PhysicsVDWorld->DestroyActor(ParticleVDInstanceWithID.Value);
+		}
 	}
 
 	ParticleVDInstancesByID.Reset();
@@ -119,7 +124,7 @@ void FChaosVDScene::CleanUpScene()
 	CollectGarbage(GARBAGE_COLLECTION_KEEPFLAGS);
 }
 
-AChaosVDParticleActor* FChaosVDScene::SpawnParticleFromRecordedData(const FChaosVDParticleDebugData& InParticleData)
+AChaosVDParticleActor* FChaosVDScene::SpawnParticleFromRecordedData(const FChaosVDParticleDebugData& InParticleData) const
 {
 	// Temp code until we have geometry data
 	const UChaosVDEditorSettings* Settings = GetDefault<UChaosVDEditorSettings>();
@@ -131,7 +136,12 @@ AChaosVDParticleActor* FChaosVDScene::SpawnParticleFromRecordedData(const FChaos
 	if (AChaosVDParticleActor* NewActor = PhysicsVDWorld->SpawnActor<AChaosVDParticleActor>(Params))
 	{
 		NewActor->UpdateFromRecordedData(InParticleData);
-		NewActor->SetActorLabel(InParticleData.DebugName);
+
+		if (!InParticleData.DebugName.IsEmpty())
+		{
+			NewActor->SetActorLabel(InParticleData.DebugName);
+		}
+
 		NewActor->SetStaticMesh(DebugMesh);
 
 		return NewActor;
diff --git a/Engine/Plugins/Experimental/ChaosVD/Source/ChaosVD/Private/Trace/ChaosVDTraceAnalyzer.cpp b/Engine/Plugins/Experimental/ChaosVD/Source/ChaosVD/Private/Trace/ChaosVDTraceAnalyzer.cpp
new file mode 100644
--- /dev/null
+++ b/Engine/Plugins/Experimental/ChaosVD/Source/ChaosVD/Private/Trace/ChaosVDTraceAnalyzer.cpp
@@ -0,0 +1,127 @@
+// Copyright Epic Games, Inc. All Rights Reserved.
+
+#include "Trace/ChaosVDTraceAnalyzer.h"
+
+#include "Trace/ChaosVDTraceProvider.h"
+#include "TraceServices/Model/AnalysisSession.h"
+
+void FChaosVDTraceAnalyzer::OnAnalysisBegin(const FOnAnalysisContext& Context)
+{
+	FInterfaceBuilder& Builder = Context.InterfaceBuilder;
+	Builder.RouteEvent(RouteId_ChaosVDParticle, "ChaosVDLogger", "ChaosVDParticle");
+	
+	Builder.RouteEvent(RouteId_ChaosVDSolverFrameStart, "ChaosVDLogger", "ChaosVDSolverFrameStart");
+	Builder.RouteEvent(RouteId_ChaosVDSolverFrameEnd, "ChaosVDLogger", "ChaosVDSolverFrameEnd");
+	
+	Builder.RouteEvent(RouteId_ChaosVDSolverStepStart, "ChaosVDLogger", "ChaosVDSolverStepStart");
+	Builder.RouteEvent(RouteId_ChaosVDSolverStepEnd, "ChaosVDLogger", "ChaosVDSolverStepEnd");
+
+	TraceServices::FAnalysisSessionEditScope _(Session);
+	ChaosVDTraceProvider->CreateRecordingInstanceForSession(Session.GetName());
+}
+
+void FChaosVDTraceAnalyzer::OnAnalysisEnd()
+{
+	OnAnalysisComplete().Broadcast();
+}
+
+bool FChaosVDTraceAnalyzer::OnEvent(uint16 RouteId, EStyle Style, const FOnEventContext& Context)
+{
+	LLM_SCOPE_BYNAME(TEXT("Insights/FChaosVDTraceAnalyzer"));
+
+	TraceServices::FAnalysisSessionEditScope _(Session);
+
+	const FEventData& EventData = Context.EventData;
+	
+	switch (RouteId)
+	{
+	case RouteId_ChaosVDSolverFrameStart:
+		{
+			FChaosVDSolverFrameData NewFrameData;
+
+			NewFrameData.SolverID = EventData.GetValue<int32>("SolverID");
+
+			FWideStringView DebugNameView;
+			EventData.GetString("DebugName", DebugNameView);
+			NewFrameData.DebugName = DebugNameView;
+
+			// Add an empty frame. It will be filled out by the solver trace events
+			ChaosVDTraceProvider->AddFrame(NewFrameData.SolverID, MoveTemp(NewFrameData));
+	
+			break;
+		}
+	case RouteId_ChaosVDSolverFrameEnd:
+		{
+			break;
+		}
+	case RouteId_ChaosVDSolverStepStart:
+		{
+			const int32 SolverID = EventData.GetValue<int32>("SolverID");
+
+			// This can be null if the recording started Mid-Frame. In this case we just discard the data for now
+			if (FChaosVDSolverFrameData* FrameData  = ChaosVDTraceProvider->GetLastFrame(SolverID))
+			{
+				// Add an empty step. It will be filled out by the particle (and later on other objects/elements) events
+				FrameData->SolverSteps.AddDefaulted();
+			}
+	
+			break;
+		}
+	case RouteId_ChaosVDSolverStepEnd:
+		{
+			break;
+		}
+	case RouteId_ChaosVDParticle:
+		{
+			const int32 SolverID = EventData.GetValue<int32>("SolverID");
+
+			FChaosVDParticleDebugData ParticleData;
+			ReadParticleDataFromEvent(EventData, ParticleData);
+
+			// This can be null if the recording started Mid-Frame. In this case we just discard the data for now
+			if (FChaosVDSolverFrameData* FrameData = ChaosVDTraceProvider->GetLastFrame(SolverID))
+			{
+				if (FrameData->SolverSteps.Num() > 0)
+				{
+					FrameData->SolverSteps.Last().RecordedParticles.Add(MoveTemp(ParticleData));
+				}
+			}
+
+			break;
+		}
+	default:
+		break;
+	}
+
+	return true;
+}
+
+void FChaosVDTraceAnalyzer::ReadParticleDataFromEvent(const FEventData& InEventData, FChaosVDParticleDebugData& OutParticleData)
+{
+	OutParticleData.ParticleType = static_cast<EChaosVDParticleType>(InEventData.GetValue<uint8>("ParticleType"));
+	OutParticleData.ParticleState = static_cast<EChaosVDParticleState>(InEventData.GetValue<int8>("ObjectState"));
+	OutParticleData.ParticleIndex = InEventData.GetValue<int32>("ParticleID");
+
+	FWideStringView DebugNameView;
+	InEventData.GetString("DebugName", DebugNameView);
+
+	OutParticleData.DebugName = DebugNameView;
+
+	OutParticleData.Position.X = InEventData.GetValue<float>("PositionX");
+	OutParticleData.Position.Y = InEventData.GetValue<float>("PositionY");
+	OutParticleData.Position.Z = InEventData.GetValue<float>("PositionZ");
+
+	OutParticleData.Rotation.X = InEventData.GetValue<float>("RotationX");
+	OutParticleData.Rotation.Y = InEventData.GetValue<float>("RotationY");
+	OutParticleData.Rotation.Z = InEventData.GetValue<float>("RotationZ");
+	OutParticleData.Rotation.W = InEventData.GetValue<float>("RotationW");
+			
+	OutParticleData.Velocity.X = InEventData.GetValue<float>("VelocityX");
+	OutParticleData.Velocity.Y = InEventData.GetValue<float>("VelocityY");
+	OutParticleData.Velocity.Z = InEventData.GetValue<float>("VelocityZ");
+			
+	OutParticleData.AngularVelocity.X = InEventData.GetValue<float>("AngularVelocityX");
+	OutParticleData.AngularVelocity.Y = InEventData.GetValue<float>("AngularVelocityY");
+	OutParticleData.AngularVelocity.Z = InEventData.GetValue<float>("AngularVelocityZ");
+}
+
diff --git a/Engine/Plugins/Experimental/ChaosVD/Source/ChaosVD/Private/Trace/ChaosVDTraceManager.cpp b/Engine/Plugins/Experimental/ChaosVD/Source/ChaosVD/Private/Trace/ChaosVDTraceManager.cpp
new file mode 100644
--- /dev/null
+++ b/Engine/Plugins/Experimental/ChaosVD/Source/ChaosVD/Private/Trace/ChaosVDTraceManager.cpp
@@ -0,0 +1,62 @@
+// Copyright Epic Games, Inc. All Rights Reserved.
+
+#include "Trace/ChaosVDTraceManager.h"
+
+#include "Modules/ModuleManager.h"
+#include "Features/IModularFeatures.h"
+#include "Trace/ChaosVDTraceModule.h"
+#include "TraceServices/AnalysisService.h"
+#include "TraceServices/ITraceServicesModule.h"
+#include "TraceServices/Model/AnalysisSession.h"
+
+FChaosVDTraceManager::FChaosVDTraceManager() 
+{
+	ChaosVDTraceModule = MakeShared<FChaosVDTraceModule>();
+	IModularFeatures::Get().RegisterModularFeature(TraceServices::ModuleFeatureName, ChaosVDTraceModule.Get());
+
+	FString ChannelNameFString(TEXT("ChaosVD"));
+	UE::Trace::ToggleChannel(ChannelNameFString.GetCharArray().GetData(), true);
+}
+
+FChaosVDTraceManager::~FChaosVDTraceManager()
+{
+	IModularFeatures::Get().UnregisterModularFeature(TraceServices::ModuleFeatureName, ChaosVDTraceModule.Get());
+}
+
+FString FChaosVDTraceManager::LoadTraceFile(const FString& InTraceFilename)
+{
+	ResetSession(InTraceFilename);
+
+	ITraceServicesModule& TraceServicesModule = FModuleManager::LoadModuleChecked<ITraceServicesModule>("TraceServices");
+	if (const TSharedPtr<TraceServices::IAnalysisService> TraceAnalysisService = TraceServicesModule.GetAnalysisService())
+	{
+		const TSharedPtr<const TraceServices::IAnalysisSession> NewSession = TraceAnalysisService->StartAnalysis(*InTraceFilename);
+		AnalysisSessionByName.Add(InTraceFilename, NewSession);
+		return NewSession->GetName();
+	}
+
+	return FString();
+}
+
+TSharedPtr<const TraceServices::IAnalysisSession> FChaosVDTraceManager::GetSession(const FString& InSessionName)
+{
+	if (TSharedPtr<const TraceServices::IAnalysisSession>* FoundSession = AnalysisSessionByName.Find(InSessionName))
+	{
+		return *FoundSession;
+	}
+
+	return nullptr;
+}
+
+void FChaosVDTraceManager::ResetSession(const FString& InSessionName)
+{
+	if (const TSharedPtr<const TraceServices::IAnalysisSession>* Session = AnalysisSessionByName.Find(InSessionName))
+	{
+		if (Session->IsValid())
+		{
+			(*Session)->Stop(true);
+		}
+		
+		AnalysisSessionByName.Remove(InSessionName);
+	}
+}
diff --git a/Engine/Plugins/Experimental/ChaosVD/Source/ChaosVD/Private/Trace/ChaosVDTraceModule.cpp b/Engine/Plugins/Experimental/ChaosVD/Source/ChaosVD/Private/Trace/ChaosVDTraceModule.cpp
new file mode 100644
--- /dev/null
+++ b/Engine/Plugins/Experimental/ChaosVD/Source/ChaosVD/Private/Trace/ChaosVDTraceModule.cpp
@@ -0,0 +1,33 @@
+// Copyright Epic Games, Inc. All Rights Reserved.
+
+#include "Trace/ChaosVDTraceModule.h"
+
+#include "Trace/ChaosVDTraceAnalyzer.h"
+#include "Trace/ChaosVDTraceProvider.h"
+#include "TraceServices/Model/AnalysisSession.h"
+
+FName FChaosVDTraceModule::ModuleName("ChaosVDTrace");
+
+void FChaosVDTraceModule::GetModuleInfo(TraceServices::FModuleInfo& OutModuleInfo)
+{
+	OutModuleInfo.Name = ModuleName;
+	OutModuleInfo.DisplayName = TEXT("ChaosVisualDebugger");
+}
+
+void FChaosVDTraceModule::OnAnalysisBegin(TraceServices::IAnalysisSession& InSession)
+{
+	const TSharedPtr<FChaosVDTraceProvider> Provider = MakeShared<FChaosVDTraceProvider>(InSession);
+
+	InSession.AddProvider(FChaosVDTraceProvider::ProviderName, Provider);
+	InSession.AddAnalyzer(new FChaosVDTraceAnalyzer(InSession, Provider));
+}
+
+void FChaosVDTraceModule::GetLoggers(TArray<const TCHAR*>& OutLoggers)
+{
+	OutLoggers.Add(TEXT("ChaosVD"));
+}
+
+void FChaosVDTraceModule::GenerateReports(const TraceServices::IAnalysisSession& Session, const TCHAR* CmdLine, const TCHAR* OutputDirectory)
+{
+	IModule::GenerateReports(Session, CmdLine, OutputDirectory);
+}
diff --git a/Engine/Plugins/Experimental/ChaosVD/Source/ChaosVD/Private/Trace/ChaosVDTraceProvider.cpp b/Engine/Plugins/Experimental/ChaosVD/Source/ChaosVD/Private/Trace/ChaosVDTraceProvider.cpp
new file mode 100644
--- /dev/null
+++ b/Engine/Plugins/Experimental/ChaosVD/Source/ChaosVD/Private/Trace/ChaosVDTraceProvider.cpp
@@ -0,0 +1,57 @@
+// Copyright Epic Games, Inc. All Rights Reserved.
+
+#include "Trace/ChaosVDTraceProvider.h"
+
+#include "ChaosVDRecording.h"
+
+FName FChaosVDTraceProvider::ProviderName("ChaosVDProvider");
+
+FChaosVDTraceProvider::FChaosVDTraceProvider(TraceServices::IAnalysisSession& InSession): Session(InSession)
+{
+}
+
+void FChaosVDTraceProvider::CreateRecordingInstanceForSession(const FString& InSessionName)
+{
+	DeleteRecordingInstanceForSession();
+
+	InternalRecording = MakeShared<FChaosVDRecording>();
+	InternalRecording->SessionName = InSessionName;
+}
+
+void FChaosVDTraceProvider::DeleteRecordingInstanceForSession()
+{
+	InternalRecording.Reset();
+}
+
+void FChaosVDTraceProvider::AddFrame(const int32 InSolverID, FChaosVDSolverFrameData&& FrameData)
+{
+	if (InternalRecording.IsValid())
+	{
+		InternalRecording->AddFrameForSolver(InSolverID, MoveTemp(FrameData));
+	}
+}
+
+FChaosVDSolverFrameData* FChaosVDTraceProvider::GetFrame(const int32 InSolverID, const int32 FrameNumber) const
+{
+	return InternalRecording.IsValid() ? InternalRecording->GetFrameForSolver(InSolverID, FrameNumber) : nullptr;
+}
+
+FChaosVDSolverFrameData* FChaosVDTraceProvider::GetLastFrame(const int32 InSolverID) const
+{
+	if (InternalRecording.IsValid() && InternalRecording->GetAvailableFramesNumber(InSolverID) > 0)
+	{
+		const int32 AvailableFramesNumber = InternalRecording->GetAvailableFramesNumber(InSolverID);
+
+		if (AvailableFramesNumber != INDEX_NONE)
+		{
+			return GetFrame(InSolverID, InternalRecording->GetAvailableFramesNumber(InSolverID) - 1);
+		}
+	}
+
+	return nullptr;
+}
+
+TSharedPtr<FChaosVDRecording> FChaosVDTraceProvider::GetRecordingForSession() const
+{
+	return InternalRecording;
+}
diff --git a/Engine/Plugins/Experimental/ChaosVD/Source/ChaosVD/Private/Widgets/SChaosVDMainTab.cpp b/Engine/Plugins/Experimental/ChaosVD/Source/ChaosVD/Private/Widgets/SChaosVDMainTab.cpp
--- a/Engine/Plugins/Experimental/ChaosVD/Source/ChaosVD/Private/Widgets/SChaosVDMainTab.cpp
+++ b/Engine/Plugins/Experimental/ChaosVD/Source/ChaosVD/Private/Widgets/SChaosVDMainTab.cpp
@@ -190,14 +190,6 @@ TSharedRef<FTabManager::FLayout> SChaosVDMainTab::GenerateMainLayout()
 		);
 }
 
-SChaosVDMainTab::~SChaosVDMainTab()
-{
-	if (ChaosVDEngine.IsValid())
-	{
-		ChaosVDEngine->DeInitialize();
-	}
-}
-
 void SChaosVDMainTab::GenerateMainWindowMenu()
 {
 	FMenuBarBuilder MenuBarBuilder = FMenuBarBuilder(TSharedPtr<FUICommandList >());
@@ -234,12 +226,18 @@ void SChaosVDMainTab::BrowseAndOpenChaosVDFile()
 	TArray<FString> OutOpenFilenames;
 	if (IDesktopPlatform* DesktopPlatform = FDesktopPlatformModule::Get())
 	{
+		
+		FString ExtensionStr;
+		ExtensionStr += TEXT("Unreal Trace|*.utrace|");
+		//TODO: Re-enable this when we add "Clips" support as these will use our own format
+		//ExtensionStr += TEXT("Chaos Visual Debugger|*.cvd");
+	
 		DesktopPlatform->OpenFileDialog(
 			FSlateApplication::Get().FindBestParentWindowHandleForDialogs(AsShared()),
 			LOCTEXT("OpenDialogTitle", "Open Chaos Visual Debug File").ToString(),
 			TEXT(""),
 			TEXT(""),
-			TEXT("Chaos Visual Debugger Files|*.cvd"),
+			*ExtensionStr,
 			EFileDialogFlags::None,
 			OutOpenFilenames
 		);
@@ -247,10 +245,10 @@ void SChaosVDMainTab::BrowseAndOpenChaosVDFile()
 
 	if (OutOpenFilenames.Num() > 0)
 	{
-		if (ensure(ChaosVDEngine->GetPlaybackController()))
+		if (OutOpenFilenames[0].EndsWith(TEXT("utrace")))
 		{
-			ChaosVDEngine->GetPlaybackController()->LoadChaosVDRecording(OutOpenFilenames[0]);
-		}	
+			GetChaosVDEngineInstance()->LoadRecording(OutOpenFilenames[0]);
+		}
 	}
 }
 
diff --git a/Engine/Plugins/Experimental/ChaosVD/Source/ChaosVD/Private/Widgets/SChaosVDPlaybackViewport.cpp b/Engine/Plugins/Experimental/ChaosVD/Source/ChaosVD/Private/Widgets/SChaosVDPlaybackViewport.cpp
--- a/Engine/Plugins/Experimental/ChaosVD/Source/ChaosVD/Private/Widgets/SChaosVDPlaybackViewport.cpp
+++ b/Engine/Plugins/Experimental/ChaosVD/Source/ChaosVD/Private/Widgets/SChaosVDPlaybackViewport.cpp
@@ -10,11 +10,15 @@
 #include "Widgets/SViewport.h"
 #include "Widgets/Text/STextBlock.h"
 
-
 #define LOCTEXT_NAMESPACE "ChaosVisualDebugger"
 
 SChaosVDPlaybackViewport::~SChaosVDPlaybackViewport()
 {
+	if (const TSharedPtr<FChaosVDPlaybackController> CurrentPlaybackControllerPtr = PlaybackController.Pin())
+	{
+		CurrentPlaybackControllerPtr->OnControllerUpdated().Unbind();
+	}
+
 	LevelViewportClient->Viewport = nullptr;
 	LevelViewportClient.Reset();
 }
@@ -41,8 +45,6 @@ void SChaosVDPlaybackViewport::Construct(const FArguments& InArgs, const UWorld*
 	ensure(DefaultWorld);
 	ensure(InPlaybackController.IsValid());
 
-	PlaybackController = InPlaybackController;
-
 	LevelViewportClient = CreateViewportClient();
 
 	ViewportWidget = SNew(SViewport)
@@ -120,23 +122,23 @@ void SChaosVDPlaybackViewport::Construct(const FArguments& InArgs, const UWorld*
 		]
 	];
 
-	if (const TSharedPtr<FChaosVDPlaybackController> PlaybackControllerPtr = PlaybackController.Pin())
-	{
-		PlaybackControllerPtr->OnControllerUpdated().BindRaw(this, &SChaosVDPlaybackViewport::OnPlaybackControllerUpdated);
-	}
+	RegisterNewController(InPlaybackController);
 }
 
-void SChaosVDPlaybackViewport::OnPlaybackControllerUpdated(FChaosVDPlaybackController* Controller) const
+void SChaosVDPlaybackViewport::OnPlaybackControllerUpdated(TWeakPtr<FChaosVDPlaybackController> InController)
 {
-	if (!ensure(Controller))
+	if (PlaybackController != InController)
 	{
-		return;
+		RegisterNewController(InController);
 	}
 
-	if (TSharedPtr<FChaosVDRecording> LoadedRecording = Controller->GetCurrentRecording().Pin())
+	const TSharedPtr<FChaosVDPlaybackController> ControllerSharedPtr = PlaybackController.Pin();
+	if (ControllerSharedPtr.IsValid() && ControllerSharedPtr->IsRecordingLoaded())
 	{
-		const int32 AvailableFrames = Controller->GetAvailableFramesNumber();
-		const int32 AvailableSteps = Controller->GetStepsForFrame(Controller->GetCurrentFrame());
+		const int32 TrackID = ControllerSharedPtr->GetActiveSolverTrackID();
+		
+		const int32 AvailableFrames = ControllerSharedPtr->GetAvailableFramesNumber(TrackID);
+		const int32 AvailableSteps = ControllerSharedPtr->GetStepsForFrame(TrackID, ControllerSharedPtr->GetCurrentFrame(TrackID));
 
 		// Max is inclusive and we use this to request as the index on the recorded frames/steps arrays so we need to -1 to the available frames/steps
 		FramesTimelineWidget->UpdateMinMaxValue(0, AvailableFrames != INDEX_NONE ? AvailableFrames -1  : 0);
@@ -156,12 +158,36 @@ void SChaosVDPlaybackViewport::OnPlaybackControllerUpdated(FChaosVDPlaybackContr
 	LevelViewportClient->bNeedsRedraw = true;
 }
 
+void SChaosVDPlaybackViewport::RegisterNewController(TWeakPtr<FChaosVDPlaybackController> NewController)
+{
+	if (PlaybackController != NewController)
+	{
+		if (const TSharedPtr<FChaosVDPlaybackController> CurrentPlaybackControllerPtr = PlaybackController.Pin())
+		{
+			CurrentPlaybackControllerPtr->OnControllerUpdated().Unbind();
+		}
+
+		PlaybackController = NewController;
+
+		if (const TSharedPtr<FChaosVDPlaybackController> NewPlaybackControllerPtr = PlaybackController.Pin())
+		{
+			NewPlaybackControllerPtr->OnControllerUpdated().BindRaw(this, &SChaosVDPlaybackViewport::OnPlaybackControllerUpdated);
+		}
+	}
+}
+
 void SChaosVDPlaybackViewport::OnFrameSelectionUpdated(int32 NewFrameIndex) const
 {
 	if (const TSharedPtr<FChaosVDPlaybackController> PlaybackControllerPtr = PlaybackController.Pin())
 	{
+		const int32 TrackID = PlaybackControllerPtr->GetActiveSolverTrackID();
+		
+		const int32 AvailableSteps = PlaybackControllerPtr->GetStepsForFrame(TrackID, PlaybackControllerPtr->GetCurrentFrame(TrackID));
+		StepsTimelineWidget->UpdateMinMaxValue(0,AvailableSteps != INDEX_NONE ? AvailableSteps -1 : 0);
+
 		// On Frame updates, always use Step 0
-		PlaybackControllerPtr->GoToRecordedStep(NewFrameIndex, 0);
+		StepsTimelineWidget->SetCurrentTimelineFrame(0);
+		PlaybackControllerPtr->GoToRecordedStep(TrackID, NewFrameIndex, 0);
 
 		LevelViewportClient->bNeedsRedraw = true;
 	}
@@ -171,8 +197,9 @@ void SChaosVDPlaybackViewport::OnStepSelectionUpdated(int32 NewStepIndex) const
 {
 	if (const TSharedPtr<FChaosVDPlaybackController> PlaybackControllerPtr = PlaybackController.Pin())
 	{
+		const int32 TrackID = PlaybackControllerPtr->GetActiveSolverTrackID();
 		// On Steps updates. Always use the current Frame
-		PlaybackControllerPtr->GoToRecordedStep(PlaybackControllerPtr->GetCurrentFrame(), NewStepIndex);
+		PlaybackControllerPtr->GoToRecordedStep(TrackID, PlaybackControllerPtr->GetCurrentFrame(TrackID), NewStepIndex);
 	
 		LevelViewportClient->bNeedsRedraw = true;
 	}
diff --git a/Engine/Plugins/Experimental/ChaosVD/Source/ChaosVD/Private/Widgets/SChaosVDTimelineWidget.cpp b/Engine/Plugins/Experimental/ChaosVD/Source/ChaosVD/Private/Widgets/SChaosVDTimelineWidget.cpp
--- a/Engine/Plugins/Experimental/ChaosVD/Source/ChaosVD/Private/Widgets/SChaosVDTimelineWidget.cpp
+++ b/Engine/Plugins/Experimental/ChaosVD/Source/ChaosVD/Private/Widgets/SChaosVDTimelineWidget.cpp
@@ -122,7 +122,11 @@ void SChaosVDTimelineWidget::UpdateMinMaxValue(float NewMin, float NewMax)
 
 	MinFrames = NewMin;
 	MaxFrames = NewMax;
-	CurrentFrame = NewMin;
+	
+	if(CurrentFrame < NewMin || CurrentFrame > NewMax)
+	{
+		CurrentFrame = NewMin;
+	}
 }
 
 void SChaosVDTimelineWidget::ResetTimeline()
@@ -134,6 +138,7 @@ void SChaosVDTimelineWidget::ResetTimeline()
 void SChaosVDTimelineWidget::SetCurrentTimelineFrame(float FrameNumber)
 {
 	CurrentFrame = static_cast<int32>(FrameNumber);
+	TimelineSlider->SetValue(CurrentFrame);
 	FrameChangedDelegate.ExecuteIfBound(CurrentFrame);
 }
 
@@ -158,7 +163,6 @@ FReply SChaosVDTimelineWidget::Stop()
 	CurrentFrame = 0;
 	CurrentPlaybackTime = 0.0f;
 	bIsPlaying = false;
-	TimelineSlider->SetValue(CurrentFrame);
 
 	SetCurrentTimelineFrame(CurrentFrame);
 
@@ -176,7 +180,6 @@ FReply SChaosVDTimelineWidget::Next()
 	}
 
 	CurrentFrame++;
-	TimelineSlider->SetValue(CurrentFrame);
 	
 	SetCurrentTimelineFrame(CurrentFrame);
 
@@ -191,7 +194,6 @@ FReply SChaosVDTimelineWidget::Prev()
 	}
 
 	CurrentFrame--;
-	TimelineSlider->SetValue(CurrentFrame);
 
 	SetCurrentTimelineFrame(CurrentFrame);
 
diff --git a/Engine/Plugins/Experimental/ChaosVD/Source/ChaosVD/Public/ChaosVDEngine.h b/Engine/Plugins/Experimental/ChaosVD/Source/ChaosVD/Public/ChaosVDEngine.h
--- a/Engine/Plugins/Experimental/ChaosVD/Source/ChaosVD/Public/ChaosVDEngine.h
+++ b/Engine/Plugins/Experimental/ChaosVD/Source/ChaosVD/Public/ChaosVDEngine.h
@@ -4,7 +4,9 @@
 
 #include "Templates/SharedPointer.h"
 #include "Containers/Ticker.h"
+#include "Misc/Guid.h"
 
+class FChaosVDTraceManager;
 class FChaosVDPlaybackController;
 class FChaosVDScene;
 class FChaosVisualDebuggerMainUI;
@@ -13,18 +15,29 @@ class FChaosVisualDebuggerMainUI;
 class FChaosVDEngine : public FTSTickerObjectBase, public TSharedFromThis<FChaosVDEngine>
 {
 public:
+	FChaosVDEngine()
+	{
+		InstanceGUID = FGuid::NewGuid();
+	}
 
 	void Initialize();
 	
 	void DeInitialize();
-	
+
 	virtual bool Tick(float DeltaTime) override;
-	
+
+	const FGuid& GetInstanceGuid() const { return InstanceGUID; };
 	TSharedPtr<FChaosVDScene>& GetCurrentScene() { return CurrentScene; };
 	TSharedPtr<FChaosVDPlaybackController>& GetPlaybackController() { return PlaybackController; };
+
+	void LoadRecording(const FString& FilePath);
 	
 private:
 
+	FGuid InstanceGUID;
+
+	FString CurrentSessionName;
+
 	TSharedPtr<FChaosVDScene> CurrentScene;
 	TSharedPtr<FChaosVDPlaybackController> PlaybackController;
 	
diff --git a/Engine/Plugins/Experimental/ChaosVD/Source/ChaosVD/Public/ChaosVDModule.h b/Engine/Plugins/Experimental/ChaosVD/Source/ChaosVD/Public/ChaosVDModule.h
--- a/Engine/Plugins/Experimental/ChaosVD/Source/ChaosVD/Public/ChaosVDModule.h
+++ b/Engine/Plugins/Experimental/ChaosVD/Source/ChaosVD/Public/ChaosVDModule.h
@@ -4,21 +4,36 @@
 
 #include "Templates/SharedPointer.h"
 #include "Modules/ModuleManager.h"
+#include "Misc/Guid.h"
 
+class FChaosVDTraceManager;
+class FChaosVDEngine;
 class SDockTab;
 class FSpawnTabArgs;
+struct FGuid;
 
 class FChaosVDModule : public IModuleInterface
 {
 public:
 
+	static FChaosVDModule& Get();
+
 	/** IModuleInterface implementation */
 	virtual void StartupModule() override;
 	virtual void ShutdownModule() override;
 
+	TSharedPtr<FChaosVDTraceManager>& GetTraceManager() { return ChaosVDTraceManager; };
+
 private:
 
 	void RegisterClassesCustomDetails() const;
 
-	TSharedRef<SDockTab> SpawnMainTab(const FSpawnTabArgs& Args) const;
+	TSharedRef<SDockTab> SpawnMainTab(const FSpawnTabArgs& Args);
+
+	void RegisterChaosVDInstance(const FGuid& InstanceGuid, TSharedPtr<FChaosVDEngine> Instance);
+	void DeregisterChaosVDInstance(const FGuid& InstanceGuid);
+	
+	TMap<FGuid, TSharedPtr<FChaosVDEngine>> ActiveChaosVDInstances;
+
+	TSharedPtr<FChaosVDTraceManager> ChaosVDTraceManager;
 };
diff --git a/Engine/Plugins/Experimental/ChaosVD/Source/ChaosVD/Public/ChaosVDPlaybackController.h b/Engine/Plugins/Experimental/ChaosVD/Source/ChaosVD/Public/ChaosVDPlaybackController.h
--- a/Engine/Plugins/Experimental/ChaosVD/Source/ChaosVD/Public/ChaosVDPlaybackController.h
+++ b/Engine/Plugins/Experimental/ChaosVD/Source/ChaosVD/Public/ChaosVDPlaybackController.h
@@ -11,28 +11,34 @@ struct FChaosVDRecording;
 class FChaosVDPlaybackController;
 class FString;
 
-DECLARE_DELEGATE_OneParam(FChaosVDPlaybackControllerUpdated, FChaosVDPlaybackController*)
+DECLARE_DELEGATE_OneParam(FChaosVDPlaybackControllerUpdated, TWeakPtr<FChaosVDPlaybackController>)
 
 /** Loads,unloads and owns a Chaos VD recording file */
-class FChaosVDPlaybackController
+class FChaosVDPlaybackController : public TSharedFromThis<FChaosVDPlaybackController>
 {
 public:
 
-	FChaosVDPlaybackController(TWeakPtr<FChaosVDScene> InSceneToControl);
-
-	bool LoadChaosVDRecording(const FString& RecordingPath);
+	FChaosVDPlaybackController(const TWeakPtr<FChaosVDScene>& InSceneToControl);
+	~FChaosVDPlaybackController();
+	
+	bool LoadChaosVDRecordingFromTraceSession(const FString& InSessionName);
 	
-	void UnloadCurrentRecording();
+	void UnloadCurrentRecording(const bool bBroadcastUpdate = true);
 
-	void GoToRecordedStep(int32 FrameNumber, int32 Step);
+	void GoToRecordedStep(const int32 InTrackID, const int32 FrameNumber, const int32 Step);
 
 	TWeakPtr<FChaosVDRecording> GetCurrentRecording() { return LoadedRecording; }
 
-	int32 GetStepsForFrame(int32 FrameNumber) const;
-	int32 GetAvailableFramesNumber() const;
+	int32 GetStepsForFrame(const int32 InTrackID, const int32 FrameNumber) const;
+	int32 GetAvailableFramesNumber(const int32 InTrackID) const;
+	int32 GetAvailableSolversNumber() const;
+
+	int32 GetActiveSolverTrackID() const;
 
-	int32 GetCurrentFrame() const { return CurrentFrame; }
-	int32 GetCurrentStep() const { return CurrentStep; }
+	int32 GetCurrentFrame(const int32 InTrackID) const;
+	int32 GetCurrentStep(const int32 InTrackID) const;
+
+	bool IsRecordingLoaded() const { return LoadedRecording.IsValid(); }
 
 	TWeakPtr<FChaosVDScene> GetControllerScene() { return SceneToControl; }
 
@@ -40,8 +46,10 @@ public:
 
 protected:
 
-	int32 CurrentFrame;
-	int32 CurrentStep;
+	void HandleCurrentRecordingUpdated();
+
+	TMap<int32,int32> CurrentFramePerTrack;
+	TMap<int32,int32> CurrentStepPerTrack;
 	
 	TSharedPtr<FChaosVDRecording> LoadedRecording;
 
diff --git a/Engine/Plugins/Experimental/ChaosVD/Source/ChaosVD/Public/ChaosVDScene.h b/Engine/Plugins/Experimental/ChaosVD/Source/ChaosVD/Public/ChaosVDScene.h
--- a/Engine/Plugins/Experimental/ChaosVD/Source/ChaosVD/Public/ChaosVDScene.h
+++ b/Engine/Plugins/Experimental/ChaosVD/Source/ChaosVD/Public/ChaosVDScene.h
@@ -32,7 +32,7 @@ public:
 	FChaosVDSceneUpdatedDelegate& OnSceneUpdated() { return SceneUpdatedDelegate; }
 
 	/** Updates, Adds and Remove actors to match the provided Step Data */
-	void UpdateFromRecordedStepData(EChaosVDSolverType SolverID, const FChaosVDStepData& InRecordedStepData);
+	void UpdateFromRecordedStepData(const int32 SolverID, const FString& SolverName, const FChaosVDStepData& InRecordedStepData);
 
 	/** Deletes all actors of the Scene and underlying UWorld */
 	void CleanUpScene();
@@ -40,10 +40,12 @@ public:
 	/** Returns the a ptr to the UWorld used to represent the current recorded frame data */
 	UWorld* GetUnderlyingWorld() const { return PhysicsVDWorld; };
 
+	bool IsInitialized() const { return  bIsInitialized; }
+
 private:
 
 	/** Creates an ChaosVDParticle actor for the Provided recorded Particle Data */
-	AChaosVDParticleActor* SpawnParticleFromRecordedData(const FChaosVDParticleDebugData& InParticleData);
+	AChaosVDParticleActor* SpawnParticleFromRecordedData(const FChaosVDParticleDebugData& InParticleData) const;
 
 	/** Returns the ID used to track this recorded particle data */
 	uint32 GetIDForRecordedParticleData(const FChaosVDParticleDebugData& InParticleData) const;
diff --git a/Engine/Plugins/Experimental/ChaosVD/Source/ChaosVD/Public/Trace/ChaosVDTraceAnalyzer.h b/Engine/Plugins/Experimental/ChaosVD/Source/ChaosVD/Public/Trace/ChaosVDTraceAnalyzer.h
new file mode 100644
--- /dev/null
+++ b/Engine/Plugins/Experimental/ChaosVD/Source/ChaosVD/Public/Trace/ChaosVDTraceAnalyzer.h
@@ -0,0 +1,53 @@
+// Copyright Epic Games, Inc. All Rights Reserved.
+
+#pragma once
+
+#include "ChaosVDRecording.h"
+#include "Trace/Analyzer.h"
+#include "Templates/SharedPointer.h"
+
+struct FChaosVDStepData;
+class FChaosVDTraceProvider;
+
+namespace TraceServices { class IAnalysisSession; }
+
+DECLARE_MULTICAST_DELEGATE(FChaosVDTraceAnalysisComplete)
+
+/** Analyzer class for Chaos VD trace recordings.
+ * It processes all Chaos VD trace events and rebuilds the recording in a usable
+ * Recording struct handled by a custom Trace Provider
+ */
+class FChaosVDTraceAnalyzer final : public UE::Trace::IAnalyzer
+{
+public:
+	FChaosVDTraceAnalyzer(TraceServices::IAnalysisSession& Session, const TSharedPtr<FChaosVDTraceProvider>& InChaosVDTraceProvider)
+	: Session(Session),
+	ChaosVDTraceProvider(InChaosVDTraceProvider)
+	{
+	}
+
+	virtual void OnAnalysisBegin(const FOnAnalysisContext& Context) override;
+	virtual void OnAnalysisEnd() override;
+	virtual bool OnEvent(uint16 RouteId, EStyle Style, const FOnEventContext& Context) override;
+
+	FChaosVDTraceAnalysisComplete& OnAnalysisComplete() { return ChaosVDTraceAnalysisCompleteDelegate; }
+
+private:
+
+	enum : uint16
+	{
+		RouteId_ChaosVDParticle,
+		RouteId_ChaosVDSolverStepStart,
+		RouteId_ChaosVDSolverStepEnd,
+		RouteId_ChaosVDSolverFrameStart,
+		RouteId_ChaosVDSolverFrameEnd
+	};
+
+	void ReadParticleDataFromEvent(const FEventData& InEventData, FChaosVDParticleDebugData& OutParticleData);
+
+	TraceServices::IAnalysisSession& Session;
+
+	TSharedPtr<FChaosVDTraceProvider> ChaosVDTraceProvider;
+	
+	FChaosVDTraceAnalysisComplete ChaosVDTraceAnalysisCompleteDelegate;
+};
diff --git a/Engine/Plugins/Experimental/ChaosVD/Source/ChaosVD/Public/Trace/ChaosVDTraceManager.h b/Engine/Plugins/Experimental/ChaosVD/Source/ChaosVD/Public/Trace/ChaosVDTraceManager.h
new file mode 100644
--- /dev/null
+++ b/Engine/Plugins/Experimental/ChaosVD/Source/ChaosVD/Public/Trace/ChaosVDTraceManager.h
@@ -0,0 +1,32 @@
+// Copyright Epic Games, Inc. All Rights Reserved.
+
+#pragma once
+
+#include "Containers/StringFwd.h"
+#include "Templates/SharedPointer.h"
+
+class FChaosVDEngine;
+class FChaosVDTraceModule;
+
+namespace TraceServices { class IAnalysisSession; }
+
+/** Manager class used by Chaos VD to interact/control UE Trace systems */
+class FChaosVDTraceManager
+{
+public:
+	FChaosVDTraceManager();
+	~FChaosVDTraceManager();
+
+	FString LoadTraceFile(const FString& InTraceFilename);
+
+	TSharedPtr<const TraceServices::IAnalysisSession> GetSession(const FString& InSessionName);
+
+private:
+	
+	void ResetSession(const FString& InSessionName);
+
+	/** The trace analysis session. */
+	TMap<FString, TSharedPtr<const TraceServices::IAnalysisSession>> AnalysisSessionByName;
+
+	TSharedPtr<FChaosVDTraceModule> ChaosVDTraceModule;
+};
diff --git a/Engine/Plugins/Experimental/ChaosVD/Source/ChaosVD/Public/Trace/ChaosVDTraceModule.h b/Engine/Plugins/Experimental/ChaosVD/Source/ChaosVD/Public/Trace/ChaosVDTraceModule.h
new file mode 100644
--- /dev/null
+++ b/Engine/Plugins/Experimental/ChaosVD/Source/ChaosVD/Public/Trace/ChaosVDTraceModule.h
@@ -0,0 +1,23 @@
+// Copyright Epic Games, Inc. All Rights Reserved.
+
+#pragma once
+#include "TraceServices/ModuleService.h"
+
+/** Modular feature used by ChaosVD to interface with the Trace Services system */
+class FChaosVDTraceModule : public TraceServices::IModule
+{
+public:
+	 FChaosVDTraceModule()
+		{
+		}
+
+	static FName ModuleName;
+
+	// TraceServices::IModule interface
+	virtual void GetModuleInfo(TraceServices::FModuleInfo& OutModuleInfo) override;
+	virtual void OnAnalysisBegin(TraceServices::IAnalysisSession& InSession) override;
+	virtual void GetLoggers(TArray<const TCHAR *>& OutLoggers) override;
+	virtual void GenerateReports(const TraceServices::IAnalysisSession& Session, const TCHAR* CmdLine, const TCHAR* OutputDirectory) override;
+	virtual const TCHAR* GetCommandLineArgument() override { return TEXT("chaosvd"); }
+
+};
diff --git a/Engine/Plugins/Experimental/ChaosVD/Source/ChaosVD/Public/Trace/ChaosVDTraceProvider.h b/Engine/Plugins/Experimental/ChaosVD/Source/ChaosVD/Public/Trace/ChaosVDTraceProvider.h
new file mode 100644
--- /dev/null
+++ b/Engine/Plugins/Experimental/ChaosVD/Source/ChaosVD/Public/Trace/ChaosVDTraceProvider.h
@@ -0,0 +1,35 @@
+// Copyright Epic Games, Inc. All Rights Reserved.
+
+#pragma once
+#include "TraceServices/Model/AnalysisSession.h"
+#include "Templates/SharedPointer.h"
+
+class FChaosVDEngine;
+struct FChaosVDSolverFrameData;
+struct FChaosVDRecording;
+
+/** Provider class for Chaos VD trace recordings.
+ * It stores and handles rebuilt recorded frame data from Trace events
+ * dispatched by the Chaos VD Trace analyzer
+ */
+class FChaosVDTraceProvider : public TraceServices::IProvider
+{
+public:
+	
+	static FName ProviderName;
+
+	FChaosVDTraceProvider(TraceServices::IAnalysisSession& InSession);
+
+	void CreateRecordingInstanceForSession(const FString& InSessionName);
+	void DeleteRecordingInstanceForSession();
+	void AddFrame(const int32 InSolverGUID, FChaosVDSolverFrameData&& FrameData);
+	FChaosVDSolverFrameData* GetFrame(const int32 InSolverGUID, const int32 FrameNumber) const;
+	FChaosVDSolverFrameData* GetLastFrame(const int32 InSolverGUID) const;
+
+	TSharedPtr<FChaosVDRecording> GetRecordingForSession() const;
+
+private:
+	TraceServices::IAnalysisSession& Session;
+
+	TSharedPtr<FChaosVDRecording> InternalRecording;
+};
diff --git a/Engine/Plugins/Experimental/ChaosVD/Source/ChaosVD/Public/Widgets/SChaosVDMainTab.h b/Engine/Plugins/Experimental/ChaosVD/Source/ChaosVD/Public/Widgets/SChaosVDMainTab.h
--- a/Engine/Plugins/Experimental/ChaosVD/Source/ChaosVD/Public/Widgets/SChaosVDMainTab.h
+++ b/Engine/Plugins/Experimental/ChaosVD/Source/ChaosVD/Public/Widgets/SChaosVDMainTab.h
@@ -22,7 +22,6 @@ public:
 		SLATE_ARGUMENT(TSharedPtr<SDockTab>, OwnerTab)
 	SLATE_END_ARGS()
 
-	virtual ~SChaosVDMainTab() override;
 	void Construct(const FArguments& InArgs, TSharedPtr<FChaosVDEngine> InChaosVDEngine);
 
 	TSharedRef<FChaosVDEngine> GetChaosVDEngineInstance() const { return ChaosVDEngine.ToSharedRef(); };
diff --git a/Engine/Plugins/Experimental/ChaosVD/Source/ChaosVD/Public/Widgets/SChaosVDPlaybackViewport.h b/Engine/Plugins/Experimental/ChaosVD/Source/ChaosVD/Public/Widgets/SChaosVDPlaybackViewport.h
--- a/Engine/Plugins/Experimental/ChaosVD/Source/ChaosVD/Public/Widgets/SChaosVDPlaybackViewport.h
+++ b/Engine/Plugins/Experimental/ChaosVD/Source/ChaosVD/Public/Widgets/SChaosVDPlaybackViewport.h
@@ -29,7 +29,9 @@ protected:
 
 	TSharedPtr<FLevelEditorViewportClient> CreateViewportClient() const;
 
-	void OnPlaybackControllerUpdated(FChaosVDPlaybackController* Controller) const;
+	void OnPlaybackControllerUpdated(TWeakPtr<FChaosVDPlaybackController> InController);
+
+	void RegisterNewController(TWeakPtr<FChaosVDPlaybackController> NewController);
 
 	void OnFrameSelectionUpdated(int32 NewFrameIndex) const;
 	void OnStepSelectionUpdated(int32 NewStepIndex) const;
diff --git a/Engine/Source/Runtime/Experimental/ChaosVisualDebugger/ChaosVDRuntime.Build.cs b/Engine/Source/Runtime/Experimental/ChaosVisualDebugger/ChaosVDRuntime.Build.cs
--- a/Engine/Source/Runtime/Experimental/ChaosVisualDebugger/ChaosVDRuntime.Build.cs
+++ b/Engine/Source/Runtime/Experimental/ChaosVisualDebugger/ChaosVDRuntime.Build.cs
@@ -12,7 +12,8 @@ namespace UnrealBuildTool.Rules
 			PublicDependencyModuleNames.AddRange(
 				new string[] {
 					"Core",
-					"CoreUObject"
+					"CoreUObject",
+					"TraceLog"
 				}
 			);
 		}
diff --git a/Engine/Source/Runtime/Experimental/ChaosVisualDebugger/Private/ChaosVDRecording.cpp b/Engine/Source/Runtime/Experimental/ChaosVisualDebugger/Private/ChaosVDRecording.cpp
--- a/Engine/Source/Runtime/Experimental/ChaosVisualDebugger/Private/ChaosVDRecording.cpp
+++ b/Engine/Source/Runtime/Experimental/ChaosVisualDebugger/Private/ChaosVDRecording.cpp
@@ -2,9 +2,52 @@
 
 #include "ChaosVDRecording.h"
 
-FChaosVDFrameData& FChaosVDRecording::AddFrame()
+int32 FChaosVDRecording::GetAvailableFramesNumber(const int32 SolverID) const
 {
-	RecordedFramesData.Add(FChaosVDFrameData());
+	if (const TArray<FChaosVDSolverFrameData>* SolverData = RecordedFramesDataPerSolver.Find(SolverID))
+	{
+		return SolverData->Num();
+	}
 
-	return RecordedFramesData.Last();
+	return INDEX_NONE;
+}
+
+FChaosVDSolverFrameData* FChaosVDRecording::GetFrameForSolver(const int32 SolverID, const int32 FrameNumber)
+{
+	if (TArray<FChaosVDSolverFrameData>* SolverFrames = RecordedFramesDataPerSolver.Find(SolverID))
+	{
+		//TODO: Find a safer way of do this. If someone stores this ptr bad things will happen
+		return SolverFrames->IsValidIndex(FrameNumber) ? &(*SolverFrames)[FrameNumber] : nullptr;
+	}
+
+	return nullptr;
+}
+
+void FChaosVDRecording::AddFrameForSolver(const int32 SolverID, FChaosVDSolverFrameData&& InFrameData)
+{
+	if (TArray<FChaosVDSolverFrameData>* SolverFrames = RecordedFramesDataPerSolver.Find(SolverID))
+	{
+		SolverFrames->Add(MoveTemp(InFrameData));
+	}
+	else
+	{	
+		RecordedFramesDataPerSolver.Add(SolverID, { MoveTemp(InFrameData) });
+	}
+
+	OnRecordingUpdated().Broadcast();
+}
+
+EChaosVDFrameLoadState FChaosVDRecording::GetFrameState(const int32 SolverID, const int32 FrameNumber)
+{
+	const TArray<FChaosVDSolverFrameData>* SolverFrames = RecordedFramesDataPerSolver.Find(SolverID);
+	if (SolverFrames && SolverFrames->IsValidIndex(FrameNumber))
+	{
+		return EChaosVDFrameLoadState::Loaded;
+	}
+	//TODO: Implement "Buffering" state handling
+	// I already implemented this on the version where we didn't use trace and CVD files were streamed from disk
+	// But removed the code for this CL - Leaving this getter here as it is used by the controller and I also left part
+	// of the implementation to handle different states there
+
+	return EChaosVDFrameLoadState::Unknown;
 }
diff --git a/Engine/Source/Runtime/Experimental/ChaosVisualDebugger/Private/ChaosVDRuntimeModule.cpp b/Engine/Source/Runtime/Experimental/ChaosVisualDebugger/Private/ChaosVDRuntimeModule.cpp
--- a/Engine/Source/Runtime/Experimental/ChaosVisualDebugger/Private/ChaosVDRuntimeModule.cpp
+++ b/Engine/Source/Runtime/Experimental/ChaosVisualDebugger/Private/ChaosVDRuntimeModule.cpp
@@ -7,17 +7,19 @@
 #include "HAL/IConsoleManager.h"
 #include "Misc/FileHelper.h"
 #include "Misc/Paths.h"
-#include "Serialization/MemoryWriter.h"
 #include "Modules/ModuleManager.h"
+#include "ProfilingDebugging/TraceAuxiliary.h"
 
 IMPLEMENT_MODULE(FChaosVDRuntimeModule, ChaosVDRuntime);
 
+DEFINE_LOG_CATEGORY_STATIC( LogChaosVDRuntime, Log, All );
+
 FAutoConsoleCommand ChaosVDStartRecordingCommand(
 	TEXT("p.Chaos.StartVDRecording"),
 	TEXT("Turn on the recording of debugging data"),
 	FConsoleCommandWithArgsDelegate::CreateLambda([](const TArray<FString>& Args)
 	{
-		FChaosVDRuntimeModule::Get().StartRecording();
+		FChaosVDRuntimeModule::Get().StartRecording(Args);
 	})
 );
 
@@ -37,57 +39,113 @@ FChaosVDRuntimeModule& FChaosVDRuntimeModule::Get()
 
 void FChaosVDRuntimeModule::StartupModule()
 {
+	FTraceAuxiliary::OnTraceStopped.AddRaw(this, &FChaosVDRuntimeModule::HandleTraceStopRequest);
 }
 
 void FChaosVDRuntimeModule::ShutdownModule()
 {
+	FTraceAuxiliary::OnTraceStopped.RemoveAll(this);
 }
 
-void FChaosVDRuntimeModule::StartRecording()
+int32 FChaosVDRuntimeModule::GenerateUniqueID()
 {
-	// Note: The code in this method is pretty much place holder code for testing purposes
+	return LastGeneratedID.Increment();
+}
 
-	if (bIsRecording)
-	{
-		return;
-	}
+void FChaosVDRuntimeModule::StopTrace()
+{
+	bRequestedStop = true;
+	FTraceAuxiliary::Stop();
+}
 
-	bIsRecording = true;
+void FChaosVDRuntimeModule::GenerateRecordingFileName(FString& OutFileName)
+{
+	const TCHAR* FilePrefix = TEXT("ChaosVD");
+	const FString FullPathPrefix = FPaths::ProfilingDir() / FilePrefix;
 
-	CurrentRecording = MakeShared<FChaosVDRecording>();
+	int32 Tries = 0;
+	do
+	{
+		OutFileName = FString::Printf(TEXT("%s_%d.utrace"), *FullPathPrefix, Tries++);
+	} while (IFileManager::Get().FileExists(*OutFileName));
 }
 
-void FChaosVDRuntimeModule::StopRecording()
+void FChaosVDRuntimeModule::StartRecording(const TArray<FString>& Args)
 {
-	// Note: The code in this method is pretty much place holder code for testing purposes
-
-	if (!ensure(bIsRecording))
+	if (bIsRecording)
 	{
 		return;
 	}
 
-	if (!ensure(CurrentRecording.IsValid()))
+#if UE_TRACE_ENABLED
+
+	// Other tools could bee using trace
+	// This is aggressive but until Trace supports multi-sessions, just take over.
+	if (FTraceAuxiliary::IsConnected())
 	{
-		return;
+		StopTrace();
 	}
 
-	const TCHAR* FilePrefix = TEXT("ChaosVD");
-	const FString FullPathPrefix = FPaths::ProfilingDir() / FilePrefix;
+	// Disable any enabled additional channel
+	UE::Trace::EnumerateChannels([](const ANSICHAR* ChannelName, bool bEnabled, void*)
+		{
+			if (bEnabled)
+			{
+				FString ChannelNameFString(ChannelName);
+				UE::Trace::ToggleChannel(ChannelNameFString.GetCharArray().GetData(), false);
+			}
+		}
+		, nullptr);
 
-	int32 Tries = 0;
-	FString UseFileName;
-	do
+
+	UE::Trace::ToggleChannel(TEXT("ChaosVDChannel"), true); 
+	UE::Trace::ToggleChannel(TEXT("Frame"), true);
+
+	if (Args.Num() == 0 || Args[0] == TEXT("File"))
 	{
-		UseFileName = FString::Printf(TEXT("%s_%d.cvd"), *FullPathPrefix, Tries++);
-	} while (IFileManager::Get().FileExists(*UseFileName));
+		FString RecordingFileName;
+		GenerateRecordingFileName(RecordingFileName);
 
+		bIsRecording = FTraceAuxiliary::Start(FTraceAuxiliary::EConnectionType::File, *RecordingFileName);
+	}
+	else if(Args[0] == TEXT("Server"))
+	{
+		const FString Target = Args.IsValidIndex(1) ? Args[1] : TEXT("127.0.0.1");
 
-	TArray<uint8> SerializedRecordedData;
-	FMemoryWriter RecordedDataWriter(SerializedRecordedData, false);
+		FTraceAuxiliary::Start(
+		FTraceAuxiliary::EConnectionType::Network,
+		*Target,
+		nullptr);
+	}
+#endif
+
+	ensure(bIsRecording);
+}
+
+void FChaosVDRuntimeModule::StopRecording()
+{
+	if (!ensure(bIsRecording))
+	{
+		return;
+	}
+#if UE_TRACE_ENABLED
+	
 
-	FChaosVDRecording::StaticStruct()->SerializeBin(RecordedDataWriter,CurrentRecording.Get());
+	UE::Trace::ToggleChannel(TEXT("ChaosVDChannel"), false);
+	UE::Trace::ToggleChannel(TEXT("Frame"), false); 
 
-	FFileHelper::SaveArrayToFile(SerializedRecordedData, *UseFileName);
+	StopTrace();
+#endif
 
 	bIsRecording = false;
 }
+
+void FChaosVDRuntimeModule::HandleTraceStopRequest(FTraceAuxiliary::EConnectionType TraceType, const FString& TraceDestination)
+{
+	if (!ensure(bRequestedStop))
+	{
+		UE_LOG(LogChaosVDRuntime, Warning, TEXT("Trace Recording has been stopped unexpectedly"));
+	}
+
+	bRequestedStop = false;
+}
diff --git a/Engine/Source/Runtime/Experimental/ChaosVisualDebugger/Public/ChaosVDRecording.h b/Engine/Source/Runtime/Experimental/ChaosVisualDebugger/Public/ChaosVDRecording.h
--- a/Engine/Source/Runtime/Experimental/ChaosVisualDebugger/Public/ChaosVDRecording.h
+++ b/Engine/Source/Runtime/Experimental/ChaosVisualDebugger/Public/ChaosVDRecording.h
@@ -3,9 +3,10 @@
 #pragma once
 #include "UObject/ObjectMacros.h"
 #include "Containers/UnrealString.h"
-
 #include "ChaosVDRecording.generated.h"
 
+DECLARE_MULTICAST_DELEGATE(FChaosVDRecordingUpdated)
+
 UENUM()
 enum class EChaosVDParticleType : uint8
 {
@@ -20,7 +21,19 @@ enum class EChaosVDParticleType : uint8
 };
 
 UENUM()
-enum class EChaosVDSolverType
+enum class EChaosVDParticleState : int8
+{
+	Uninitialized = 0,
+	Sleeping = 1,
+	Kinematic = 2,
+	Static = 3,
+	Dynamic = 4,
+
+	Count
+};
+
+UENUM()
+enum class EChaosVDSolverType : uint32
 {
 	Rigid
 };
@@ -32,90 +45,73 @@ struct FChaosVDParticleDebugData
 
 	UPROPERTY(VisibleAnywhere, Category= "Chaos Visual Debugger Data")
 	EChaosVDParticleType ParticleType;
-
 	UPROPERTY(VisibleAnywhere, Category= "Chaos Visual Debugger Data")
 	FString DebugName;
-
 	UPROPERTY(VisibleAnywhere, Category= "Chaos Visual Debugger Data")
 	int32 ParticleIndex;
-
 	UPROPERTY(VisibleAnywhere, Category= "Chaos Visual Debugger Data")
 	FVector Position;
-
 	UPROPERTY(VisibleAnywhere, Category= "Chaos Visual Debugger Data")
 	FQuat Rotation;
+	UPROPERTY(VisibleAnywhere, Category= "Chaos Visual Debugger Data")
+	FVector Velocity;
+	UPROPERTY(VisibleAnywhere, Category= "Chaos Visual Debugger Data")
+	FVector AngularVelocity;
+	UPROPERTY(VisibleAnywhere, Category= "Chaos Visual Debugger Data")
+	EChaosVDParticleState ParticleState;
 };
 
-USTRUCT()
 struct FChaosVDStepData
 {
-	GENERATED_BODY()
-
-	UPROPERTY()
 	TArray<FChaosVDParticleDebugData> RecordedParticles;
 };
 
-USTRUCT()
-struct FChaosVDSolverData
+struct CHAOSVDRUNTIME_API FChaosVDSolverFrameData
 {
-	GENERATED_BODY()
-
-	UPROPERTY()
-	EChaosVDSolverType SolverType;
-
-	UPROPERTY()
-	TArray<FChaosVDStepData> SolverSteps;
+	FString DebugName;
+	int32 SolverID;
+	TArray<FChaosVDStepData, TInlineAllocator<16>> SolverSteps;
 };
 
-USTRUCT()
-struct FChaosVDEventData
+enum class EChaosVDFrameLoadState
 {
-	GENERATED_BODY()
-
-	UPROPERTY()
-	FString EventID;
-
-	UPROPERTY()
-	TArray<uint8> SerializedEventData;
+	Unloaded,
+	Loaded,
+	Buffering,
+	Unknown
 };
 
-USTRUCT()
-struct CHAOSVDRUNTIME_API FChaosVDFrameData
+/**
+ * Struct that represents a recorded Physics simulation.
+ * It is currently populated while analyzing a Trace session
+ */
+struct CHAOSVDRUNTIME_API FChaosVDRecording
 {
-	GENERATED_BODY()
-
-	UPROPERTY()
-	TMap<EChaosVDSolverType, FChaosVDSolverData> RecordedSolvers;
-
-	UPROPERTY()
-	TArray<FChaosVDEventData> RecordedEvents;
+	/** Returns the current available recorded frames per recorded solver */
+	int32 GetAvailableFramesNumber(const int32 SolverID) const;
+	/** Returns the current available recorded solvers number */
+	int32 GetAvailableSolversNumber() const { return RecordedFramesDataPerSolver.Num();}
+	
+	const TMap<int32, TArray<FChaosVDSolverFrameData>>& GetAvailableSolvers() const { return RecordedFramesDataPerSolver; }
 
-	FChaosVDSolverData& GetSolverData(EChaosVDSolverType SolverType) { return RecordedSolvers.Contains(SolverType) ? RecordedSolvers[SolverType] : RecordedSolvers.Add(SolverType,FChaosVDSolverData()); }
+	/** Returns a ptr to the recorded frame data for a specific solver -  Do not store as it is a reference to the element in the array */
+	FChaosVDSolverFrameData* GetFrameForSolver(const int32 SolverID, const int32 FrameNumber);
 
-	//FChaosVDSolverData& AddSolverData(EChaosVDSolverType SolverType) { return RecordedSolvers.Contains(SolverType) ? RecordedSolvers[SolverType] : RecordedSolvers.Add(SolverType,FChaosVDSolverData()); }
-};
-
-USTRUCT()
-struct FChaosVDRecordingHeader
-{
-	GENERATED_BODY()
+	/** Adds a Frame Data for a specific Solver ID. Creates a solver entry if it does not exist */
+	void AddFrameForSolver(const int32 SolverID, FChaosVDSolverFrameData&& InFrameData);
 
-	UPROPERTY()
-	bool bIsRealTimeRecording = false;
-};
-
-USTRUCT()
-struct CHAOSVDRUNTIME_API FChaosVDRecording 
-{
-	GENERATED_BODY()
+	/** Returns the current frame state of a frame used to determine if it is ready for use*/
+	EChaosVDFrameLoadState GetFrameState(const int32 SolverID, const int32 FrameNumber);
 
-	UPROPERTY()
-	FChaosVDRecordingHeader RecordingHeader;
+	/** Called each time the recording changes - Mainly when a new frame is added from the Trace analysis */
+	FChaosVDRecordingUpdated& OnRecordingUpdated() { return RecordingUpdatedDelegate; };
 
-	UPROPERTY()
-	TArray<FChaosVDFrameData> RecordedFramesData;
+	/** Session name of the trace session used to re-build this recording */
+	FString SessionName;
 
-	FChaosVDFrameData& GetCurrentFrame() { return RecordedFramesData.Num() ? RecordedFramesData.Last() : AddFrame(); }
+protected:
 
-	FChaosVDFrameData& AddFrame();
+	TMap<int32, EChaosVDFrameLoadState> AvailableFramesState;
+	TMap<int32, TArray<FChaosVDSolverFrameData>> RecordedFramesDataPerSolver;
+	FChaosVDRecordingUpdated RecordingUpdatedDelegate;
 };
diff --git a/Engine/Source/Runtime/Experimental/ChaosVisualDebugger/Public/ChaosVDRuntimeModule.h b/Engine/Source/Runtime/Experimental/ChaosVisualDebugger/Public/ChaosVDRuntimeModule.h
--- a/Engine/Source/Runtime/Experimental/ChaosVisualDebugger/Public/ChaosVDRuntimeModule.h
+++ b/Engine/Source/Runtime/Experimental/ChaosVisualDebugger/Public/ChaosVDRuntimeModule.h
@@ -4,6 +4,7 @@
 
 #include "Templates/SharedPointer.h"
 #include "Modules/ModuleInterface.h"
+#include "ProfilingDebugging/TraceAuxiliary.h"
 
 struct FChaosVDRecording;
 
@@ -16,16 +17,33 @@ public:
 	virtual void StartupModule() override;
 	virtual void ShutdownModule() override;
 
-	TSharedPtr<FChaosVDRecording>& GetRecording() { return CurrentRecording; }
-
-	void StartRecording();
+	/** Starts a CVD recording by starting a Trace session. It will stop any existing trace session
+	 * @param Args : Arguments array provided by the commandline. Used to determine if we want to record to file or a local trace server
+	 */
+	void StartRecording(const TArray<FString>& Args);
+	/**/
 	void StopRecording();
 
+	/** Returns true if we are currently recording a Physics simulation */
 	bool IsRecording() const { return bIsRecording; }
 
+	/** Returns a unique ID used to be used to identify CVD (Chaos Visual Debugger) data */
+	int32 GenerateUniqueID();
+
 private:
 
+	/** Stops the current Trace session */
+	void StopTrace();
+	/** Finds a valid file name for a new file - Used to generate the file name for the Trace recording */
+	void GenerateRecordingFileName(FString& OutFileName);
+
+	/** Used to handle stop requests to the active trace session that were not done by us
+	 * That is a possible scenario because Trace is shared by other In-Editor tools
+	 */
+	void HandleTraceStopRequest(FTraceAuxiliary::EConnectionType TraceType, const FString& TraceDestination);
+
 	bool bIsRecording = false;
-	TSharedPtr<FChaosVDRecording> CurrentRecording;
+	bool bRequestedStop = false;
 
+	FThreadSafeCounter LastGeneratedID;
 };
