[Chaos VD] Fix CIS issue in the ChaosVDPlaybackController

Changing the base class of FChaosVDTraceProvider to TraceServices::IEditableProvider instead of TraceServices::IProvider

#rnx

diff --git a/Engine/Plugins/Experimental/ChaosVD/Source/ChaosVD/Public/Trace/ChaosVDTraceProvider.h b/Engine/Plugins/Experimental/ChaosVD/Source/ChaosVD/Public/Trace/ChaosVDTraceProvider.h
--- a/Engine/Plugins/Experimental/ChaosVD/Source/ChaosVD/Public/Trace/ChaosVDTraceProvider.h
+++ b/Engine/Plugins/Experimental/ChaosVD/Source/ChaosVD/Public/Trace/ChaosVDTraceProvider.h
@@ -12,7 +12,7 @@ struct FChaosVDRecording;
  * It stores and handles rebuilt recorded frame data from Trace events
  * dispatched by the Chaos VD Trace analyzer
  */
-class FChaosVDTraceProvider : public TraceServices::IProvider
+class FChaosVDTraceProvider : public TraceServices::IEditableProvider
 {
 public:
 	
