[Chaos VD] Fix CIS issue on Chaos VD plugin

Removing use of Trace Edit Provider for now

#rnx

diff --git a/Engine/Plugins/Experimental/ChaosVD/Source/ChaosVD/Private/ChaosVDPlaybackController.cpp b/Engine/Plugins/Experimental/ChaosVD/Source/ChaosVD/Private/ChaosVDPlaybackController.cpp
--- a/Engine/Plugins/Experimental/ChaosVD/Source/ChaosVD/Private/ChaosVDPlaybackController.cpp
+++ b/Engine/Plugins/Experimental/ChaosVD/Source/ChaosVD/Private/ChaosVDPlaybackController.cpp
@@ -67,7 +67,8 @@ void FChaosVDPlaybackController::UnloadCurrentRecording(const bool bBroadcastUpd
 	{
 		LoadedRecording->OnRecordingUpdated().RemoveAll(this);
 
-		if (const TSharedPtr<const TraceServices::IAnalysisSession> TraceSession = FChaosVDModule::Get().GetTraceManager()->GetSession(LoadedRecording->SessionName))
+		//TODO: Add proper support to edit the trace provider
+		/*if (const TSharedPtr<const TraceServices::IAnalysisSession> TraceSession = FChaosVDModule::Get().GetTraceManager()->GetSession(LoadedRecording->SessionName))
 		{
 			TraceServices::IAnalysisSession* TraceSessionPtr = const_cast<TraceServices::IAnalysisSession*>(TraceSession.Get());
 			if (FChaosVDTraceProvider* ChaosVDProvider = TraceSessionPtr->EditProvider<FChaosVDTraceProvider>(FChaosVDTraceProvider::ProviderName))
@@ -75,7 +76,7 @@ void FChaosVDPlaybackController::UnloadCurrentRecording(const bool bBroadcastUpd
 				TraceServices::FAnalysisSessionEditScope SessionEditScope(*TraceSessionPtr);
 				ChaosVDProvider->DeleteRecordingInstanceForSession();
 			}
-		}
+		}*/
 
 		LoadedRecording.Reset();
 	}
diff --git a/Engine/Plugins/Experimental/ChaosVD/Source/ChaosVD/Public/Trace/ChaosVDTraceProvider.h b/Engine/Plugins/Experimental/ChaosVD/Source/ChaosVD/Public/Trace/ChaosVDTraceProvider.h
--- a/Engine/Plugins/Experimental/ChaosVD/Source/ChaosVD/Public/Trace/ChaosVDTraceProvider.h
+++ b/Engine/Plugins/Experimental/ChaosVD/Source/ChaosVD/Public/Trace/ChaosVDTraceProvider.h
@@ -12,7 +12,7 @@ struct FChaosVDRecording;
  * It stores and handles rebuilt recorded frame data from Trace events
  * dispatched by the Chaos VD Trace analyzer
  */
-class FChaosVDTraceProvider : public TraceServices::IEditableProvider
+class FChaosVDTraceProvider : public TraceServices::IProvider
 {
 public:
 	
