[Chaos VD] Initial implementation of Geometry Data Recording using Trace

- Created mesh generators for several chaos geometry types
- Added support to Generate Dynamic Mesh or Static meshes (and use instanced rendering when using static meshes)
- Added a new macro to record binary data using Trace
- Extended the Chaos VD Trace Analyzer to read the binary data
- Updated the ChaosArchiveContext to be able to map TagsToObjects
- Updated the ChaosVDParticleActor to use the new geometry data builder when new geometry data is available
- Record the simulation space for RBAN solvers
- Added a cvar to enable compression of the Traced Binary data using Oodle, and another to control the compression level. It is not cheap, and we need to do further testing to evaluate what are the best settings for it. Meanwhile, it is disabled by default

diff --git a/Engine/Plugins/Experimental/ChaosVD/Source/ChaosVD/ChaosVD.Build.cs b/Engine/Plugins/Experimental/ChaosVD/Source/ChaosVD/ChaosVD.Build.cs
--- a/Engine/Plugins/Experimental/ChaosVD/Source/ChaosVD/ChaosVD.Build.cs
+++ b/Engine/Plugins/Experimental/ChaosVD/Source/ChaosVD/ChaosVD.Build.cs
@@ -14,6 +14,7 @@ public class ChaosVD : ModuleRules
 				"Engine",
 				"ChaosVDRuntime", 
 				"TraceServices",
+				"ChaosVDData"
 			});
 
 		PrivateDependencyModuleNames.AddRange(
@@ -32,8 +33,15 @@ public class ChaosVD : ModuleRules
 				"WorkspaceMenuStructure", 
 				"TraceAnalysis", 
 				"TraceInsights",
-				"TraceLog",
+				"TraceLog", 
+				"GeometryFramework",
+				"GeometryCore",
+				"MeshConversion",
+				"MeshDescription",
+				"StaticMeshDescription",
 			}
 			);
+
+		SetupModulePhysicsSupport(Target);
 	}
 }
diff --git a/Engine/Plugins/Experimental/ChaosVD/Source/ChaosVD/Private/ChaosVDConvexMeshGenerator.cpp b/Engine/Plugins/Experimental/ChaosVD/Source/ChaosVD/Private/ChaosVDConvexMeshGenerator.cpp
new file mode 100644
--- /dev/null
+++ b/Engine/Plugins/Experimental/ChaosVD/Source/ChaosVD/Private/ChaosVDConvexMeshGenerator.cpp
@@ -0,0 +1,103 @@
+// Copyright Epic Games, Inc. All Rights Reserved.
+
+#include "ChaosVDConvexMeshGenerator.h"
+
+#include "Chaos/Core.h"
+#include "Chaos/Convex.h"
+#include "Chaos/Utilities.h"
+#include "CompGeom/PolygonTriangulation.h"
+
+using namespace UE::Geometry;
+
+struct FChaosVDTriangleData
+{
+	int32 FaceIndex;
+	TArray<FIndex3i> FaceTriangles;
+	int NormalIndices[3];
+	FVector Normal;
+	double TriangleArea = 0;
+};
+
+void FChaosVDConvexMeshGenerator::GenerateFromConvex(const Chaos::FConvex& Convex)
+{
+	using namespace Chaos;
+	using FRealType = FRealSingle;
+	using FVec3Type = TVec3<FRealType>;
+
+	const int32 FacesNum = Convex.GetFaces().Num();
+
+	TArray<FChaosVDTriangleData> PerFaceTriangleData;
+	PerFaceTriangleData.Reserve(FacesNum);
+	
+	// Triangulate each face of the Convex shape
+	int32 TriangleCount = 0;
+	for (int32 PlaneIndex = 0; PlaneIndex < Convex.GetFaces().Num(); ++PlaneIndex)
+	{
+		const int32 PlaneVerticesNum = Convex.NumPlaneVertices(PlaneIndex);
+
+		// Get the positions of each vertex of the plane
+		TArray<FVector> PlaneVertices;
+		PlaneVertices.Reserve(PlaneVerticesNum);
+
+		for (int32 PlaneLocalVertexIndex = 0; PlaneLocalVertexIndex < PlaneVerticesNum; ++PlaneLocalVertexIndex)
+		{
+			const int32 GlobalVertexIndex = Convex.GetPlaneVertex(PlaneIndex, PlaneLocalVertexIndex);
+			PlaneVertices.Add(FVector(UE::Math::TVector<double>(Convex.GetVertex(GlobalVertexIndex))));
+		}
+
+		FChaosVDTriangleData FaceTriangleData;
+		PolygonTriangulation::TriangulateSimplePolygon(PlaneVertices,FaceTriangleData.FaceTriangles);
+
+		FaceTriangleData.Normal = Convex.GetPlane(PlaneIndex).Normal();
+		FaceTriangleData.FaceIndex = PlaneIndex;
+
+		TriangleCount += FaceTriangleData.FaceTriangles.Num();
+		
+		PerFaceTriangleData.Add(MoveTemp(FaceTriangleData));
+	}
+
+	const TArray<FVec3Type>& ConvexVertices = Convex.GetVertices();
+	const int32 NormalsNum = TriangleCount * 3;
+	constexpr int32 UVsNum = 0;
+	SetBufferSizes(ConvexVertices.Num(), TriangleCount, UVsNum, NormalsNum);
+
+	// Fill the vertex buffer with the transformed vertices of the Convex Shape
+	for (int32 i = 0; i < ConvexVertices.Num(); ++i) 
+	{
+		Vertices[i] = FVector3d(UE::Math::TVector<double>(ConvexVertices[i]));
+	}
+
+	int32 TriangleIndex = 0;
+	int CurrentNormalIndex = 0;
+	for (FChaosVDTriangleData& Data : PerFaceTriangleData)
+	{
+		for (FIndex3i& PolyTriangle : Data.FaceTriangles)
+		{
+			// Re-map from Plane Local vertex to Global Vertices
+			PolyTriangle.A = Convex.GetPlaneVertex(Data.FaceIndex, PolyTriangle.A);
+			PolyTriangle.B = Convex.GetPlaneVertex(Data.FaceIndex, PolyTriangle.B);
+			PolyTriangle.C = Convex.GetPlaneVertex(Data.FaceIndex, PolyTriangle.C);
+
+			// Use the Face Normal in all Vertices of the triangle
+			for (int LocalVertexIndex = 0; LocalVertexIndex < 3; ++LocalVertexIndex)
+			{
+				Normals[CurrentNormalIndex] = UE::Math::TVector<float>(Data.Normal);
+				Data.NormalIndices[LocalVertexIndex] = CurrentNormalIndex;
+				CurrentNormalIndex++;
+			}
+
+			SetTriangle(TriangleIndex, PolyTriangle);
+			SetTrianglePolygon(TriangleIndex,Data.FaceIndex);
+			SetTriangleNormals(TriangleIndex, Data.NormalIndices[0], Data.NormalIndices[1], Data.NormalIndices[2]);
+			TriangleIndex++;
+		}
+	}
+
+	bIsGenerated = true;
+}
+
+UE::Geometry::FMeshShapeGenerator& FChaosVDConvexMeshGenerator::Generate()
+{
+	ensureAlwaysMsgf(bIsGenerated, TEXT("You need to call FChaosVDConvexMeshGenerator::GenerateFromConvex before calling Generate"));
+	return *this;
+}
diff --git a/Engine/Plugins/Experimental/ChaosVD/Source/ChaosVD/Private/ChaosVDGeometryBuilder.cpp b/Engine/Plugins/Experimental/ChaosVD/Source/ChaosVD/Private/ChaosVDGeometryBuilder.cpp
new file mode 100644
--- /dev/null
+++ b/Engine/Plugins/Experimental/ChaosVD/Source/ChaosVD/Private/ChaosVDGeometryBuilder.cpp
@@ -0,0 +1,77 @@
+// Copyright Epic Games, Inc. All Rights Reserved.
+
+#include "ChaosVDGeometryBuilder.h"
+
+#include "ChaosVDHeightFieldMeshGenerator.h"
+#include "DynamicMeshToMeshDescription.h"
+#include "Engine/StaticMesh.h"
+#include "StaticMeshAttributes.h"
+#include "UDynamicMesh.h"
+
+#include "UObject/UObjectGlobals.h"
+
+void FChaosVDGeometryBuilder::AddReferencedObjects(FReferenceCollector& Collector)
+{
+	Collector.AddStableReferenceMap(MeshCacheMap);
+	Collector.AddStableReferenceMap(StaticMeshCacheMap);
+}
+
+UDynamicMesh* FChaosVDGeometryBuilder::CreateAndCacheDynamicMesh(const uint32 GeometryCacheKey, UE::Geometry::FMeshShapeGenerator& MeshGenerator)
+{
+	//TODO: Make this return what is cached when the system is more robust
+	// For now this should not happen and we want to catch it and make it visually noticeable
+	if (MeshCacheMap.Contains(GeometryCacheKey))
+	{
+		ensureMsgf(false, TEXT("Tried to create a new mesh with an existing Cache key"));
+		return nullptr;
+	}
+
+	UDynamicMesh* Mesh = NewObject<UDynamicMesh>();
+	Mesh->SetMesh(&MeshGenerator.Generate());
+
+	MeshCacheMap.Add(GeometryCacheKey, Mesh);
+
+	return Mesh;
+}
+
+UStaticMesh* FChaosVDGeometryBuilder::CreateAndCacheStaticMesh(const uint32 GeometryCacheKey, UE::Geometry::FMeshShapeGenerator& MeshGenerator)
+{
+	//TODO: Make this return what is cached when the system is more robust
+	// For now this should not happen and we want to catch it and make it visually noticeable
+	if (StaticMeshCacheMap.Contains(GeometryCacheKey))
+	{
+		ensureMsgf(false, TEXT("Tried to create a new mesh with an existing Cache key"));
+		return nullptr;
+	}
+
+	//TODO: Instead of generating a dynamic mesh and discard it, we should
+	// Create a Mesh description directly.
+	// We could create a base class for our mesh Generators and add a Generate method that generates these mesh descriptions
+	UDynamicMesh* Mesh = NewObject<UDynamicMesh>();
+	UStaticMesh* StaticMesh = NewObject<UStaticMesh>();
+	StaticMesh->GetStaticMaterials().Add(FStaticMaterial());
+
+	Mesh->SetMesh(&MeshGenerator.Generate());
+
+	FMeshDescription MeshDescription;
+	FStaticMeshAttributes Attributes(MeshDescription);
+	Attributes.Register();
+
+	FDynamicMeshToMeshDescription Converter;
+	Converter.Convert(Mesh->GetMeshPtr(), MeshDescription, true);
+
+	UStaticMesh::FBuildMeshDescriptionsParams Params;
+	Params.bUseHashAsGuid = true;
+	Params.bMarkPackageDirty = false;
+	Params.bBuildSimpleCollision = false;
+	Params.bCommitMeshDescription = false;
+	Params.bFastBuild = true;
+
+	StaticMesh->NaniteSettings.bEnabled = true;
+	StaticMesh->BuildFromMeshDescriptions({&MeshDescription}, Params);
+
+	StaticMeshCacheMap.Add(GeometryCacheKey, StaticMesh);
+
+	return StaticMesh;
+}
+
diff --git a/Engine/Plugins/Experimental/ChaosVD/Source/ChaosVD/Private/ChaosVDHeightfieldMeshGenerator.cpp b/Engine/Plugins/Experimental/ChaosVD/Source/ChaosVD/Private/ChaosVDHeightfieldMeshGenerator.cpp
new file mode 100644
--- /dev/null
+++ b/Engine/Plugins/Experimental/ChaosVD/Source/ChaosVD/Private/ChaosVDHeightfieldMeshGenerator.cpp
@@ -0,0 +1,87 @@
+// Copyright Epic Games, Inc. All Rights Reserved.
+
+#include "ChaosVDHeightFieldMeshGenerator.h"
+
+#include "Chaos/HeightField.h"
+
+void FChaosVDHeightFieldMeshGenerator::AppendTriangle(int32& OutTriIndex, int& OutCurrentNormalIndex, int32 PolygonIndex, const Chaos::TVec2<Chaos::FReal>& InCellCoordinates, const UE::Geometry::FIndex3i& InTriangle, const Chaos::FHeightField& InHeightField)
+{
+	int32 NormalIndices[3];
+    for (int32 LocalVertexIndex = 0; LocalVertexIndex < 3; ++LocalVertexIndex)
+    {
+        Normals[OutCurrentNormalIndex] = UE::Math::TVector<float>(InHeightField.GetNormalAt(InCellCoordinates));
+        NormalIndices[LocalVertexIndex] = OutCurrentNormalIndex;
+        OutCurrentNormalIndex++;
+    }
+
+	SetTriangle(OutTriIndex, InTriangle);
+	SetTrianglePolygon(OutTriIndex, PolygonIndex);
+	SetTriangleNormals(OutTriIndex, NormalIndices[0], NormalIndices[1], NormalIndices[2]);
+	OutTriIndex++;
+}
+
+void FChaosVDHeightFieldMeshGenerator::GenerateFromHeightField(const Chaos::FHeightField& InHeightField)
+{
+	using namespace UE::Geometry;
+
+	const int32 NumRows = InHeightField.GetNumRows();
+	const int32 NumCols = InHeightField.GetNumCols();
+	const int32 VertexCount = NumRows * NumCols;
+
+	const int32 NumQuads = (NumRows - 1) * (NumCols - 1);
+	const int32 NumTris = NumQuads * 2;
+	const int32 NumNormals = NumTris * 3;
+	constexpr int32 NumUVs = 0;
+	
+	SetBufferSizes(VertexCount, NumTris, NumUVs, NumNormals);
+
+	//Fill the vertex buffer with the height data
+	for (int32 Y = 0; Y < NumRows; Y++)
+	{
+		for (int32 X = 0; X < NumCols; X++)
+		{
+			const int32 SampleIndex = Y * NumCols + X;
+			Vertices[SampleIndex] = InHeightField.GetPointScaled(SampleIndex);
+		}
+	}
+  
+	int32 CurrentNormalIndex = 0;
+	int32 CurrentTriangleIndex = 0;
+	int32 CurrentPolygonIndex = 0;
+	for (int32 Y = 0; Y < NumRows - 1; Y++)
+	{
+		for (int32 X = 0; X < NumCols - 1; X++)
+		{
+			if (InHeightField.IsHole(X, Y))
+			{
+				continue;
+			}
+
+			Chaos::TVec2<Chaos::FReal> CellCoordinates(X,Y);
+
+			// Vertices of each corner of the current cell
+			const int32 Vertex0 = Y * NumCols + X;
+			const int32 Vertex1 = Vertex0 + 1;
+			const int32 Vertex2 = Vertex0 + NumCols;
+			const int32 Vertex3 = Vertex2 + 1;
+
+			// Define the two triangles that form the current cell and add it to the generator
+			FIndex3i Triangle(Vertex0,Vertex3,Vertex1);
+			FIndex3i Triangle2(Vertex0,Vertex2,Vertex3);
+
+			AppendTriangle(CurrentTriangleIndex,CurrentNormalIndex, CurrentPolygonIndex, CellCoordinates, Triangle, InHeightField);
+			AppendTriangle(CurrentTriangleIndex,CurrentNormalIndex, CurrentPolygonIndex, CellCoordinates, Triangle2, InHeightField);
+
+			CurrentPolygonIndex++;
+		}
+	}
+
+	bIsGenerated = true;
+}
+
+
+UE::Geometry::FMeshShapeGenerator& FChaosVDHeightFieldMeshGenerator::Generate()
+{
+	ensureAlwaysMsgf(bIsGenerated, TEXT("You need to call FChaosVDHeightFieldMeshGenerator::GenerateFromHeightField before calling Generate"));
+	return *this;
+}
diff --git a/Engine/Plugins/Experimental/ChaosVD/Source/ChaosVD/Private/ChaosVDParticleActor.cpp b/Engine/Plugins/Experimental/ChaosVD/Source/ChaosVD/Private/ChaosVDParticleActor.cpp
--- a/Engine/Plugins/Experimental/ChaosVD/Source/ChaosVD/Private/ChaosVDParticleActor.cpp
+++ b/Engine/Plugins/Experimental/ChaosVD/Source/ChaosVD/Private/ChaosVDParticleActor.cpp
@@ -1,6 +1,9 @@
 // Copyright Epic Games, Inc. All Rights Reserved.
 
 #include "ChaosVDParticleActor.h"
+
+#include "ChaosVDGeometryBuilder.h"
+#include "ChaosVDScene.h"
 #include "ChaosVDRecording.h"
 #include "Components/SceneComponent.h"
 #include "Components/StaticMeshComponent.h"
@@ -9,35 +12,65 @@
 
 AChaosVDParticleActor::AChaosVDParticleActor(const FObjectInitializer& ObjectInitializer) : Super(ObjectInitializer)
 {
-	RootComponent = CreateDefaultSubobject<USceneComponent>(TEXT("SceneComponent0"));;
-	
-	//Note: This static mesh component is just a place holder for debugging purposes 
-	StaticMeshComponent = CreateDefaultSubobject<UStaticMeshComponent>(TEXT("StaticMeshComponent0"));
-	StaticMeshComponent->SetCollisionProfileName(UCollisionProfile::NoCollision_ProfileName);
-	StaticMeshComponent->SetGenerateOverlapEvents(false);
-	StaticMeshComponent->SetupAttachment(RootComponent);
-
-	StaticMeshComponent->Mobility = EComponentMobility::Movable;
-	StaticMeshComponent->bUseDefaultCollision = false;
-	StaticMeshComponent->bCastStaticShadow = false;
-	StaticMeshComponent->bCastDynamicShadow = false;
-	StaticMeshComponent->bSelectable = true;
+	RootComponent = CreateDefaultSubobject<USceneComponent>(TEXT("SceneComponent0"));
 }
 
-void AChaosVDParticleActor::UpdateFromRecordedData(const FChaosVDParticleDebugData& RecordedData)
+void AChaosVDParticleActor::UpdateFromRecordedData(const FChaosVDParticleDebugData& RecordedData, const Chaos::FRigidTransform3& SimulationTransform)
 {
-	SetActorLocation(RecordedData.Position);
-	SetActorRotation(RecordedData.Rotation);
+	SetActorLocationAndRotation(SimulationTransform.TransformPosition(RecordedData.Position), SimulationTransform.GetRotation() * RecordedData.Rotation, false);
 
 	// TODO: We should store a ptr to the data and in our custom details panel draw it, but the current data format is just a test
 	// So we should do this once it is final
 	RecordedDebugData = RecordedData;
 }
 
-void AChaosVDParticleActor::SetStaticMesh(UStaticMesh* NewStaticMesh)
+void AChaosVDParticleActor::UpdateGeometryData(const TSharedPtr<const Chaos::FImplicitObject>& ImplicitObject)
+{
+	if (bIsGeometryDataLoaded)
+	{
+		return;
+	}
+
+	TArray<UMeshComponent*> OutGeneratedMeshComponents;
+
+	if (const TSharedPtr<FChaosVDScene>& ScenePtr = OwningScene.Pin())
+	{
+		if (const TSharedPtr<FChaosVDGeometryBuilder>& GeometryGenerator = ScenePtr->GetGeometryGenerator())
+		{
+			Chaos::FRigidTransform3 Transform;
+			GeometryGenerator->CreateMeshComponentsFromImplicit<UStaticMesh, UInstancedStaticMeshComponent>(ImplicitObject.Get(), this, OutGeneratedMeshComponents, Transform);
+		}
+	}
+	
+	if (OutGeneratedMeshComponents.Num() > 0)
+	{
+		bIsGeometryDataLoaded = true;
+	}
+}
+
+void AChaosVDParticleActor::SetScene(const TSharedPtr<FChaosVDScene>& InScene)
 {
-	if (ensure(StaticMeshComponent))
+	OwningScene = InScene;
+
+	if (const TSharedPtr<FChaosVDScene>& ScenePtr = OwningScene.Pin())
 	{
-		StaticMeshComponent->SetStaticMesh(NewStaticMesh);
+		GeometryUpdatedDelegate = ScenePtr->OnNewGeometryAvailable().AddWeakLambda(this, [this](const TSharedPtr<const Chaos::FImplicitObject>& ImplicitObject, const int32 ID)
+		{
+			if (RecordedDebugData.ImplicitObjectID == ID)
+			{
+				UpdateGeometryData(ImplicitObject);
+			}
+		});
 	}
 }
+
+
+void AChaosVDParticleActor::BeginDestroy()
+{
+	if (const TSharedPtr<FChaosVDScene>& ScenePtr = OwningScene.Pin())
+	{
+		ScenePtr->OnNewGeometryAvailable().Remove(GeometryUpdatedDelegate);
+	}
+
+	Super::BeginDestroy();
+}
diff --git a/Engine/Plugins/Experimental/ChaosVD/Source/ChaosVD/Private/ChaosVDPlaybackController.cpp b/Engine/Plugins/Experimental/ChaosVD/Source/ChaosVD/Private/ChaosVDPlaybackController.cpp
--- a/Engine/Plugins/Experimental/ChaosVD/Source/ChaosVD/Private/ChaosVDPlaybackController.cpp
+++ b/Engine/Plugins/Experimental/ChaosVD/Source/ChaosVD/Private/ChaosVDPlaybackController.cpp
@@ -52,6 +52,19 @@ bool FChaosVDPlaybackController::LoadChaosVDRecordingFromTraceSession(const FStr
 	{
 		GoToRecordedStep(SolversIterator->Key, 0, 0);
 	}
+
+	LoadedRecording->OnGeometryDataLoaded().AddLambda([this](const TSharedPtr<const Chaos::FImplicitObject>& NewGeometry, const int32 GeometryID)
+	{
+		if (const TSharedPtr<FChaosVDScene> ScenePtr = SceneToControl.Pin())
+		{
+			ScenePtr->HandleNewGeometryData(NewGeometry, GeometryID);
+		}
+	});
+	
+	if (TSharedPtr<FChaosVDScene> ScenePtr = SceneToControl.Pin())
+	{
+		ScenePtr->LoadedRecording = LoadedRecording;
+	}
 	
 	OnControllerUpdated().ExecuteIfBound(AsWeak());
 
@@ -124,7 +137,7 @@ void FChaosVDPlaybackController::GoToRecordedStep(const int32 InTrackID, const i
 				const FChaosVDSolverFrameData* SolverFrameData = LoadedRecording->GetFrameForSolver(InTrackID, FrameNumber);
 				if (SolverFrameData && ensure(SolverFrameData->SolverSteps.IsValidIndex(Step)))
 				{
-					SceneToControlSharedPtr->UpdateFromRecordedStepData(InTrackID, SolverFrameData->DebugName, SolverFrameData->SolverSteps[Step]);
+					SceneToControlSharedPtr->UpdateFromRecordedStepData(InTrackID, SolverFrameData->DebugName, SolverFrameData->SolverSteps[Step], *SolverFrameData);
 				}
 			}
 
diff --git a/Engine/Plugins/Experimental/ChaosVD/Source/ChaosVD/Private/ChaosVDScene.cpp b/Engine/Plugins/Experimental/ChaosVD/Source/ChaosVD/Private/ChaosVDScene.cpp
--- a/Engine/Plugins/Experimental/ChaosVD/Source/ChaosVD/Private/ChaosVDScene.cpp
+++ b/Engine/Plugins/Experimental/ChaosVD/Source/ChaosVD/Private/ChaosVDScene.cpp
@@ -3,14 +3,20 @@
 #include "ChaosVDScene.h"
 
 #include "ChaosVDEditorSettings.h"
-#include "ChaosVDModule.h"
+#include "ChaosVDGeometryBuilder.h"
 #include "ChaosVDParticleActor.h"
+#include "Chaos/ImplicitObject.h"
 #include "ChaosVDRecording.h"
 #include "Engine/Engine.h"
 #include "Engine/LevelStreamingDynamic.h"
 #include "Engine/World.h"
+
 #include "UObject/Package.h"
 
+FChaosVDScene::FChaosVDScene() = default;
+
+FChaosVDScene::~FChaosVDScene() = default;
+
 void FChaosVDScene::Initialize()
 {
 	if (!ensure(!bIsInitialized))
@@ -20,6 +26,8 @@ void FChaosVDScene::Initialize()
 
 	PhysicsVDWorld = CreatePhysicsVDWorld();
 
+	GeometryGenerator = MakeShared<FChaosVDGeometryBuilder>();
+
 	bIsInitialized = true;
 }
 
@@ -30,6 +38,8 @@ void FChaosVDScene::DeInitialize()
 		return;
 	}
 
+	GeometryGenerator.Reset();
+
 	if (PhysicsVDWorld)
 	{
 		PhysicsVDWorld->DestroyWorld(true);
@@ -47,22 +57,22 @@ void FChaosVDScene::AddReferencedObjects(FReferenceCollector& Collector)
 	Collector.AddStableReferenceMap(ParticleVDInstancesByID);
 }
 
-void FChaosVDScene::UpdateFromRecordedStepData(const int32 SolverID, const FString& SolverName, const FChaosVDStepData& InRecordedStepData)
+void FChaosVDScene::UpdateFromRecordedStepData(const int32 SolverID, const FString& SolverName, const FChaosVDStepData& InRecordedStepData, const FChaosVDSolverFrameData& InFrameData)
 {
-	TSet<uint32> ParticlesIDsInRecordedStepData;
+	TSet<int32> ParticlesIDsInRecordedStepData;
 	ParticlesIDsInRecordedStepData.Reserve(InRecordedStepData.RecordedParticles.Num());
 
 	// Go over existing Particle VD Instances and update them or create them if needed 
 	for (const FChaosVDParticleDebugData& Particle : InRecordedStepData.RecordedParticles)
 	{
-		const uint32 ParticleVDInstanceID = GetIDForRecordedParticleData(Particle);
+		const int32 ParticleVDInstanceID = GetIDForRecordedParticleData(Particle);
 		ParticlesIDsInRecordedStepData.Add(ParticleVDInstanceID);
 
 		if (AChaosVDParticleActor** ExistingParticleVDInstancePtrPtr = ParticleVDInstancesByID.Find(ParticleVDInstanceID))
 		{
 			if (AChaosVDParticleActor* ExistingParticleVDInstancePtr = *ExistingParticleVDInstancePtrPtr)
 			{
-				ExistingParticleVDInstancePtr->UpdateFromRecordedData(Particle);
+				ExistingParticleVDInstancePtr->UpdateFromRecordedData(Particle, InFrameData.SimulationTransform);
 			}
 			else
 			{
@@ -72,7 +82,7 @@ void FChaosVDScene::UpdateFromRecordedStepData(const int32 SolverID, const FStri
 		}
 		else
 		{
-			if (AChaosVDParticleActor* NewParticleVDInstance = SpawnParticleFromRecordedData(Particle))
+			if (AChaosVDParticleActor* NewParticleVDInstance = SpawnParticleFromRecordedData(Particle, InFrameData))
 			{
 				FStringFormatOrderedArguments Args {SolverName, FString::FromInt(SolverID)};
 				const FName FolderPath = *FPaths::Combine(FString::Format(TEXT("Solver {0} | ID {1}"), Args), UEnum::GetDisplayValueAsText(Particle.ParticleType).ToString());
@@ -91,7 +101,7 @@ void FChaosVDScene::UpdateFromRecordedStepData(const int32 SolverID, const FStri
 	}
 
 	int32 AmountRemoved = 0;
-	for (TMap<uint32, AChaosVDParticleActor*>::TIterator RemoveIterator = ParticleVDInstancesByID.CreateIterator(); RemoveIterator; ++RemoveIterator)
+	for (TMap<int32, AChaosVDParticleActor*>::TIterator RemoveIterator = ParticleVDInstancesByID.CreateIterator(); RemoveIterator; ++RemoveIterator)
 	{
 		if (!ParticlesIDsInRecordedStepData.Contains(RemoveIterator.Key()))
 		{
@@ -109,11 +119,16 @@ void FChaosVDScene::UpdateFromRecordedStepData(const int32 SolverID, const FStri
 	OnSceneUpdated().Broadcast();
 }
 
+void FChaosVDScene::HandleNewGeometryData(const TSharedPtr<const Chaos::FImplicitObject>& GeometryData, const int32 GeometryID) const
+{
+	NewGeometryAvailableDelegate.Broadcast(GeometryData, GeometryID);
+}
+
 void FChaosVDScene::CleanUpScene()
 {
 	if (PhysicsVDWorld)
 	{
-		for (const TPair<uint32, AChaosVDParticleActor*>& ParticleVDInstanceWithID : ParticleVDInstancesByID)
+		for (const TPair<int32, AChaosVDParticleActor*>& ParticleVDInstanceWithID : ParticleVDInstancesByID)
 		{
 			PhysicsVDWorld->DestroyActor(ParticleVDInstanceWithID.Value);
 		}
@@ -124,25 +139,31 @@ void FChaosVDScene::CleanUpScene()
 	CollectGarbage(GARBAGE_COLLECTION_KEEPFLAGS);
 }
 
-AChaosVDParticleActor* FChaosVDScene::SpawnParticleFromRecordedData(const FChaosVDParticleDebugData& InParticleData) const
+AChaosVDParticleActor* FChaosVDScene::SpawnParticleFromRecordedData(const FChaosVDParticleDebugData& InParticleData, const FChaosVDSolverFrameData& InFrameData)
 {
-	// Temp code until we have geometry data
-	const UChaosVDEditorSettings* Settings = GetDefault<UChaosVDEditorSettings>();
-	UStaticMesh* DebugMesh = Settings->DebugMesh.LoadSynchronous();
+	using namespace Chaos;
 
 	FActorSpawnParameters Params;
 	Params.Name = *InParticleData.DebugName;
 	Params.NameMode = FActorSpawnParameters::ESpawnActorNameMode::Requested;
 	if (AChaosVDParticleActor* NewActor = PhysicsVDWorld->SpawnActor<AChaosVDParticleActor>(Params))
 	{
-		NewActor->UpdateFromRecordedData(InParticleData);
+		NewActor->UpdateFromRecordedData(InParticleData, InFrameData.SimulationTransform);
 
 		if (!InParticleData.DebugName.IsEmpty())
 		{
 			NewActor->SetActorLabel(InParticleData.DebugName);
 		}
 
-		NewActor->SetStaticMesh(DebugMesh);
+		NewActor->SetScene(AsShared());
+		
+		if (ensure(LoadedRecording.IsValid()))
+		{
+			if (const TSharedPtr<const Chaos::FImplicitObject>* Geometry = LoadedRecording->GetGeometryDataMap().Find(InParticleData.ImplicitObjectID))
+			{
+				NewActor->UpdateGeometryData(*Geometry);
+			}
+		}
 
 		return NewActor;
 	}
@@ -150,7 +171,7 @@ AChaosVDParticleActor* FChaosVDScene::SpawnParticleFromRecordedData(const FChaos
 	return nullptr;
 }
 
-uint32 FChaosVDScene::GetIDForRecordedParticleData(const FChaosVDParticleDebugData& InParticleData) const
+int32 FChaosVDScene::GetIDForRecordedParticleData(const FChaosVDParticleDebugData& InParticleData) const
 {
 	return InParticleData.ParticleIndex;
 }
@@ -197,5 +218,3 @@ UWorld* FChaosVDScene::CreatePhysicsVDWorld() const
 
 	return NewWorld;
 }
-
-
diff --git a/Engine/Plugins/Experimental/ChaosVD/Source/ChaosVD/Private/ChaosVDTriMeshGenerator.cpp b/Engine/Plugins/Experimental/ChaosVD/Source/ChaosVD/Private/ChaosVDTriMeshGenerator.cpp
new file mode 100644
--- /dev/null
+++ b/Engine/Plugins/Experimental/ChaosVD/Source/ChaosVD/Private/ChaosVDTriMeshGenerator.cpp
@@ -0,0 +1,38 @@
+// Copyright Epic Games, Inc. All Rights Reserved.
+
+#include "ChaosVDTriMeshGenerator.h"
+#include "Chaos/TriangleMeshImplicitObject.h"
+
+void FChaosVDTriMeshGenerator::GenerateFromTriMesh(const Chaos::FTriangleMeshImplicitObject& InTriMesh)
+{
+	const int32 NumTriangles = InTriMesh.Elements().GetNumTriangles();
+	const int32 NumVertices = InTriMesh.Particles().Size();
+
+	const int32 NormalsNum = NumTriangles * 3;
+	constexpr int32 UVsNum = 0;
+	SetBufferSizes(NumVertices, NumTriangles , UVsNum, NormalsNum );
+
+	// Fill the vertex buffer with the transformed vertices of the TriMesh Shape
+	for (int32 i = 0; i < NumVertices; i++)
+	{
+		Vertices[i] = FVector3d(UE::Math::TVector<double>(InTriMesh.Particles().X(i)));
+	}
+
+	const Chaos::FTrimeshIndexBuffer& IdxBuffer = InTriMesh.Elements();
+	if (IdxBuffer.RequiresLargeIndices())
+	{
+		ProcessTriangles(IdxBuffer.GetLargeIndexBuffer(), NumTriangles, InTriMesh);
+	}
+	else
+	{
+		ProcessTriangles(IdxBuffer.GetSmallIndexBuffer(), NumTriangles, InTriMesh);
+	}
+	
+	bIsGenerated = true;
+}
+
+UE::Geometry::FMeshShapeGenerator& FChaosVDTriMeshGenerator::Generate()
+{
+	ensureAlwaysMsgf(bIsGenerated, TEXT("You need to call FChaosVDTriMeshGenerator::GenerateFromTriMesh before calling Generate"));
+	return *this;
+}
diff --git a/Engine/Plugins/Experimental/ChaosVD/Source/ChaosVD/Private/Trace/ChaosVDTraceAnalyzer.cpp b/Engine/Plugins/Experimental/ChaosVD/Source/ChaosVD/Private/Trace/ChaosVDTraceAnalyzer.cpp
--- a/Engine/Plugins/Experimental/ChaosVD/Source/ChaosVD/Private/Trace/ChaosVDTraceAnalyzer.cpp
+++ b/Engine/Plugins/Experimental/ChaosVD/Source/ChaosVD/Private/Trace/ChaosVDTraceAnalyzer.cpp
@@ -2,8 +2,10 @@
 
 #include "Trace/ChaosVDTraceAnalyzer.h"
 
+#include "Chaos/GeometryParticlesfwd.h"
 #include "Trace/ChaosVDTraceProvider.h"
 #include "TraceServices/Model/AnalysisSession.h"
+#include "TraceServices/Model/Frames.h"
 
 void FChaosVDTraceAnalyzer::OnAnalysisBegin(const FOnAnalysisContext& Context)
 {
@@ -15,6 +17,11 @@ void FChaosVDTraceAnalyzer::OnAnalysisBegin(const FOnAnalysisContext& Context)
 	
 	Builder.RouteEvent(RouteId_ChaosVDSolverStepStart, "ChaosVDLogger", "ChaosVDSolverStepStart");
 	Builder.RouteEvent(RouteId_ChaosVDSolverStepEnd, "ChaosVDLogger", "ChaosVDSolverStepEnd");
+	
+	Builder.RouteEvent(RouteId_ChaosVDBinaryDataStart, "ChaosVDLogger", "ChaosVDBinaryDataStart");
+	Builder.RouteEvent(RouteId_ChaosVDBinaryDataContent, "ChaosVDLogger", "ChaosVDBinaryDataContent");
+	Builder.RouteEvent(RouteId_ChaosVDBinaryDataEnd, "ChaosVDLogger", "ChaosVDBinaryDataEnd");
+	Builder.RouteEvent(RouteId_ChaosVDSolverSimulationSpace, "ChaosVDLogger", "ChaosVDSolverSimulationSpace");
 
 	TraceServices::FAnalysisSessionEditScope _(Session);
 	ChaosVDTraceProvider->CreateRecordingInstanceForSession(Session.GetName());
@@ -89,6 +96,63 @@ bool FChaosVDTraceAnalyzer::OnEvent(uint16 RouteId, EStyle Style, const FOnEvent
 
 			break;
 		}
+	case RouteId_ChaosVDBinaryDataStart:
+		{
+			const int32 DataID = EventData.GetValue<int32>("DataID");
+			
+			FChaosVDBinaryDataContainer& DataContainer = ChaosVDTraceProvider->FindOrAddUnprocessedData(DataID);
+			DataContainer.bIsCompressed = EventData.GetValue<bool>("IsCompressed");
+			DataContainer.UncompressedSize = EventData.GetValue<uint32>("OriginalSize");
+			DataContainer.DataID = EventData.GetValue<int32>("DataID");
+
+			EventData.GetString("TypeName", DataContainer.TypeName);
+
+			const uint32 DataSize = EventData.GetValue<uint32>("DataSize");
+			DataContainer.RawData.Reserve(DataSize);
+
+			break;
+		}
+	case RouteId_ChaosVDBinaryDataContent:
+		{
+			const int32 DataID = EventData.GetValue<int32>("DataID");	
+
+			FChaosVDBinaryDataContainer& DataContainer = ChaosVDTraceProvider->FindOrAddUnprocessedData(DataID);
+
+			const TArrayView<const uint8> SerializedDataChunk = EventData.GetArrayView<uint8>("RawData");
+			DataContainer.RawData.Append(SerializedDataChunk.GetData(), SerializedDataChunk.Num());
+
+			break;
+		}
+	case RouteId_ChaosVDBinaryDataEnd:
+		{
+			const int32 DataID = EventData.GetValue<int32>("DataID");
+
+			ChaosVDTraceProvider->SetBinaryDataReadyToUse(DataID);
+			break;
+		}
+	case RouteId_ChaosVDSolverSimulationSpace:
+	{
+		const int32 SolverID = EventData.GetValue<int32>("SolverID");
+
+		FVector Position;
+		Position.X = EventData.GetValue<float>("PositionX");
+		Position.Y = EventData.GetValue<float>("PositionY");
+		Position.Z = EventData.GetValue<float>("PositionZ");
+
+		FQuat Rotation;
+		Rotation.X = EventData.GetValue<float>("RotationX");
+		Rotation.Y = EventData.GetValue<float>("RotationY");
+		Rotation.Z = EventData.GetValue<float>("RotationZ");
+		Rotation.W = EventData.GetValue<float>("RotationW");
+
+		// This can be null if the recording started Mid-Frame. In this case we just discard the data for now
+		if (FChaosVDSolverFrameData* FrameData = ChaosVDTraceProvider->GetLastFrame(SolverID))
+		{
+			FrameData->SimulationTransform.SetLocation(Position);
+			FrameData->SimulationTransform.SetRotation(Rotation);
+		}
+		break;
+	}
 	default:
 		break;
 	}
@@ -101,6 +165,7 @@ void FChaosVDTraceAnalyzer::ReadParticleDataFromEvent(const FEventData& InEventD
 	OutParticleData.ParticleType = static_cast<EChaosVDParticleType>(InEventData.GetValue<uint8>("ParticleType"));
 	OutParticleData.ParticleState = static_cast<EChaosVDParticleState>(InEventData.GetValue<int8>("ObjectState"));
 	OutParticleData.ParticleIndex = InEventData.GetValue<int32>("ParticleID");
+	OutParticleData.ImplicitObjectID = InEventData.GetValue<int32>("ImplicitObjectID");
 
 	FWideStringView DebugNameView;
 	InEventData.GetString("DebugName", DebugNameView);
diff --git a/Engine/Plugins/Experimental/ChaosVD/Source/ChaosVD/Private/Trace/ChaosVDTraceProvider.cpp b/Engine/Plugins/Experimental/ChaosVD/Source/ChaosVD/Private/Trace/ChaosVDTraceProvider.cpp
--- a/Engine/Plugins/Experimental/ChaosVD/Source/ChaosVD/Private/Trace/ChaosVDTraceProvider.cpp
+++ b/Engine/Plugins/Experimental/ChaosVD/Source/ChaosVD/Private/Trace/ChaosVDTraceProvider.cpp
@@ -2,7 +2,10 @@
 
 #include "Trace/ChaosVDTraceProvider.h"
 
+#include "Chaos/ImplicitObject.h"
 #include "ChaosVDRecording.h"
+#include "Compression/OodleDataCompressionUtil.h"
+#include "Serialization/MemoryReader.h"
 
 FName FChaosVDTraceProvider::ProviderName("ChaosVDProvider");
 
@@ -51,6 +54,76 @@ FChaosVDSolverFrameData* FChaosVDTraceProvider::GetLastFrame(const int32 InSolve
 	return nullptr;
 }
 
+FChaosVDBinaryDataContainer& FChaosVDTraceProvider::FindOrAddUnprocessedData(const int32 DataID)
+{
+	if (const TSharedPtr<FChaosVDBinaryDataContainer>* UnprocessedData = UnprocessedDataByID.Find(DataID))
+	{
+		check(UnprocessedData->IsValid());
+		return *UnprocessedData->Get();
+	}
+	else
+	{
+		const TSharedPtr<FChaosVDBinaryDataContainer> DataContainer = MakeShared<FChaosVDBinaryDataContainer>(DataID);
+		UnprocessedDataByID.Add(DataID, DataContainer);
+		return *DataContainer.Get();
+	}
+}
+
+void FChaosVDTraceProvider::SetBinaryDataReadyToUse(const int32 DataID)
+{
+	if (const TSharedPtr<FChaosVDBinaryDataContainer>* UnprocessedDataPtr = UnprocessedDataByID.Find(DataID))
+	{
+		const TSharedPtr<FChaosVDBinaryDataContainer> UnprocessedData = *UnprocessedDataPtr;
+		if (UnprocessedData.IsValid())
+		{
+			UnprocessedData->bIsReady = true;
+
+			BinaryDataReadyDelegate.Broadcast(*UnprocessedDataPtr);
+
+			const TArray<uint8>* RawData = nullptr;
+			TArray<uint8> UncompressedData;
+			if (UnprocessedData->bIsCompressed)
+			{
+				UncompressedData.Reserve(UnprocessedData->UncompressedSize);
+				FOodleCompressedArray::DecompressToTArray(UncompressedData, UnprocessedData->RawData);
+				RawData = &UncompressedData;
+			}
+			else
+			{
+				RawData = &UnprocessedData->RawData;
+			}
+
+			// TODO: Create a system to register external "Data handlers" with the logic on how serialize each type
+			// This should not be done here but as this is the only type we have is ok for now
+			if (UnprocessedData->TypeName == TEXT("FImplicitObject"))
+			{
+				Chaos::TSerializablePtr<Chaos::FImplicitObject> Geometry;
+
+				FMemoryReader MemeReader(*RawData);
+				Chaos::FChaosArchive Ar(MemeReader);
+
+				if (ChaosContext.IsValid())
+				{
+					Ar.SetContext(MoveTemp(ChaosContext));
+					Ar << Geometry;
+					ChaosContext = Ar.StealContext();
+				}
+				else
+				{
+					Ar << Geometry;
+					ChaosContext = Ar.StealContext();
+				}
+
+				InternalRecording->AddImplicitObject(ChaosContext->ObjToTag[Geometry.Get()], Geometry.Get());
+			}
+		}
+	}
+	else
+	{
+		ensure(false);	
+	}
+}
+
 TSharedPtr<FChaosVDRecording> FChaosVDTraceProvider::GetRecordingForSession() const
 {
 	return InternalRecording;
diff --git a/Engine/Plugins/Experimental/ChaosVD/Source/ChaosVD/Public/ChaosVDConvexMeshGenerator.h b/Engine/Plugins/Experimental/ChaosVD/Source/ChaosVD/Public/ChaosVDConvexMeshGenerator.h
new file mode 100644
--- /dev/null
+++ b/Engine/Plugins/Experimental/ChaosVD/Source/ChaosVD/Public/ChaosVDConvexMeshGenerator.h
@@ -0,0 +1,30 @@
+// Copyright Epic Games, Inc. All Rights Reserved.
+
+#pragma once
+
+#include "Generators/MeshShapeGenerator.h"
+
+namespace Chaos
+{
+	class FConvex;
+}
+;
+/**
+ * Generates an convex mesh with the smallest number of triangles possible 
+ */
+class FChaosVDConvexMeshGenerator : public UE::Geometry::FMeshShapeGenerator
+{
+public:
+
+	/** Prepares the Mesh Generator so the Generate method can generate the desired Dynamic Mesh
+	 * @param Convex shape data used to generate the dynamic mesh
+	 */
+	void GenerateFromConvex(const Chaos::FConvex& Convex);
+	
+	virtual FMeshShapeGenerator& Generate() override;
+
+private:
+
+	bool bIsGenerated = false; 
+};
+
diff --git a/Engine/Plugins/Experimental/ChaosVD/Source/ChaosVD/Public/ChaosVDGeometryBuilder.h b/Engine/Plugins/Experimental/ChaosVD/Source/ChaosVD/Public/ChaosVDGeometryBuilder.h
new file mode 100644
--- /dev/null
+++ b/Engine/Plugins/Experimental/ChaosVD/Source/ChaosVD/Public/ChaosVDGeometryBuilder.h
@@ -0,0 +1,422 @@
+// Copyright Epic Games, Inc. All Rights Reserved.
+
+#pragma once
+
+#include "ChaosVDConvexMeshGenerator.h"
+#include "ChaosVDHeightFieldMeshGenerator.h"
+#include "ChaosVDTriMeshGenerator.h"
+#include "Chaos/HeightField.h"
+#include "Chaos/ImplicitObjectScaled.h"
+#include "Chaos/ImplicitObjectType.h"
+#include "Components/DynamicMeshComponent.h"
+#include "Components/InstancedStaticMeshComponent.h"
+#include "Generators/CapsuleGenerator.h"
+#include "Generators/MinimalBoxMeshGenerator.h"
+#include "Generators/SphereGenerator.h"
+#include "UDynamicMesh.h"
+#include "UObject/GCObject.h"
+#include "UObject/UObjectGlobals.h"
+
+class AActor;
+class UDynamicMesh;
+class UDynamicMeshComponent;
+
+namespace UE
+{
+	namespace Geometry
+	{
+		class FMeshShapeGenerator;
+		class FDynamicMesh3;
+	}
+}
+
+/*
+ * Generates Dynamic mesh components and dynamic meshes based on Chaos implicit object data
+ */
+class FChaosVDGeometryBuilder : public FGCObject
+{
+public:
+
+	/** Creates Dynamic Mesh components for each object within the provided Implicit object
+	 *	@param InImplicitObject : Implicit object to process
+	 *	@param Owner Actor who will own the generated components
+	 *	@param OutMeshComponents Array containing all the generated components
+	 *	@param Index Index of the current component being processed. This is useful when this method is called recursively
+	 *	@param Transform to apply to the generated components/geometry
+	 */
+	template<typename MeshType, typename ComponentType>
+	void CreateMeshComponentsFromImplicit(const Chaos::FImplicitObject* InImplicitObject, AActor* Owner, TArray<UMeshComponent*>& OutMeshComponents, Chaos::FRigidTransform3& Transform, int32 Index = 0);
+	
+	virtual void AddReferencedObjects(FReferenceCollector& Collector) override;
+	virtual FString GetReferencerName() const override
+	{
+		return TEXT("FChaosVDScene");
+	}
+
+private:
+
+	/** Creates a Dynamic Mesh for the provided Implicit object and generator, and then caches it to be reused later
+	 * @param GeometryCacheKey Key to be used to find this geometry in the cache
+	 * @param MeshGenerator Generator class with the data and rules to create the mesh
+	 */
+	UDynamicMesh* CreateAndCacheDynamicMesh(const uint32 GeometryCacheKey, UE::Geometry::FMeshShapeGenerator& MeshGenerator);
+
+	/** Returns an already mesh for the provided implicit object if exists, otherwise returns null
+	 * @param GeometryCacheKey Key to be used to find this geometry in the cache
+	 */
+	template<typename MeshType>
+	MeshType* GetCachedMeshForImplicit(const uint32 GeometryCacheKey);
+
+	/** Creates a Dynamic Mesh for the provided Implicit object and generator, and then caches it to be reused later
+	 * @param GeometryCacheKey Key to be used to find this geometry in the cache
+	 * @param MeshGenerator Generator class with the data and rules to create the mesh
+	 */
+	UStaticMesh* CreateAndCacheStaticMesh(const uint32 GeometryCacheKey, UE::Geometry::FMeshShapeGenerator& MeshGenerator);
+
+	/**
+	 * Creates an empty DynamicMeshComponent/Static Mesh Component or Instanced mesg component and adds it to the actor
+	 * @param Owner Actor who will own the component
+	 * @param Name Name of the component. It has to be unique within the components in the owner actor
+	 * @param Transform (Optional) Transform to apply as Relative Transform in the component after its creating and attachment to the provided actor
+	 * */
+	template<typename ComponentType>
+	ComponentType* CreateMeshComponent(AActor* Owner, const FString& Name, const Chaos::FRigidTransform3& Transform) const;
+
+	/**
+	 * Creates a Mesh from the provided Implicit object geometry data
+	 * @param GeometryCacheKey Key to be used to find this geometry in the cache
+	 * @param MeshGenerator Generator class that has the data to generate the mesh
+	 * */
+	template<typename MeshType>
+	MeshType* CreateAndCacheMeshForImplicit(const uint32 GeometryCacheKey, UE::Geometry::FMeshShapeGenerator& MeshGenerator);
+
+	/* Process an FImplicitObject and returns de desired geometry type. Could be directly the shape or another version of the implicit */
+	template <bool bIsInstanced, typename GeometryType>
+	const GeometryType* GetGeometry(const Chaos::FImplicitObject* InImplicit, const bool bIsScaled, Chaos::FRigidTransform3& OutTransform) const;
+
+	/** Process an FImplicitObject and returns de desired geometry type based on the packed object flags. Could be directly the shape or another version of the implicit */
+	template<typename GeometryType>
+	const GeometryType* GetGeometryBasedOnPackedType(const Chaos::FImplicitObject* InImplicitObject, Chaos::FRigidTransform3& Transform, const Chaos::EImplicitObjectType PackedType) const;
+
+	/** Map containing already generated dynamic mesh for any given implicit object */
+	TMap<uint32, UDynamicMesh*> MeshCacheMap;
+
+	/** Map containing already generated static mesh for any given implicit object */
+	TMap<uint32, UStaticMesh*> StaticMeshCacheMap;
+};
+
+
+template <typename MeshType, typename ComponentType>
+void FChaosVDGeometryBuilder::CreateMeshComponentsFromImplicit(const Chaos::FImplicitObject* InImplicitObject, AActor* Owner, TArray<UMeshComponent*>& OutMeshComponents, Chaos::FRigidTransform3& Transform, int32 Index)
+{
+	static_assert(std::is_same_v<MeshType, UStaticMesh> || std::is_same_v<MeshType, UDynamicMesh>, "CreateMeshComponentsFromImplicit Only supports DynamicMesh and Static Mesh");
+	static_assert(std::is_same_v<ComponentType, UStaticMeshComponent> || std::is_same_v<ComponentType, UInstancedStaticMeshComponent> || std::is_same_v<MeshType, UDynamicMeshComponent>, "CreateMeshComponentsFromImplicit Only supports DynamicMeshComponent, Static MeshComponent and Instanced Static Mesh Component");
+
+	// We could have you make the Mesh type as template and infer the component type based on that, but we also want to be able to create Static Meshes with either Instanced or normal static mesh components
+	constexpr bool bHasValidCombinationForStaticMesh =  std::is_same_v<MeshType, UStaticMesh> && (std::is_same_v<ComponentType, UStaticMeshComponent> || std::is_same_v<ComponentType, UInstancedStaticMeshComponent>);
+	constexpr bool bHasValidCombinationForDynamicMesh =  std::is_same_v<MeshType, UDynamicMesh> && std::is_same_v<ComponentType, UDynamicMeshComponent>;
+	static_assert(bHasValidCombinationForStaticMesh || bHasValidCombinationForDynamicMesh , "Incorrect Component type for Mesh type. Did you use a Dynamic Mesh with a Static Mesh component type?.");
+
+	using namespace Chaos;
+
+	const EImplicitObjectType InnerType = GetInnerType(InImplicitObject->GetType());
+	const EImplicitObjectType PackedType = InImplicitObject->GetType();
+	
+	if (InnerType == ImplicitObjectType::Union)
+	{
+		const FImplicitObjectUnion* Union = InImplicitObject->template GetObject<FImplicitObjectUnion>();
+
+		for (int i = 0; i < Union->GetObjects().Num(); ++i)
+		{
+			const TUniquePtr<FImplicitObject>& UnionImplicit = Union->GetObjects()[i];
+
+			CreateMeshComponentsFromImplicit<MeshType, ComponentType>(UnionImplicit.Get(), Owner, OutMeshComponents, Transform, i);	
+		}
+
+		return;
+	}
+
+	if (InnerType ==  ImplicitObjectType::Transformed)
+	{
+		const TImplicitObjectTransformed<FReal, 3>* Transformed = InImplicitObject->template GetObject<TImplicitObjectTransformed<FReal, 3>>();
+		FRigidTransform3 TransformCopy = Transformed->GetTransform();
+		CreateMeshComponentsFromImplicit<MeshType, ComponentType>(Transformed->GetTransformedObject(), Owner, OutMeshComponents,TransformCopy, Index);
+		return;
+	}
+
+	ComponentType* MeshComponent = nullptr;
+	MeshType* Mesh = nullptr;
+	switch (InnerType)
+	{
+		case ImplicitObjectType::Sphere:
+		{
+			const TSphere<FReal, 3>* Sphere = InImplicitObject->template GetObject<TSphere<FReal, 3>>();
+
+			const FString Name = FString::Format(TEXT("{0} - {1}"), {TEXT("Sphere"), FString::FromInt(Index)});
+
+			MeshComponent = CreateMeshComponent<ComponentType>(Owner, Name, Transform);
+
+			Mesh = GetCachedMeshForImplicit<MeshType>(Sphere->GetTypeHash());
+
+			if (!Mesh)
+			{
+				UE::Geometry::FSphereGenerator SphereGen;
+				SphereGen.Radius = Sphere->GetRadius();
+				SphereGen.NumTheta = 50;
+				SphereGen.NumPhi = 50;
+				SphereGen.bPolygroupPerQuad = false;
+
+				Mesh = CreateAndCacheMeshForImplicit<MeshType>(Sphere->GetTypeHash(), SphereGen);
+			}
+
+			break;
+		}
+		case ImplicitObjectType::Box:
+		{
+			const TBox<FReal, 3>* Box = InImplicitObject->template GetObject<TBox<FReal, 3>>();
+
+			const FString Name = FString::Format(TEXT("{0} - {1}"), {TEXT("Box"), FString::FromInt(Index)});
+			MeshComponent = CreateMeshComponent<ComponentType>(Owner, Name, Transform);
+
+			Mesh = GetCachedMeshForImplicit<MeshType>(Box->GetTypeHash());
+
+			if (!Mesh)
+			{
+				UE::Geometry::FMinimalBoxMeshGenerator BoxGen;
+				UE::Geometry::FOrientedBox3d OrientedBox;
+				OrientedBox.Frame = UE::Geometry::FFrame3d(Box->Center());
+				OrientedBox.Extents = Box->Extents() * 0.5;
+				BoxGen.Box = OrientedBox;
+
+				Mesh = CreateAndCacheMeshForImplicit<MeshType>(Box->GetTypeHash(), BoxGen);
+			}
+
+			break;
+		}
+		case ImplicitObjectType::Plane:
+			break;
+		case ImplicitObjectType::Capsule:
+		{
+			const FCapsule* Capsule = InImplicitObject->template GetObject<FCapsule>();
+
+			const FString Name = FString::Format(TEXT("{0} - {1}"), {TEXT("Capsule"), FString::FromInt(Index)});
+			const FRigidTransform3 StartingTransform;
+			MeshComponent = CreateMeshComponent<ComponentType>(Owner, Name, StartingTransform);
+
+			// Re-adjust the location so the pivot is not the center of the capsule, and transform it based on the provided transform
+			const FVector FinalLocation = Transform.TransformPosition(Capsule->GetCenter() - Capsule->GetAxis() * Capsule->GetSegment().GetLength() * 0.5f);
+			const FQuat Rotation = FRotationMatrix::MakeFromZ(Capsule->GetAxis()).Rotator().Quaternion();
+
+			MeshComponent->SetRelativeRotation(Transform.GetRotation() * Rotation);
+			MeshComponent->SetRelativeLocation(FinalLocation);
+			MeshComponent->SetRelativeScale3D(Transform.GetScale3D());
+
+			Mesh = GetCachedMeshForImplicit<MeshType>(Capsule->GetTypeHash());
+
+			if (!Mesh)
+			{
+				UE::Geometry::FCapsuleGenerator CapsuleGenerator;
+				CapsuleGenerator.Radius = FMath::Max(FMathf::ZeroTolerance, Capsule->GetRadius());
+				CapsuleGenerator.SegmentLength = FMath::Max(FMathf::ZeroTolerance, Capsule->GetSegment().GetLength());
+				CapsuleGenerator.NumHemisphereArcSteps = 12;
+				CapsuleGenerator.NumCircleSteps = 12;
+				CapsuleGenerator.bPolygroupPerQuad = false;
+				
+				Mesh = CreateAndCacheMeshForImplicit<MeshType>(Capsule->GetTypeHash(), CapsuleGenerator);
+			}
+			break;
+		}
+		case ImplicitObjectType::LevelSet:
+		{
+			//TODO: Implement
+			break;
+		}
+		break;
+		case ImplicitObjectType::Convex:
+		{
+			if (const FConvex* Convex = GetGeometryBasedOnPackedType<FConvex>(InImplicitObject, Transform, PackedType))
+			{
+				const FString Name = FString::Format(TEXT("{0} - {1}"), {TEXT("Convex"), FString::FromInt(Index)});
+				MeshComponent = CreateMeshComponent<ComponentType>(Owner, Name, Transform);
+				Mesh = GetCachedMeshForImplicit<MeshType>(Convex->GetTypeHash());
+	
+				if (!Mesh)
+				{
+					FChaosVDConvexMeshGenerator ConvexMeshGen;
+					ConvexMeshGen.GenerateFromConvex(*Convex);
+					Mesh = CreateAndCacheMeshForImplicit<MeshType>(Convex->GetTypeHash() ,ConvexMeshGen);
+				}
+			}
+
+			break;
+		}
+		case ImplicitObjectType::TaperedCylinder:
+			break;
+		case ImplicitObjectType::Cylinder:
+			break;
+		case ImplicitObjectType::TriangleMesh:
+		{
+			if (const FTriangleMeshImplicitObject* TriangleMesh = GetGeometryBasedOnPackedType<FTriangleMeshImplicitObject>(InImplicitObject, Transform, PackedType))
+			{
+				const FString Name = FString::Format(TEXT("{0} - {1}"), {TEXT("Trimesh"), FString::FromInt(Index)});
+				MeshComponent = CreateMeshComponent<ComponentType>(Owner, Name, Transform);
+
+				Mesh = GetCachedMeshForImplicit<MeshType>(TriangleMesh->GetTypeHash());
+				if (!Mesh)
+				{
+					FChaosVDTriMeshGenerator TriMeshGen;
+					TriMeshGen.bReverseOrientation = true;
+					TriMeshGen.GenerateFromTriMesh(*TriangleMesh);		
+					Mesh = CreateAndCacheMeshForImplicit<MeshType>(TriangleMesh->GetTypeHash(), TriMeshGen);
+				}
+			}
+				
+			break;
+		}
+		case ImplicitObjectType::HeightField:
+		{
+			if (const FHeightField* HeightField = GetGeometryBasedOnPackedType<FHeightField>(InImplicitObject, Transform, PackedType))
+			{
+				const FString Name = FString::Format(TEXT("{0} - {1}"), {TEXT("HeightField"), FString::FromInt(Index)});
+				MeshComponent = CreateMeshComponent<ComponentType>(Owner, Name, Transform);
+				Mesh = GetCachedMeshForImplicit<MeshType>(HeightField->GetTypeHash());
+	
+				if (!Mesh)
+				{
+					FChaosVDHeightFieldMeshGenerator HeightFieldMeshGen;
+					HeightFieldMeshGen.bReverseOrientation = false;
+					HeightFieldMeshGen.GenerateFromHeightField(*HeightField);
+
+					Mesh = CreateAndCacheMeshForImplicit<MeshType>(HeightField->GetTypeHash(), HeightFieldMeshGen);
+				}
+			}
+		
+			break;
+		}
+		default:
+			break;
+		}
+
+		if (MeshComponent != nullptr && Mesh != nullptr)
+		{
+			if constexpr (std::is_same_v<ComponentType, UDynamicMeshComponent>)
+			{
+				MeshComponent->SetDynamicMesh(Mesh);
+			}
+			else if constexpr(std::is_same_v<ComponentType, UInstancedStaticMeshComponent> || std::is_same_v<ComponentType, UStaticMeshComponent>)
+			{
+				MeshComponent->SetStaticMesh(Mesh);
+			}
+
+			OutMeshComponents.Add(MeshComponent);
+		}
+}
+
+template <typename MeshType>
+MeshType* FChaosVDGeometryBuilder::GetCachedMeshForImplicit(const uint32 GeometryCacheKey)
+{
+	if constexpr (std::is_same_v<MeshType, UDynamicMesh>)
+	{
+		if (MeshType** MeshPtrPtr = MeshCacheMap.Find(GeometryCacheKey))
+		{
+			return *MeshPtrPtr;
+		}
+	}
+	else if constexpr (std::is_same_v<MeshType, UStaticMesh>)
+	{
+		if (MeshType** MeshPtrPtr = StaticMeshCacheMap.Find(GeometryCacheKey))
+		{
+			return *MeshPtrPtr;
+		}
+	}
+
+	return nullptr;
+}
+
+template <typename ComponentType>
+ComponentType* FChaosVDGeometryBuilder::CreateMeshComponent(AActor* Owner, const FString& Name, const Chaos::FRigidTransform3& Transform) const
+{
+	ComponentType* MeshComponent = NewObject<ComponentType>(Owner, *Name);
+
+	if (Owner)
+	{
+		MeshComponent->RegisterComponent();
+		MeshComponent->AttachToComponent(Owner->GetRootComponent(), FAttachmentTransformRules::SnapToTargetIncludingScale);
+		Owner->AddInstanceComponent(MeshComponent);
+	}
+
+	MeshComponent->bSelectable = true;
+
+	if constexpr (std::is_same_v<ComponentType, UDynamicMeshComponent> || std::is_same_v<ComponentType, UStaticMeshComponent>)
+	{
+		MeshComponent->SetRelativeTransform(Transform);
+	}
+	else if constexpr (std::is_same_v<ComponentType, UInstancedStaticMeshComponent>)
+	{
+		MeshComponent->AddInstance(Transform);
+	}
+
+	return MeshComponent;
+}
+
+template <typename MeshType>
+MeshType* FChaosVDGeometryBuilder::CreateAndCacheMeshForImplicit(const uint32 GeometryKey, UE::Geometry::FMeshShapeGenerator& MeshGenerator)
+{
+	if constexpr (std::is_same_v<MeshType, UDynamicMesh>)
+	{
+		return CreateAndCacheDynamicMesh(GeometryKey, MeshGenerator);
+	}
+	else if constexpr (std::is_same_v<MeshType, UStaticMesh>)
+	{
+		return CreateAndCacheStaticMesh(GeometryKey, MeshGenerator);
+	}
+
+	return nullptr;
+}
+
+template <bool bIsInstanced, typename GeometryType>
+const GeometryType* FChaosVDGeometryBuilder::GetGeometry(const Chaos::FImplicitObject* InImplicitObject, const bool bIsScaled, Chaos::FRigidTransform3& OutTransform) const
+{
+	if (bIsScaled)
+	{
+		if (const Chaos::TImplicitObjectScaled<GeometryType, bIsInstanced>* ImplicitScaled = InImplicitObject->template GetObject<Chaos::TImplicitObjectScaled<GeometryType, bIsInstanced>>())
+		{
+			OutTransform.SetScale3D(ImplicitScaled->GetScale());
+			return ImplicitScaled->GetUnscaledObject()->template GetObject<GeometryType>();
+		}
+	}
+	else
+	{
+		if (bIsInstanced)
+		{
+			const Chaos::TImplicitObjectInstanced<GeometryType>* ImplicitInstanced = InImplicitObject->template GetObject<Chaos::TImplicitObjectInstanced<GeometryType>>();
+			return ImplicitInstanced->GetInnerObject()->template GetObject<GeometryType>();
+		}
+		else
+		{
+			return InImplicitObject->template GetObject<GeometryType>();
+		}
+	}
+	
+	return nullptr;
+}
+
+template <typename GeometryType>
+const GeometryType* FChaosVDGeometryBuilder::GetGeometryBasedOnPackedType(const Chaos::FImplicitObject* InImplicitObject, Chaos::FRigidTransform3& Transform,  const Chaos::EImplicitObjectType PackedType) const
+{
+	using namespace Chaos;
+
+	const bool bIsInstanced = IsInstanced(PackedType);
+	const bool bIsScaled = IsScaled(PackedType);
+
+	if (bIsInstanced)
+	{
+		return GetGeometry<true,GeometryType>(InImplicitObject, bIsScaled, Transform);
+	}
+	else
+	{
+		return GetGeometry<false, GeometryType>(InImplicitObject, bIsScaled, Transform);
+	}
+
+	return nullptr;
+}
diff --git a/Engine/Plugins/Experimental/ChaosVD/Source/ChaosVD/Public/ChaosVDHeightfieldMeshGenerator.h b/Engine/Plugins/Experimental/ChaosVD/Source/ChaosVD/Public/ChaosVDHeightfieldMeshGenerator.h
new file mode 100644
--- /dev/null
+++ b/Engine/Plugins/Experimental/ChaosVD/Source/ChaosVD/Public/ChaosVDHeightfieldMeshGenerator.h
@@ -0,0 +1,30 @@
+// Copyright Epic Games, Inc. All Rights Reserved.
+
+#pragma once
+#include "Chaos/Core.h"
+#include "Generators/MeshShapeGenerator.h"
+
+namespace Chaos
+{
+	class FHeightField;
+}
+
+/**
+ * Generates an dynamic mesh using the data stored in a FHeightField implicit object
+ */
+class FChaosVDHeightFieldMeshGenerator : public UE::Geometry::FMeshShapeGenerator
+{
+public:
+
+	/** Prepares the Mesh Generator so the Generate method can generate the desired Dynamic Mesh
+	* @param InHeightField Implicit object used to generate the dynamic mesh
+	*/
+	void GenerateFromHeightField(const Chaos::FHeightField& InHeightField);
+
+	virtual FMeshShapeGenerator& Generate() override;
+
+private:
+	void AppendTriangle(int32& OutTriIndex, int& OutCurrentNormalIndex, int32 PolygonIndex, const Chaos::TVec2<Chaos::FReal>& InCellCoordinates, const UE::Geometry::FIndex3i& InTriangle, const Chaos::FHeightField& InHeightField);
+
+	bool bIsGenerated = false;
+};
diff --git a/Engine/Plugins/Experimental/ChaosVD/Source/ChaosVD/Public/ChaosVDParticleActor.h b/Engine/Plugins/Experimental/ChaosVD/Source/ChaosVD/Public/ChaosVDParticleActor.h
--- a/Engine/Plugins/Experimental/ChaosVD/Source/ChaosVD/Public/ChaosVDParticleActor.h
+++ b/Engine/Plugins/Experimental/ChaosVD/Source/ChaosVD/Public/ChaosVDParticleActor.h
@@ -2,32 +2,43 @@
 
 #pragma once
 #include "ChaosVDRecording.h"
+#include "Chaos/Core.h"
 #include "GameFramework/Actor.h"
 
 #include "ChaosVDParticleActor.generated.h"
 
+class FChaosVDScene;
 struct FChaosVDParticleDebugData;
 class USceneComponent;
 class UStaticMeshComponent;
 class UStaticMesh;
 
 /** Actor used to represent a Chaos Particle in the Visual Debugger's world */
-UCLASS(HideCategories=("Transform"))
+UCLASS()
 class AChaosVDParticleActor : public AActor
 {
 	GENERATED_BODY()
 public:
 	AChaosVDParticleActor(const FObjectInitializer& ObjectInitializer);
 
-	void UpdateFromRecordedData(const FChaosVDParticleDebugData& RecordedData);
+	void UpdateFromRecordedData(const FChaosVDParticleDebugData& RecordedData, const Chaos::FRigidTransform3& SimulationTransform);
+	
+	const FChaosVDParticleDebugData& GetDebugData() const { return RecordedDebugData; }
 
-	void SetStaticMesh(UStaticMesh* NewStaticMesh);
+	void UpdateGeometryData(const TSharedPtr<const Chaos::FImplicitObject>& ImplicitObject);
 
-protected:
+	void SetScene(const TSharedPtr<FChaosVDScene>& InScene);
+
+	virtual void BeginDestroy() override;
 
-	UPROPERTY(EditDefaultsOnly, Category= Debug)
-	TObjectPtr<UStaticMeshComponent> StaticMeshComponent;
+protected:
 
 	UPROPERTY(VisibleAnywhere, Category="Chaos Visual Debugger Data")
 	FChaosVDParticleDebugData RecordedDebugData;
+
+	bool bIsGeometryDataLoaded = false;
+
+	TWeakPtr<FChaosVDScene> OwningScene;
+
+	FDelegateHandle GeometryUpdatedDelegate;
 };
diff --git a/Engine/Plugins/Experimental/ChaosVD/Source/ChaosVD/Public/ChaosVDScene.h b/Engine/Plugins/Experimental/ChaosVD/Source/ChaosVD/Public/ChaosVDScene.h
--- a/Engine/Plugins/Experimental/ChaosVD/Source/ChaosVD/Public/ChaosVDScene.h
+++ b/Engine/Plugins/Experimental/ChaosVD/Source/ChaosVD/Public/ChaosVDScene.h
@@ -8,6 +8,7 @@
 #include "UObject/ObjectMacros.h"
 #include "UObject/GCObject.h"
 
+class FChaosVDGeometryBuilder;
 class AChaosVDParticleActor;
 class FReferenceCollector;
 class UWorld;
@@ -15,9 +16,11 @@ class UWorld;
 DECLARE_MULTICAST_DELEGATE(FChaosVDSceneUpdatedDelegate)
 
 /** Recreates a UWorld from a recorded Chaos VD Frame */
-class FChaosVDScene : public FGCObject
+class FChaosVDScene : public FGCObject , public TSharedFromThis<FChaosVDScene>
 {
 public:
+	FChaosVDScene();
+	virtual ~FChaosVDScene() override;
 
 	void Initialize();
 	void DeInitialize();
@@ -32,7 +35,9 @@ public:
 	FChaosVDSceneUpdatedDelegate& OnSceneUpdated() { return SceneUpdatedDelegate; }
 
 	/** Updates, Adds and Remove actors to match the provided Step Data */
-	void UpdateFromRecordedStepData(const int32 SolverID, const FString& SolverName, const FChaosVDStepData& InRecordedStepData);
+	void UpdateFromRecordedStepData(const int32 SolverID, const FString& SolverName, const FChaosVDStepData& InRecordedStepData, const FChaosVDSolverFrameData& InFrameData);
+
+	void HandleNewGeometryData(const TSharedPtr<const Chaos::FImplicitObject>&, const int32 GeometryID) const;
 
 	/** Deletes all actors of the Scene and underlying UWorld */
 	void CleanUpScene();
@@ -42,24 +47,34 @@ public:
 
 	bool IsInitialized() const { return  bIsInitialized; }
 
+	const TSharedPtr<FChaosVDGeometryBuilder>& GetGeometryGenerator() { return  GeometryGenerator; }
+
+	FChaosVDGeometryDataLoaded& OnNewGeometryAvailable(){ return NewGeometryAvailableDelegate; }
+
+	TSharedPtr<FChaosVDRecording> LoadedRecording;
+
 private:
 
 	/** Creates an ChaosVDParticle actor for the Provided recorded Particle Data */
-	AChaosVDParticleActor* SpawnParticleFromRecordedData(const FChaosVDParticleDebugData& InParticleData) const;
+	AChaosVDParticleActor* SpawnParticleFromRecordedData(const FChaosVDParticleDebugData& InParticleData, const FChaosVDSolverFrameData& InFrameData);
 
 	/** Returns the ID used to track this recorded particle data */
-	uint32 GetIDForRecordedParticleData(const FChaosVDParticleDebugData& InParticleData) const;
+	int32 GetIDForRecordedParticleData(const FChaosVDParticleDebugData& InParticleData) const;
 
 	/** Creates the instance of the World which will be used the recorded data*/
 	UWorld* CreatePhysicsVDWorld() const;
 
 	/** Map of ID-ChaosVDParticle Actor. Used to keep track of actor instances and be able to modify them as needed*/
-	TMap<uint32, AChaosVDParticleActor*> ParticleVDInstancesByID;
+	TMap<int32, AChaosVDParticleActor*> ParticleVDInstancesByID;
 
 	/** UWorld instance used to represent the recorded debug data */
 	UWorld* PhysicsVDWorld = nullptr;
 
 	FChaosVDSceneUpdatedDelegate SceneUpdatedDelegate;
 
+	TSharedPtr<FChaosVDGeometryBuilder> GeometryGenerator;
+
+	FChaosVDGeometryDataLoaded NewGeometryAvailableDelegate;
+
 	bool bIsInitialized = false;
 };
diff --git a/Engine/Plugins/Experimental/ChaosVD/Source/ChaosVD/Public/ChaosVDTriMeshGenerator.h b/Engine/Plugins/Experimental/ChaosVD/Source/ChaosVD/Public/ChaosVDTriMeshGenerator.h
new file mode 100644
--- /dev/null
+++ b/Engine/Plugins/Experimental/ChaosVD/Source/ChaosVD/Public/ChaosVDTriMeshGenerator.h
@@ -0,0 +1,52 @@
+// Copyright Epic Games, Inc. All Rights Reserved.
+
+#pragma once
+#include "Chaos/TriangleMeshImplicitObject.h"
+#include "Generators/MeshShapeGenerator.h"
+
+namespace Chaos
+{
+	class FTriangleMeshImplicitObject;
+}
+
+/** Generates a Dynamic mesh based on a Triangle Mesh Implicit object*/
+class FChaosVDTriMeshGenerator : public UE::Geometry::FMeshShapeGenerator
+{
+
+public:
+	/** Triangle Mesh Implicit object used as data source to generate the dynamic mesh */
+	void GenerateFromTriMesh(const Chaos::FTriangleMeshImplicitObject& InTriMesh);
+
+	virtual FMeshShapeGenerator& Generate() override;
+
+private:
+
+	template<typename BufferIndexType>
+	void ProcessTriangles(const TArray<BufferIndexType>& InTriangles, const int32 NumTriangles, const Chaos::FTriangleMeshImplicitObject& InTriMesh);
+
+	bool bIsGenerated = false;;
+};
+
+template <typename BufferIndexType>
+void FChaosVDTriMeshGenerator::ProcessTriangles(const TArray<BufferIndexType>& InTriangles, const int32 NumTriangles, const Chaos::FTriangleMeshImplicitObject& InTriMesh)
+{
+	int CurrentNormalIndex = 0;
+	for (int32 TriangleIndex = 0; TriangleIndex < NumTriangles; ++TriangleIndex)
+	{
+		UE::Geometry::FIndex3i Triangle(InTriangles[TriangleIndex][0],InTriangles[TriangleIndex][1],InTriangles[TriangleIndex][2]);
+
+		// Create a Normal entry per triangle vertex.
+		int NormalIndices[3];
+		for (int32 LocalVertexIndex = 0; LocalVertexIndex < 3; ++LocalVertexIndex)
+		{
+			// Use the normal of the face for all it's vertices
+			Normals[CurrentNormalIndex] = UE::Math::TVector<float>(InTriMesh.GetFaceNormal(TriangleIndex));
+			NormalIndices[LocalVertexIndex] = CurrentNormalIndex;
+			CurrentNormalIndex++;
+		}
+
+		SetTriangle(TriangleIndex, MoveTemp(Triangle));
+		SetTrianglePolygon(TriangleIndex, TriangleIndex);
+		SetTriangleNormals(TriangleIndex, NormalIndices[0], NormalIndices[1], NormalIndices[2]);
+	}
+}
diff --git a/Engine/Plugins/Experimental/ChaosVD/Source/ChaosVD/Public/Trace/ChaosVDTraceAnalyzer.h b/Engine/Plugins/Experimental/ChaosVD/Source/ChaosVD/Public/Trace/ChaosVDTraceAnalyzer.h
--- a/Engine/Plugins/Experimental/ChaosVD/Source/ChaosVD/Public/Trace/ChaosVDTraceAnalyzer.h
+++ b/Engine/Plugins/Experimental/ChaosVD/Source/ChaosVD/Public/Trace/ChaosVDTraceAnalyzer.h
@@ -40,7 +40,11 @@ private:
 		RouteId_ChaosVDSolverStepStart,
 		RouteId_ChaosVDSolverStepEnd,
 		RouteId_ChaosVDSolverFrameStart,
-		RouteId_ChaosVDSolverFrameEnd
+		RouteId_ChaosVDSolverFrameEnd,
+		RouteId_ChaosVDBinaryDataStart,
+		RouteId_ChaosVDBinaryDataContent,
+		RouteId_ChaosVDBinaryDataEnd,
+		RouteId_ChaosVDSolverSimulationSpace
 	};
 
 	void ReadParticleDataFromEvent(const FEventData& InEventData, FChaosVDParticleDebugData& OutParticleData);
diff --git a/Engine/Plugins/Experimental/ChaosVD/Source/ChaosVD/Public/Trace/ChaosVDTraceProvider.h b/Engine/Plugins/Experimental/ChaosVD/Source/ChaosVD/Public/Trace/ChaosVDTraceProvider.h
--- a/Engine/Plugins/Experimental/ChaosVD/Source/ChaosVD/Public/Trace/ChaosVDTraceProvider.h
+++ b/Engine/Plugins/Experimental/ChaosVD/Source/ChaosVD/Public/Trace/ChaosVDTraceProvider.h
@@ -3,11 +3,44 @@
 #pragma once
 #include "TraceServices/Model/AnalysisSession.h"
 #include "Templates/SharedPointer.h"
+#include "Chaos/ParticleHandleFwd.h"
+#include "Delegates/Delegate.h"
+#include "Containers/UnrealString.h"
+#include "Containers/Array.h"
 
 class FChaosVDEngine;
 struct FChaosVDSolverFrameData;
 struct FChaosVDRecording;
 
+namespace Chaos
+{
+	class FChaosArchiveContext;
+}
+
+struct FChaosVDBinaryDataContainer
+{
+	FChaosVDBinaryDataContainer(const int32 InDataID)
+		: DataID(InDataID)
+	{
+	}
+
+	int32 DataID;
+	bool bIsReady = false;
+	bool bIsCompressed = false;
+	uint32 UncompressedSize = 0;
+	FString TypeName;
+	TArray<uint8> RawData;
+};
+
+struct FChaosVDTraceSessionData
+{
+	TSharedPtr<FChaosVDRecording> InternalRecordingsMap;
+
+	TMap<int32, TSharedPtr<FChaosVDBinaryDataContainer>> UnprocessedDataByID;
+};
+
+DECLARE_MULTICAST_DELEGATE_OneParam(FOnBinaryDataReady, TSharedPtr<FChaosVDBinaryDataContainer>)
+
 /** Provider class for Chaos VD trace recordings.
  * It stores and handles rebuilt recorded frame data from Trace events
  * dispatched by the Chaos VD Trace analyzer
@@ -26,10 +59,22 @@ public:
 	FChaosVDSolverFrameData* GetFrame(const int32 InSolverGUID, const int32 FrameNumber) const;
 	FChaosVDSolverFrameData* GetLastFrame(const int32 InSolverGUID) const;
 
+	FChaosVDBinaryDataContainer& FindOrAddUnprocessedData(const int32 DataID);
+
+	void SetBinaryDataReadyToUse(const int32 DataID);
+
 	TSharedPtr<FChaosVDRecording> GetRecordingForSession() const;
+	
+	FOnBinaryDataReady& OnBinaryDataReady() { return BinaryDataReadyDelegate; }
 
 private:
 	TraceServices::IAnalysisSession& Session;
 
 	TSharedPtr<FChaosVDRecording> InternalRecording;
+
+	TMap<int32, TSharedPtr<FChaosVDBinaryDataContainer>> UnprocessedDataByID;
+	
+	FOnBinaryDataReady BinaryDataReadyDelegate;
+
+	TUniquePtr<Chaos::FChaosArchiveContext> ChaosContext;
 };
diff --git a/Engine/Source/Runtime/Engine/Private/PhysicsEngine/ImmediatePhysics/ImmediatePhysicsChaos/ImmediatePhysicsSimulation_Chaos.cpp b/Engine/Source/Runtime/Engine/Private/PhysicsEngine/ImmediatePhysics/ImmediatePhysicsChaos/ImmediatePhysicsSimulation_Chaos.cpp
--- a/Engine/Source/Runtime/Engine/Private/PhysicsEngine/ImmediatePhysics/ImmediatePhysicsChaos/ImmediatePhysicsSimulation_Chaos.cpp
+++ b/Engine/Source/Runtime/Engine/Private/PhysicsEngine/ImmediatePhysics/ImmediatePhysicsChaos/ImmediatePhysicsSimulation_Chaos.cpp
@@ -799,6 +799,7 @@ namespace ImmediatePhysics_Chaos
 	{
 		SCOPE_CYCLE_COUNTER(STAT_ImmediateSimulate_Chaos);
 		CVD_SCOPE_TRACE_SOLVER_FRAME(FSimulation, *this);
+		CVD_TRACE_SOLVER_SIMULATION_SPACE(Implementation->SimulationSpace.Transform);
 
 		using namespace Chaos;
 
@@ -1004,7 +1005,6 @@ namespace ImmediatePhysics_Chaos
 		INC_DWORD_STAT_BY(STAT_ImmediateSimulate_ChaosCounter_NumContacts, Implementation->Collisions.NumConstraints());
 		INC_DWORD_STAT_BY(STAT_ImmediateSimulate_ChaosCounter_NumJoints, Implementation->Joints.NumConstraints());
 	}
-
 	void FSimulation::DebugDrawStaticParticles()
 	{
 #if CHAOS_DEBUG_DRAW
diff --git a/Engine/Source/Runtime/Experimental/Chaos/Private/Chaos/HeightField.cpp b/Engine/Source/Runtime/Experimental/Chaos/Private/Chaos/HeightField.cpp
--- a/Engine/Source/Runtime/Experimental/Chaos/Private/Chaos/HeightField.cpp
+++ b/Engine/Source/Runtime/Experimental/Chaos/Private/Chaos/HeightField.cpp
@@ -848,6 +848,18 @@ namespace Chaos
 		return TNumericLimits<FReal>::Max();
 	}
 
+
+	FVec3 FHeightField::GetPointScaled(int32 InIndex) const
+	{
+		if (CHAOS_ENSURE(InIndex >= 0 && InIndex < GeomData.Heights.Num()))
+		{
+			return GeomData.GetPointScaled(InIndex);
+		}
+		
+		return FVec3::ZeroVector;
+	}
+
+
 	FReal Chaos::FHeightField::GetHeight(int32 InX, int32 InY) const
 	{
 		const int32 Index = InY * GeomData.NumCols + InX;
diff --git a/Engine/Source/Runtime/Experimental/Chaos/Private/ChaosVisualDebugger/ChaosVisualDebuggerTrace.cpp b/Engine/Source/Runtime/Experimental/Chaos/Private/ChaosVisualDebugger/ChaosVisualDebuggerTrace.cpp
--- a/Engine/Source/Runtime/Experimental/Chaos/Private/ChaosVisualDebugger/ChaosVisualDebuggerTrace.cpp
+++ b/Engine/Source/Runtime/Experimental/Chaos/Private/ChaosVisualDebugger/ChaosVisualDebuggerTrace.cpp
@@ -5,6 +5,12 @@
 #if CHAOS_VISUAL_DEBUGGER_ENABLED
 #include "Chaos/ParticleHandle.h"
 #include "Chaos/Framework/PhysicsSolverBase.h"
+#include "Chaos/ImplicitObject.h"
+#include "ChaosVDRuntimeModule.h"
+#include "Compression/OodleDataCompressionUtil.h"
+#include "HAL/CriticalSection.h"
+#include "Misc/ScopeRWLock.h"
+#include "Serialization/MemoryWriter.h"
 
 UE_TRACE_EVENT_DEFINE(ChaosVDLogger, ChaosVDSolverFrameStart)
 
@@ -14,21 +20,48 @@ UE_TRACE_CHANNEL_DEFINE(ChaosVDChannel);
 UE_TRACE_EVENT_DEFINE(ChaosVDLogger, ChaosVDParticle)
 UE_TRACE_EVENT_DEFINE(ChaosVDLogger, ChaosVDSolverStepStart)
 UE_TRACE_EVENT_DEFINE(ChaosVDLogger, ChaosVDSolverStepEnd)
+UE_TRACE_EVENT_DEFINE(ChaosVDLogger, ChaosVDBinaryDataStart)
+UE_TRACE_EVENT_DEFINE(ChaosVDLogger, ChaosVDBinaryDataContent)
+UE_TRACE_EVENT_DEFINE(ChaosVDLogger, ChaosVDBinaryDataEnd)
+UE_TRACE_EVENT_DEFINE(ChaosVDLogger, ChaosVDSolverSimulationSpace)
+
+static FAutoConsoleVariable CVarChaosVDCompressBinaryData(
+	TEXT("p.Chaos.VD.CompressBinaryData"),
+	false,
+	TEXT("If true, serialized binary data will be compressed using Oodle on the fly before being traced"));
+
+static FAutoConsoleVariable CVarChaosVDCompressionMode(
+	TEXT("p.Chaos.VD.CompressionMode"),
+	2,
+	TEXT("Oodle compression mode to use, 4 is by default which equsals to ECompressionLevel::VeryFast"));
+
+struct FChaosVDGeometryTraceContext
+{
+	FRWLock RWLock;
+	TUniquePtr<Chaos::FChaosArchiveContext> ChaosContext;
+};
+
+static FChaosVDGeometryTraceContext GeometryTracerObject = FChaosVDGeometryTraceContext();
 
 void FChaosVisualDebuggerTrace::TraceParticle(const Chaos::FGeometryParticleHandle* ParticleHandle)
 {
+	if (!UE_TRACE_CHANNELEXPR_IS_ENABLED(ChaosVDChannel))
+	{
+		return;
+	}
+
 	const FChaosVDContext* CVDContextData = FChaosVDThreadContext::Get().GetCurrentContext();
 	if (!ensure(CVDContextData))
 	{
 		return;
 	}
 
-	TraceParticle(ParticleHandle, *CVDContextData);
+	TraceParticle(const_cast<Chaos::FGeometryParticleHandle*>(ParticleHandle), *CVDContextData);
 }
 
-void FChaosVisualDebuggerTrace::TraceParticle(const Chaos::FGeometryParticleHandle* ParticleHandle, const FChaosVDContext& ContextData)
+void FChaosVisualDebuggerTrace::TraceParticle(Chaos::FGeometryParticleHandle* ParticleHandle, const FChaosVDContext& ContextData)
 {
-	if (!TRACE_PRIVATE_CHANNELEXPR_IS_ENABLED(ChaosVDChannel))
+	if (!UE_TRACE_CHANNELEXPR_IS_ENABLED(ChaosVDChannel))
 	{
 		return;
 	}
@@ -39,15 +72,16 @@ void FChaosVisualDebuggerTrace::TraceParticle(const Chaos::FGeometryParticleHand
 		return;
 	}
 
-	FString TempDebugName;
+	static FString DefaultName = TEXT("NONAME");
+	FStringView ParticleNameView(DefaultName);
+
 #if CHAOS_DEBUG_NAME
-	if (const TSharedPtr<FString, ESPMode::ThreadSafe> DebugNamePtr = ParticleHandle->DebugName())
+	if (const TSharedPtr<FString, ESPMode::ThreadSafe>& DebugNamePtr = ParticleHandle->DebugName())
 	{
-		TempDebugName = *DebugNamePtr.Get();
+		ParticleNameView = FStringView(*DebugNamePtr.Get());
 	}
-#else
-	TempDebugName = TEXT("NOT AVAILABLE - COMPILED OUT");
 #endif
+	const int32 GeometryID = TraceImplicitObject(ParticleHandle->Geometry());
 	
 	UE_TRACE_LOG(ChaosVDLogger, ChaosVDParticle, ChaosVDChannel)
 		<< ChaosVDParticle.SolverID(ContextData.Id)
@@ -55,19 +89,22 @@ void FChaosVisualDebuggerTrace::TraceParticle(const Chaos::FGeometryParticleHand
 
 		<< ChaosVDParticle.ParticleID(ParticleHandle->UniqueIdx().Idx)
 		<< ChaosVDParticle.ParticleType(static_cast<uint8>(ParticleHandle->Type))
-		<< ChaosVDParticle.DebugName(*TempDebugName, TempDebugName.Len())
+		<< ChaosVDParticle.DebugName(ParticleNameView.GetData(), ParticleNameView.Len())
 
 		<< CVD_TRACE_VECTOR_ON_EVENT(ChaosVDParticle, Position, ParticleHandle->X())
 		<< CVD_TRACE_ROTATOR_ON_EVENT(ChaosVDParticle, Rotation, ParticleHandle->R())
 		<< CVD_TRACE_VECTOR_ON_EVENT(ChaosVDParticle, Velocity, Chaos::FConstGenericParticleHandle(ParticleHandle)->V())
 		<< CVD_TRACE_VECTOR_ON_EVENT(ChaosVDParticle, AngularVelocity, Chaos::FConstGenericParticleHandle(ParticleHandle)->W())
 
+		<< ChaosVDParticle.ImplicitObjectID(GeometryID)
+
 		<< ChaosVDParticle.ObjectState(static_cast<int8>(ParticleHandle->ObjectState()));
+	
 }
 
 void FChaosVisualDebuggerTrace::TraceParticles(const Chaos::TGeometryParticleHandles<Chaos::FReal, 3>& ParticleHandles)
 {
-	if (!TRACE_PRIVATE_CHANNELEXPR_IS_ENABLED(ChaosVDChannel))
+	if (!UE_TRACE_CHANNELEXPR_IS_ENABLED(ChaosVDChannel))
 	{
 		return;
 	}
@@ -117,6 +154,11 @@ void FChaosVisualDebuggerTrace::TraceSolverFrameEnd(const FChaosVDContext& Conte
 
 void FChaosVisualDebuggerTrace::TraceSolverStepStart()
 {
+	if (!UE_TRACE_CHANNELEXPR_IS_ENABLED(ChaosVDChannel))
+	{
+		return;
+	}
+
 	const FChaosVDContext* CVDContextData = FChaosVDThreadContext::Get().GetCurrentContext();
 	if (!ensure(CVDContextData))
 	{
@@ -130,6 +172,11 @@ void FChaosVisualDebuggerTrace::TraceSolverStepStart()
 
 void FChaosVisualDebuggerTrace::TraceSolverStepEnd()
 {
+	if (!UE_TRACE_CHANNELEXPR_IS_ENABLED(ChaosVDChannel))
+	{
+		return;
+	}
+
 	const FChaosVDContext* CVDContextData = FChaosVDThreadContext::Get().GetCurrentContext();
 	if (!ensure(CVDContextData))
 	{
@@ -140,4 +187,150 @@ void FChaosVisualDebuggerTrace::TraceSolverStepEnd()
 		<< ChaosVDSolverStepEnd.Cycle(FPlatformTime::Cycles64())
 		<< ChaosVDSolverStepEnd.SolverID(CVDContextData->Id);
 }
+
+void FChaosVisualDebuggerTrace::TraceSolverSimulationSpace(const Chaos::FRigidTransform3& Transform)
+{
+	if (!UE_TRACE_CHANNELEXPR_IS_ENABLED(ChaosVDChannel))
+	{
+		return;
+	}
+
+	const FChaosVDContext* CVDContextData = FChaosVDThreadContext::Get().GetCurrentContext();
+	if (!ensure(CVDContextData))
+	{
+		return;
+	}
+	
+	UE_TRACE_LOG(ChaosVDLogger, ChaosVDSolverSimulationSpace, ChaosVDChannel)
+		<< ChaosVDSolverSimulationSpace.Cycle(FPlatformTime::Cycles64())
+		<< ChaosVDSolverSimulationSpace.SolverID(CVDContextData->Id)
+		<< CVD_TRACE_VECTOR_ON_EVENT(ChaosVDSolverSimulationSpace, Position, Transform.GetLocation())
+		<< CVD_TRACE_ROTATOR_ON_EVENT(ChaosVDSolverSimulationSpace, Rotation, Transform.GetRotation());
+}
+
+void FChaosVisualDebuggerTrace::TraceBinaryData(const TArray<uint8>& InData, const FString& TypeName)
+{
+	if (!UE_TRACE_CHANNELEXPR_IS_ENABLED(ChaosVDChannel))
+	{
+		return;
+	}
+
+	//TODO: This might overflow
+	static FThreadSafeCounter LastDataID;
+
+	const int32 DataID = LastDataID.Increment();
+
+	ensure(DataID < TNumericLimits<int32>::Max());
+
+	const TArray<uint8>* DataToTrace = &InData;
+
+	// Handle Compression if enabled
+	const bool bIsCompressed = CVarChaosVDCompressBinaryData->GetBool();
+	TArray<uint8> CompressedData;
+	if (bIsCompressed)
+	{
+		CompressedData.Reserve(CompressedData.Num());
+		FOodleCompressedArray::CompressTArray(CompressedData, InData, FOodleDataCompression::ECompressor::Kraken,
+			static_cast<FOodleDataCompression::ECompressionLevel>(CVarChaosVDCompressionMode->GetInt()));
+
+		DataToTrace = &CompressedData;
+	}
+
+	const uint32 DataSize = static_cast<uint32>(DataToTrace->Num());
+	constexpr uint32 MaxChunkSize = TNumericLimits<uint16>::Max();
+	const uint32 ChunkNum = (DataSize + MaxChunkSize - 1) / MaxChunkSize;
+
+	UE_TRACE_LOG(ChaosVDLogger, ChaosVDBinaryDataStart, ChaosVDChannel)
+		<< ChaosVDBinaryDataStart.Cycle(FPlatformTime::Cycles64())
+		<< ChaosVDBinaryDataStart.TypeName(*TypeName, TypeName.Len())
+		<< ChaosVDBinaryDataStart.DataID(DataID)
+		<< ChaosVDBinaryDataStart.DataSize(DataSize)
+		<< ChaosVDBinaryDataStart.OriginalSize(InData.Num())
+		<< ChaosVDBinaryDataStart.IsCompressed(bIsCompressed);
+
+	uint32 RemainingSize = DataSize;
+	for (uint32 Index = 0; Index < ChunkNum; ++Index)
+	{
+		const uint16 Size = static_cast<uint16>(FMath::Min(RemainingSize, MaxChunkSize));
+		const uint8* ChunkData = DataToTrace->GetData() + MaxChunkSize * Index;
+
+		UE_TRACE_LOG(ChaosVDLogger, ChaosVDBinaryDataContent, ChaosVDChannel)
+			<< ChaosVDBinaryDataContent.Cycle(FPlatformTime::Cycles64())
+			<< ChaosVDBinaryDataContent.DataID(DataID)
+			<< ChaosVDBinaryDataContent.RawData(ChunkData, Size);
+
+		RemainingSize -= Size;
+	}
+
+	UE_TRACE_LOG(ChaosVDLogger, ChaosVDBinaryDataEnd, ChaosVDChannel)
+		<< ChaosVDBinaryDataEnd.Cycle(FPlatformTime::Cycles64())
+		<< ChaosVDBinaryDataEnd.DataID(DataID);
+
+	ensure(RemainingSize == 0);
+}
+
+int32 FChaosVisualDebuggerTrace::TraceImplicitObject(Chaos::TSerializablePtr<Chaos::FImplicitObject> Geometry)
+{
+	if (!UE_TRACE_CHANNELEXPR_IS_ENABLED(ChaosVDChannel))
+	{
+		return INDEX_NONE;
+	}
+
+	//TODO: Find a better place to do this
+	if (!FChaosVDRuntimeModule::OnRecordingStop().IsBound())
+	{
+		FChaosVDRuntimeModule::OnRecordingStop().BindStatic(&FChaosVisualDebuggerTrace::ResetGeometryTracerContext);
+	}
+	
+	if (!FChaosVDRuntimeModule::OnRecordingStarted().IsBound())
+	{
+		FChaosVDRuntimeModule::OnRecordingStarted().BindStatic(&FChaosVisualDebuggerTrace::ResetGeometryTracerContext);
+	}
+
+	{
+		FReadScopeLock Lock(GeometryTracerObject.RWLock);
+		if (GeometryTracerObject.ChaosContext.IsValid())
+		{
+			int32 SerializedObjectPtrTag = GeometryTracerObject.ChaosContext.Get()->GetObjectTag(Geometry.Get());
+			if (SerializedObjectPtrTag != INDEX_NONE)
+			{
+				// Geometry Data is already serialized, just return the tag
+				return SerializedObjectPtrTag;
+			}
+		}
+	}
+
+	// TODO: Change this so it is not allocated each time
+	// We need to take into account we could be serializing on multiple threads
+	constexpr uint32 MaxTraceChunkSize = TNumericLimits<uint16>::Max();
+	TArray<uint8> RawData;
+	RawData.Reserve(MaxTraceChunkSize);
+	
+	int32 SerializedObjectPtrTag = INDEX_NONE;
+
+	{
+		FMemoryWriter MemWriterAr(RawData);
+		Chaos::FChaosArchive Ar(MemWriterAr);
+		
+		FWriteScopeLock WriteLock(GeometryTracerObject.RWLock);
+		if (GeometryTracerObject.ChaosContext.IsValid())
+		{
+			Ar.SetContext(MoveTemp(GeometryTracerObject.ChaosContext));
+		}
+
+		Ar << Geometry;
+		GeometryTracerObject.ChaosContext = Ar.StealContext();
+		SerializedObjectPtrTag = GeometryTracerObject.ChaosContext.Get()->GetObjectTag(Geometry.Get());
+	}
+
+	TraceBinaryData(RawData, TEXT("FImplicitObject"));
+
+	return SerializedObjectPtrTag;
+}
+
+void FChaosVisualDebuggerTrace::ResetGeometryTracerContext()
+{
+	GeometryTracerObject.ChaosContext.Release();
+}
+
 #endif
diff --git a/Engine/Source/Runtime/Experimental/Chaos/Public/Chaos/ChaosArchive.h b/Engine/Source/Runtime/Experimental/Chaos/Public/Chaos/ChaosArchive.h
--- a/Engine/Source/Runtime/Experimental/Chaos/Public/Chaos/ChaosArchive.h
+++ b/Engine/Source/Runtime/Experimental/Chaos/Public/Chaos/ChaosArchive.h
@@ -111,6 +111,15 @@ public:
 		}
 	}
 
+	int32 GetObjectTag(const void* ObjectPtr) const
+	{
+		if (const int32* SerializedObjectPtrTag = ObjToTag.Find(ObjectPtr))
+		{
+			return *SerializedObjectPtrTag;
+		}
+		return INDEX_NONE;
+	}
+
 private:
 	class FSharedPtrHolder
 	{
@@ -158,12 +167,24 @@ public:
 			int32 Tag;
 			InnerArchive << Tag;
 
+			if (Tag < 0)
+			{
+				InnerArchive.SetCriticalError();
+				return;
+			}
+
 			const int32 SlotsNeeded = Tag + 1 - Context->TagToObject.Num();
 			if (SlotsNeeded > 0)
 			{
 				Context->TagToObject.AddZeroed(SlotsNeeded);
 			}
 
+			if (!Context->TagToObject.IsValidIndex(Tag))
+			{
+				InnerArchive.SetCriticalError();
+				return;
+			}
+
 			if (Context->TagToObject[Tag])
 			{
 				Obj.SetFromRawLowLevel((const T*)(Context->TagToObject[Tag]));
@@ -172,6 +193,7 @@ public:
 			{
 				StaticSerialize(Obj);
 				Context->TagToObject[Tag] = (void*)Obj.Get();
+				Context->ObjToTag.Add((void*)Obj.Get(), Tag);
 			}
 		}
 		else if (InnerArchive.IsSaving() || InnerArchive.IsCountingMemory())
diff --git a/Engine/Source/Runtime/Experimental/Chaos/Public/Chaos/HeightField.h b/Engine/Source/Runtime/Experimental/Chaos/Public/Chaos/HeightField.h
--- a/Engine/Source/Runtime/Experimental/Chaos/Public/Chaos/HeightField.h
+++ b/Engine/Source/Runtime/Experimental/Chaos/Public/Chaos/HeightField.h
@@ -40,6 +40,7 @@ namespace Chaos
 		void EditHeights(TArrayView<FReal> InHeights, int32 InBeginRow, int32 InBeginCol, int32 InNumRows, int32 InNumCols);
 		void EditHeights(TArrayView<const uint16> InHeights, int32 InBeginRow, int32 InBeginCol, int32 InNumRows, int32 InNumCols);
 		FReal GetHeight(int32 InIndex) const;
+		FVec3 GetPointScaled(int32 InIndex) const;
 		FReal GetHeight(int32 InX, int32 InY) const;
 		uint8 GetMaterialIndex(int32 InIndex) const;
 		uint8 GetMaterialIndex(int32 InX, int32 InY) const;
diff --git a/Engine/Source/Runtime/Experimental/Chaos/Public/Chaos/ImplicitObjectScaled.h b/Engine/Source/Runtime/Experimental/Chaos/Public/Chaos/ImplicitObjectScaled.h
--- a/Engine/Source/Runtime/Experimental/Chaos/Public/Chaos/ImplicitObjectScaled.h
+++ b/Engine/Source/Runtime/Experimental/Chaos/Public/Chaos/ImplicitObjectScaled.h
@@ -994,7 +994,6 @@ public:
 		FImplicitObject::SerializeImp(Ar);
 		Ar << MObject << MScale << MInvScale;
 		TBox<FReal,d>::SerializeAsAABB(Ar, MLocalBoundingBox);
-		ensure(OuterMargin == 0);	//not supported: do we care?
 
 		Ar.UsingCustomVersion(FExternalPhysicsCustomObjectVersion::GUID);
 		if (Ar.CustomVer(FExternalPhysicsCustomObjectVersion::GUID) < FExternalPhysicsCustomObjectVersion::ScaledGeometryIsConcrete)
diff --git a/Engine/Source/Runtime/Experimental/Chaos/Public/ChaosVisualDebugger/ChaosVisualDebuggerTrace.h b/Engine/Source/Runtime/Experimental/Chaos/Public/ChaosVisualDebugger/ChaosVisualDebuggerTrace.h
--- a/Engine/Source/Runtime/Experimental/Chaos/Public/ChaosVisualDebugger/ChaosVisualDebuggerTrace.h
+++ b/Engine/Source/Runtime/Experimental/Chaos/Public/ChaosVisualDebugger/ChaosVisualDebuggerTrace.h
@@ -1,8 +1,11 @@
 // Copyright Epic Games, Inc. All Rights Reserved.
 
 #pragma once
+#include "Chaos/Core.h"
 #include "Chaos/ParticleHandleFwd.h"
+#include "Chaos/Serializable.h"
 #include "Chaos/Vector.h"
+#include "Containers/Set.h"
 
 #ifndef CHAOS_VISUAL_DEBUGGER_ENABLED
 	#define CHAOS_VISUAL_DEBUGGER_ENABLED (WITH_CHAOS_VISUAL_DEBUGGER && UE_TRACE_ENABLED)
@@ -41,6 +44,14 @@
 	#ifndef CVD_SCOPE_TRACE_SOLVER_STEP
 		#define CVD_SCOPE_TRACE_SOLVER_STEP()
 	#endif
+
+	#ifndef CVD_TRACE_BINARY_DATA
+		#define CVD_TRACE_BINARY_DATA(InData, TypeName)
+	#endif
+
+#ifndef CVD_TRACE_SOLVER_SIMULATION_SPACE
+		#define CVD_TRACE_SOLVER_SIMULATION_SPACE(InSimulationSpace)
+#endif
 #else
 
 #include "Trace/Trace.h"
@@ -100,6 +111,7 @@ UE_TRACE_EVENT_BEGIN_EXTERN(ChaosVDLogger, ChaosVDParticle)
 	CVD_DEFINE_TRACE_VECTOR(Chaos::FReal, AngularVelocity)
 	UE_TRACE_EVENT_FIELD(int8, ObjectState)
 	UE_TRACE_EVENT_FIELD(UE::Trace::WideString, DebugName)
+	UE_TRACE_EVENT_FIELD(int32, ImplicitObjectID)
 UE_TRACE_EVENT_END()
 
 UE_TRACE_EVENT_BEGIN_EXTERN(ChaosVDLogger, ChaosVDSolverStepStart)
@@ -114,6 +126,33 @@ UE_TRACE_EVENT_BEGIN_EXTERN(ChaosVDLogger, ChaosVDSolverStepEnd)
 	UE_TRACE_EVENT_FIELD(uint32, StepNumber)
 UE_TRACE_EVENT_END()
 
+UE_TRACE_EVENT_BEGIN_EXTERN(ChaosVDLogger, ChaosVDBinaryDataStart, NoSync)
+	UE_TRACE_EVENT_FIELD(uint64, Cycle)
+	UE_TRACE_EVENT_FIELD(UE::Trace::WideString, TypeName)
+	UE_TRACE_EVENT_FIELD(int32, DataID)
+	UE_TRACE_EVENT_FIELD(uint32, DataSize)
+	UE_TRACE_EVENT_FIELD(uint32, OriginalSize)
+	UE_TRACE_EVENT_FIELD(bool, IsCompressed)
+UE_TRACE_EVENT_END()
+
+UE_TRACE_EVENT_BEGIN_EXTERN(ChaosVDLogger, ChaosVDBinaryDataContent, NoSync)
+	UE_TRACE_EVENT_FIELD(uint64, Cycle)
+	UE_TRACE_EVENT_FIELD(int32, DataID)
+	UE_TRACE_EVENT_FIELD(uint8[], RawData)
+UE_TRACE_EVENT_END()
+
+UE_TRACE_EVENT_BEGIN_EXTERN(ChaosVDLogger, ChaosVDBinaryDataEnd, NoSync)
+	UE_TRACE_EVENT_FIELD(uint64, Cycle)
+	UE_TRACE_EVENT_FIELD(int32, DataID)
+UE_TRACE_EVENT_END()
+
+UE_TRACE_EVENT_BEGIN_EXTERN(ChaosVDLogger, ChaosVDSolverSimulationSpace)
+	UE_TRACE_EVENT_FIELD(int32, SolverID)
+	UE_TRACE_EVENT_FIELD(uint64, Cycle)
+	CVD_DEFINE_TRACE_VECTOR(Chaos::FReal, Position)
+	CVD_DEFINE_TRACE_ROTATOR(Chaos::FReal, Rotation)
+UE_TRACE_EVENT_END()
+
 #ifndef CVD_TRACE_PARTICLE
 	#define CVD_TRACE_PARTICLE(ParticleHandle) \
 		FChaosVisualDebuggerTrace::TraceParticle(ParticleHandle);
@@ -158,10 +197,22 @@ UE_TRACE_EVENT_END()
 		FChaosVDScopeSolverStep ScopeSolverStep;
 #endif
 
+#ifndef CVD_TRACE_BINARY_DATA
+	#define CVD_TRACE_BINARY_DATA(InData, TypeName) \
+	FChaosVisualDebuggerTrace::TraceBinaryData(InData, TypeName);
+#endif
+
+#ifndef CVD_TRACE_SOLVER_SIMULATION_SPACE
+	#define CVD_TRACE_SOLVER_SIMULATION_SPACE(InSimulationSpace) \
+	FChaosVisualDebuggerTrace::TraceSolverSimulationSpace(InSimulationSpace);
+#endif
+
 struct FChaosVDContext;
 
 namespace Chaos
 {
+	class FChaosArchiveContext;
+	class FImplicitObject;
 	class FPhysicsSolverBase;
 	template <typename T, int d>
 	class TGeometryParticleHandles;
@@ -173,18 +224,26 @@ class FChaosVisualDebuggerTrace
 public:
 
 	static CHAOS_API void TraceParticle(const Chaos::FGeometryParticleHandle* ParticleHandle);
-	static CHAOS_API void TraceParticle(const Chaos::FGeometryParticleHandle* ParticleHandle, const FChaosVDContext& ContextData);
+	static CHAOS_API void TraceParticle(Chaos::FGeometryParticleHandle* ParticleHandle, const FChaosVDContext& ContextData);
 	static CHAOS_API void TraceParticles(const Chaos::TGeometryParticleHandles<Chaos::FReal, 3>& ParticleHandles);
 	static CHAOS_API void TraceSolverFrameStart(const FChaosVDContext& ContextData, const FString& InDebugName);
 	static CHAOS_API void TraceSolverFrameEnd(const FChaosVDContext& ContextData);
 	static CHAOS_API void TraceSolverStepStart();
 	static CHAOS_API void TraceSolverStepEnd();
+	static CHAOS_API void TraceSolverSimulationSpace(const Chaos::FRigidTransform3& Transform);
+	
+	static CHAOS_API void TraceBinaryData(const TArray<uint8>& InData, const FString& TypeName);
+	
+	static CHAOS_API int32 TraceImplicitObject(Chaos::TSerializablePtr<Chaos::FImplicitObject> Geometry);
 
 	template<typename T>
 	static void GetCVDContext(T& ObjectWithContext, FChaosVDContext& OutCVDContext);
 
 	template<typename T>
 	static FString GetDebugName(T& ObjectWithContext);
+
+private:
+	static void ResetGeometryTracerContext();
 };
 
 template <typename T>
diff --git a/Engine/Source/Runtime/Experimental/ChaosVDData/ChaosVDData.build.cs b/Engine/Source/Runtime/Experimental/ChaosVDData/ChaosVDData.build.cs
new file mode 100644
--- /dev/null
+++ b/Engine/Source/Runtime/Experimental/ChaosVDData/ChaosVDData.build.cs
@@ -0,0 +1,21 @@
+// Copyright Epic Games, Inc. All Rights Reserved.
+
+using System.Linq;
+using System.Collections.Generic;
+
+namespace UnrealBuildTool.Rules
+{
+	public class ChaosVDData : ModuleRules
+	{
+		public ChaosVDData(ReadOnlyTargetRules Target) : base(Target)
+		{
+			PublicDependencyModuleNames.AddRange(
+				new string[] {
+					"Core",
+					"CoreUObject",
+					"Chaos"
+				}
+			);
+		}
+	}
+}
\ No newline at end of file
diff --git a/Engine/Source/Runtime/Experimental/ChaosVDData/Private/ChaosVDDataModule.cpp b/Engine/Source/Runtime/Experimental/ChaosVDData/Private/ChaosVDDataModule.cpp
new file mode 100644
--- /dev/null
+++ b/Engine/Source/Runtime/Experimental/ChaosVDData/Private/ChaosVDDataModule.cpp
@@ -0,0 +1,6 @@
+// Copyright Epic Games, Inc. All Rights Reserved.
+
+#include "ChaosVDDataModule.h"
+#include "Modules/ModuleManager.h"
+
+IMPLEMENT_MODULE(FChaosVDDataModule, ChaosVDData);
diff --git a/Engine/Source/Runtime/Experimental/ChaosVDData/Private/ChaosVDRecording.cpp b/Engine/Source/Runtime/Experimental/ChaosVDData/Private/ChaosVDRecording.cpp
new file mode 100644
--- /dev/null
+++ b/Engine/Source/Runtime/Experimental/ChaosVDData/Private/ChaosVDRecording.cpp
@@ -0,0 +1,78 @@
+// Copyright Epic Games, Inc. All Rights Reserved.
+
+#include "ChaosVDRecording.h"
+#include "Chaos/ImplicitObject.h"
+
+int32 FChaosVDRecording::GetAvailableFramesNumber(const int32 SolverID) const
+{
+	if (const TArray<FChaosVDSolverFrameData>* SolverData = RecordedFramesDataPerSolver.Find(SolverID))
+	{
+		return SolverData->Num();
+	}
+
+	return INDEX_NONE;
+}
+
+FChaosVDSolverFrameData* FChaosVDRecording::GetFrameForSolver(const int32 SolverID, const int32 FrameNumber)
+{
+	if (TArray<FChaosVDSolverFrameData>* SolverFrames = RecordedFramesDataPerSolver.Find(SolverID))
+	{
+		//TODO: Find a safer way of do this. If someone stores this ptr bad things will happen
+		return SolverFrames->IsValidIndex(FrameNumber) ? &(*SolverFrames)[FrameNumber] : nullptr;
+	}
+
+	return nullptr;
+}
+
+void FChaosVDRecording::AddFrameForSolver(const int32 SolverID, FChaosVDSolverFrameData&& InFrameData)
+{
+	if (TArray<FChaosVDSolverFrameData>* SolverFrames = RecordedFramesDataPerSolver.Find(SolverID))
+	{
+		SolverFrames->Add(MoveTemp(InFrameData));
+	}
+	else
+	{	
+		RecordedFramesDataPerSolver.Add(SolverID, { MoveTemp(InFrameData) });
+	}
+
+	OnRecordingUpdated().Broadcast();
+}
+
+EChaosVDFrameLoadState FChaosVDRecording::GetFrameState(const int32 SolverID, const int32 FrameNumber)
+{
+	const TArray<FChaosVDSolverFrameData>* SolverFrames = RecordedFramesDataPerSolver.Find(SolverID);
+	if (SolverFrames && SolverFrames->IsValidIndex(FrameNumber))
+	{
+		return EChaosVDFrameLoadState::Loaded;
+	}
+	//TODO: Implement "Buffering" state handling
+	// I already implemented this on the version where we didn't use trace and CVD files were streamed from disk
+	// But removed the code for this CL - Leaving this getter here as it is used by the controller and I also left part
+	// of the implementation to handle different states there
+
+	return EChaosVDFrameLoadState::Unknown;
+}
+
+void FChaosVDRecording::AddImplicitObject(const int32 ID, const TSharedPtr<Chaos::FImplicitObject>& InImplicitObject)
+{
+	if (ensure(!ImplicitObjects.Contains(ID)))
+	{
+		AddImplicitObject_Internal(ID, InImplicitObject);
+	}
+}
+
+void FChaosVDRecording::AddImplicitObject(const int32 ID, const Chaos::FImplicitObject* InImplicitObject)
+{
+	if (ensure(!ImplicitObjects.Contains(ID)))
+	{
+		// Only take ownership after we know we will add it to the map
+		const TSharedPtr<const Chaos::FImplicitObject> SharedImplicit(InImplicitObject);
+		AddImplicitObject_Internal(ID, SharedImplicit);
+	}
+}
+
+void FChaosVDRecording::AddImplicitObject_Internal(const int32 ID, const TSharedPtr<const Chaos::FImplicitObject>& InImplicitObject)
+{
+	ImplicitObjects.Add(ID, InImplicitObject);
+	GeometryDataLoaded.Broadcast(InImplicitObject, ID);
+}
diff --git a/Engine/Source/Runtime/Experimental/ChaosVDData/Public/ChaosVDDataModule.h b/Engine/Source/Runtime/Experimental/ChaosVDData/Public/ChaosVDDataModule.h
new file mode 100644
--- /dev/null
+++ b/Engine/Source/Runtime/Experimental/ChaosVDData/Public/ChaosVDDataModule.h
@@ -0,0 +1,9 @@
+// Copyright Epic Games, Inc. All Rights Reserved.
+
+#pragma once
+
+#include "Modules/ModuleInterface.h"
+
+class CHAOSVDDATA_API FChaosVDDataModule : public IModuleInterface
+{
+};
diff --git a/Engine/Source/Runtime/Experimental/ChaosVDData/Public/ChaosVDRecording.h b/Engine/Source/Runtime/Experimental/ChaosVDData/Public/ChaosVDRecording.h
new file mode 100644
--- /dev/null
+++ b/Engine/Source/Runtime/Experimental/ChaosVDData/Public/ChaosVDRecording.h
@@ -0,0 +1,148 @@
+// Copyright Epic Games, Inc. All Rights Reserved.
+
+#pragma once
+
+#include "Chaos/Core.h"
+#include "Chaos/Serializable.h"
+#include "Chaos/ShapeInstanceFwd.h"
+#include "UObject/ObjectMacros.h"
+#include "Containers/UnrealString.h"
+#include "ChaosVDRecording.generated.h"
+
+namespace Chaos
+{
+	class FImplicitObject;
+}
+
+DECLARE_MULTICAST_DELEGATE(FChaosVDRecordingUpdated)
+DECLARE_MULTICAST_DELEGATE_TwoParams(FChaosVDGeometryDataLoaded, const TSharedPtr<const Chaos::FImplicitObject>&, const int32 GeometryID)
+
+UENUM()
+enum class EChaosVDParticleType : uint8
+{
+	Static,
+	Kinematic,
+	Rigid,
+	Clustered,
+	StaticMesh,
+	SkeletalMesh,
+	GeometryCollection,
+	Unknown
+};
+
+UENUM()
+enum class EChaosVDParticleState : int8
+{
+	Uninitialized = 0,
+	Sleeping = 1,
+	Kinematic = 2,
+	Static = 3,
+	Dynamic = 4,
+
+	Count
+};
+
+UENUM()
+enum class EChaosVDSolverType : uint32
+{
+	Rigid
+};
+
+USTRUCT()
+struct FChaosVDParticleDebugData
+{
+	GENERATED_BODY()
+
+	UPROPERTY(VisibleAnywhere, Category= "Chaos Visual Debugger Data")
+	EChaosVDParticleType ParticleType;
+	UPROPERTY(VisibleAnywhere, Category= "Chaos Visual Debugger Data")
+	FString DebugName;
+	UPROPERTY(VisibleAnywhere, Category= "Chaos Visual Debugger Data")
+	int32 ParticleIndex;
+	UPROPERTY(VisibleAnywhere, Category= "Chaos Visual Debugger Data")
+	FVector Position;
+	UPROPERTY(VisibleAnywhere, Category= "Chaos Visual Debugger Data")
+	FQuat Rotation;
+	UPROPERTY(VisibleAnywhere, Category= "Chaos Visual Debugger Data")
+	FVector Velocity;
+	UPROPERTY(VisibleAnywhere, Category= "Chaos Visual Debugger Data")
+	FVector AngularVelocity;
+	UPROPERTY(VisibleAnywhere, Category= "Chaos Visual Debugger Data")
+	EChaosVDParticleState ParticleState;
+
+	int32 ImplicitObjectID = INDEX_NONE;
+};
+
+struct FChaosVDStepData
+{
+	TArray<FChaosVDParticleDebugData> RecordedParticles;
+};
+
+struct CHAOSVDDATA_API FChaosVDSolverFrameData
+{
+	FString DebugName;
+	int32 SolverID;
+	Chaos::FRigidTransform3 SimulationTransform;
+	TArray<FChaosVDStepData, TInlineAllocator<16>> SolverSteps;
+};
+
+enum class EChaosVDFrameLoadState
+{
+	Unloaded,
+	Loaded,
+	Buffering,
+	Unknown
+};
+
+/**
+ * Struct that represents a recorded Physics simulation.
+ * It is currently populated while analyzing a Trace session
+ */
+struct CHAOSVDDATA_API FChaosVDRecording
+{
+	/** Returns the current available recorded frames per recorded solver */
+	int32 GetAvailableFramesNumber(const int32 SolverID) const;
+	/** Returns the current available recorded solvers number */
+	int32 GetAvailableSolversNumber() const { return RecordedFramesDataPerSolver.Num();}
+	
+	const TMap<int32, TArray<FChaosVDSolverFrameData>>& GetAvailableSolvers() const { return RecordedFramesDataPerSolver; }
+
+	/** Returns a ptr to the recorded frame data for a specific solver -  Do not store as it is a reference to the element in the array */
+	FChaosVDSolverFrameData* GetFrameForSolver(const int32 SolverID, const int32 FrameNumber);
+
+	/** Adds a Frame Data for a specific Solver ID. Creates a solver entry if it does not exist */
+	void AddFrameForSolver(const int32 SolverID, FChaosVDSolverFrameData&& InFrameData);
+
+	/** Returns the current frame state of a frame used to determine if it is ready for use*/
+	EChaosVDFrameLoadState GetFrameState(const int32 SolverID, const int32 FrameNumber);
+
+	/** Called each time the recording changes - Mainly when a new frame is added from the Trace analysis */
+	FChaosVDRecordingUpdated& OnRecordingUpdated() { return RecordingUpdatedDelegate; };
+
+	FChaosVDGeometryDataLoaded& OnGeometryDataLoaded() { return GeometryDataLoaded; };
+
+	const TMap<int32, TSharedPtr<const Chaos::FImplicitObject>>& GetGeometryDataMap() const { return ImplicitObjects; };
+
+	/** Adds a shared Implicit Object to the recording */
+	void AddImplicitObject(const int32 ID, const TSharedPtr<Chaos::FImplicitObject>& InImplicitObject);
+
+	/** Session name of the trace session used to re-build this recording */
+	FString SessionName;
+
+protected:
+
+	/** Adds an Implicit Object to the recording and takes ownership of it */
+	void AddImplicitObject(const int32 ID, const Chaos::FImplicitObject* InImplicitObject);
+	
+	void AddImplicitObject_Internal(const int32 ID, const TSharedPtr<const Chaos::FImplicitObject>& InImplicitObject);
+	
+	TMap<int32, EChaosVDFrameLoadState> AvailableFramesState;
+	TMap<int32, TArray<FChaosVDSolverFrameData>> RecordedFramesDataPerSolver;
+	FChaosVDRecordingUpdated RecordingUpdatedDelegate;
+	FChaosVDGeometryDataLoaded GeometryDataLoaded;
+
+	/** Id to Ptr map of all shared geometry data required to visualize */
+	TMap<int32, TSharedPtr<const Chaos::FImplicitObject>> ImplicitObjects;
+
+	friend class FChaosVDTraceProvider;
+};
diff --git a/Engine/Source/Runtime/Experimental/ChaosVisualDebugger/Private/ChaosVDRecording.cpp b/Engine/Source/Runtime/Experimental/ChaosVisualDebugger/Private/ChaosVDRecording.cpp
deleted file mode 100644
diff --git a/Engine/Source/Runtime/Experimental/ChaosVisualDebugger/Private/ChaosVDRuntimeModule.cpp b/Engine/Source/Runtime/Experimental/ChaosVisualDebugger/Private/ChaosVDRuntimeModule.cpp
--- a/Engine/Source/Runtime/Experimental/ChaosVisualDebugger/Private/ChaosVDRuntimeModule.cpp
+++ b/Engine/Source/Runtime/Experimental/ChaosVisualDebugger/Private/ChaosVDRuntimeModule.cpp
@@ -2,7 +2,6 @@
 
 #include "ChaosVDRuntimeModule.h"
 
-#include "ChaosVDRecording.h"
 #include "Containers/Array.h"
 #include "HAL/IConsoleManager.h"
 #include "Misc/FileHelper.h"
@@ -32,6 +31,9 @@ FAutoConsoleCommand StopVDStartRecordingCommand(
 	})
 );
 
+FRecordingStateChangedDelegate FChaosVDRuntimeModule::RecordingStartedDelegate = FRecordingStateChangedDelegate();
+FRecordingStateChangedDelegate FChaosVDRuntimeModule::RecordingStopDelegate = FRecordingStateChangedDelegate();
+
 FChaosVDRuntimeModule& FChaosVDRuntimeModule::Get()
 {
 	return FModuleManager::Get().LoadModuleChecked<FChaosVDRuntimeModule>(TEXT("ChaosVDRuntime"));
@@ -39,6 +41,11 @@ FChaosVDRuntimeModule& FChaosVDRuntimeModule::Get()
 
 void FChaosVDRuntimeModule::StartupModule()
 {
+#if UE_TRACE_ENABLED
+	// Make sure it is off until we support auto trace
+	UE::Trace::ToggleChannel(TEXT("ChaosVDChannel"), false);
+#endif
+
 	FTraceAuxiliary::OnTraceStopped.AddRaw(this, &FChaosVDRuntimeModule::HandleTraceStopRequest);
 }
 
@@ -77,6 +84,8 @@ void FChaosVDRuntimeModule::StartRecording(const TArray<FString>& Args)
 		return;
 	}
 
+	OnRecordingStarted().ExecuteIfBound();
+
 #if UE_TRACE_ENABLED
 
 	// Other tools could bee using trace
@@ -128,6 +137,7 @@ void FChaosVDRuntimeModule::StopRecording()
 	{
 		return;
 	}
+
 #if UE_TRACE_ENABLED
 	
 
@@ -137,6 +147,8 @@ void FChaosVDRuntimeModule::StopRecording()
 	StopTrace();
 #endif
 
+	OnRecordingStop().ExecuteIfBound();
+
 	bIsRecording = false;
 }
 
diff --git a/Engine/Source/Runtime/Experimental/ChaosVisualDebugger/Public/ChaosVDRecording.h b/Engine/Source/Runtime/Experimental/ChaosVisualDebugger/Public/ChaosVDRecording.h
deleted file mode 100644
diff --git a/Engine/Source/Runtime/Experimental/ChaosVisualDebugger/Public/ChaosVDRuntimeModule.h b/Engine/Source/Runtime/Experimental/ChaosVisualDebugger/Public/ChaosVDRuntimeModule.h
--- a/Engine/Source/Runtime/Experimental/ChaosVisualDebugger/Public/ChaosVDRuntimeModule.h
+++ b/Engine/Source/Runtime/Experimental/ChaosVisualDebugger/Public/ChaosVDRuntimeModule.h
@@ -8,6 +8,8 @@
 
 struct FChaosVDRecording;
 
+DECLARE_DELEGATE(FRecordingStateChangedDelegate)
+
 class CHAOSVDRUNTIME_API FChaosVDRuntimeModule : public IModuleInterface
 {
 public:
@@ -30,6 +32,9 @@ public:
 	/** Returns a unique ID used to be used to identify CVD (Chaos Visual Debugger) data */
 	int32 GenerateUniqueID();
 
+	static FRecordingStateChangedDelegate& OnRecordingStarted() { return RecordingStartedDelegate; };
+	static FRecordingStateChangedDelegate& OnRecordingStop() { return RecordingStopDelegate;};
+
 private:
 
 	/** Stops the current Trace session */
@@ -45,5 +50,8 @@ private:
 	bool bIsRecording = false;
 	bool bRequestedStop = false;
 
+	static FRecordingStateChangedDelegate RecordingStartedDelegate;
+	static FRecordingStateChangedDelegate RecordingStopDelegate;
+
 	FThreadSafeCounter LastGeneratedID;
 };
