Fix CIS issue in ChaosVDGeometryBuilder

#rnx

diff --git a/Engine/Plugins/Experimental/ChaosVD/Source/ChaosVD/Private/ChaosVDGeometryBuilder.cpp b/Engine/Plugins/Experimental/ChaosVD/Source/ChaosVD/Private/ChaosVDGeometryBuilder.cpp
--- a/Engine/Plugins/Experimental/ChaosVD/Source/ChaosVD/Private/ChaosVDGeometryBuilder.cpp
+++ b/Engine/Plugins/Experimental/ChaosVD/Source/ChaosVD/Private/ChaosVDGeometryBuilder.cpp
@@ -2,7 +2,6 @@
 
 #include "ChaosVDGeometryBuilder.h"
 
-#include "ChaosVDHeightFieldMeshGenerator.h"
 #include "DynamicMeshToMeshDescription.h"
 #include "Engine/StaticMesh.h"
 #include "StaticMeshAttributes.h"
diff --git a/Engine/Plugins/Experimental/ChaosVD/Source/ChaosVD/Public/ChaosVDGeometryBuilder.h b/Engine/Plugins/Experimental/ChaosVD/Source/ChaosVD/Public/ChaosVDGeometryBuilder.h
--- a/Engine/Plugins/Experimental/ChaosVD/Source/ChaosVD/Public/ChaosVDGeometryBuilder.h
+++ b/Engine/Plugins/Experimental/ChaosVD/Source/ChaosVD/Public/ChaosVDGeometryBuilder.h
@@ -3,7 +3,7 @@
 #pragma once
 
 #include "ChaosVDConvexMeshGenerator.h"
-#include "ChaosVDHeightFieldMeshGenerator.h"
+#include "ChaosVDHeightfieldMeshGenerator.h"
 #include "ChaosVDTriMeshGenerator.h"
 #include "Chaos/HeightField.h"
 #include "Chaos/ImplicitObjectScaled.h"
