Fix CIS ambiguous symbol ChaosVDGeometryBuilder.h

#rnx

diff --git a/Engine/Plugins/Experimental/ChaosVD/Source/ChaosVD/Public/ChaosVDGeometryBuilder.h b/Engine/Plugins/Experimental/ChaosVD/Source/ChaosVD/Public/ChaosVDGeometryBuilder.h
--- a/Engine/Plugins/Experimental/ChaosVD/Source/ChaosVD/Public/ChaosVDGeometryBuilder.h
+++ b/Engine/Plugins/Experimental/ChaosVD/Source/ChaosVD/Public/ChaosVDGeometryBuilder.h
@@ -150,7 +150,7 @@ void FChaosVDGeometryBuilder::CreateMeshComponentsFromImplicit(const Chaos::FImp
 	{
 		case ImplicitObjectType::Sphere:
 		{
-			const TSphere<FReal, 3>* Sphere = InImplicitObject->template GetObject<TSphere<FReal, 3>>();
+			const Chaos::TSphere<FReal, 3>* Sphere = InImplicitObject->template GetObject<Chaos::TSphere<FReal, 3>>();
 
 			const FString Name = FString::Format(TEXT("{0} - {1}"), {TEXT("Sphere"), FString::FromInt(Index)});
 
@@ -173,7 +173,7 @@ void FChaosVDGeometryBuilder::CreateMeshComponentsFromImplicit(const Chaos::FImp
 		}
 		case ImplicitObjectType::Box:
 		{
-			const TBox<FReal, 3>* Box = InImplicitObject->template GetObject<TBox<FReal, 3>>();
+			const Chaos::TBox<FReal, 3>* Box = InImplicitObject->template GetObject<Chaos::TBox<FReal, 3>>();
 
 			const FString Name = FString::Format(TEXT("{0} - {1}"), {TEXT("Box"), FString::FromInt(Index)});
 			MeshComponent = CreateMeshComponent<ComponentType>(Owner, Name, Transform);
