Fix non-unity CIS ChaosVDTraceProvider.h

#rnx

diff --git a/Engine/Plugins/Experimental/ChaosVD/Source/ChaosVD/Public/Trace/ChaosVDTraceProvider.h b/Engine/Plugins/Experimental/ChaosVD/Source/ChaosVD/Public/Trace/ChaosVDTraceProvider.h
--- a/Engine/Plugins/Experimental/ChaosVD/Source/ChaosVD/Public/Trace/ChaosVDTraceProvider.h
+++ b/Engine/Plugins/Experimental/ChaosVD/Source/ChaosVD/Public/Trace/ChaosVDTraceProvider.h
@@ -8,15 +8,12 @@
 #include "Containers/UnrealString.h"
 #include "Containers/Array.h"
 
+#include "Chaos/ChaosArchive.h"
+
 class FChaosVDEngine;
 struct FChaosVDSolverFrameData;
 struct FChaosVDRecording;
 
-namespace Chaos
-{
-	class FChaosArchiveContext;
-}
-
 struct FChaosVDBinaryDataContainer
 {
 	FChaosVDBinaryDataContainer(const int32 InDataID)
