Fix a capitalization issue in ChaosVDHeightfieldmeshgenerator

#rnx

diff --git a/Engine/Plugins/Experimental/ChaosVD/Source/ChaosVD/Private/ChaosVDHeightfieldMeshGenerator.cpp b/Engine/Plugins/Experimental/ChaosVD/Source/ChaosVD/Private/ChaosVDHeightfieldMeshGenerator.cpp
--- a/Engine/Plugins/Experimental/ChaosVD/Source/ChaosVD/Private/ChaosVDHeightfieldMeshGenerator.cpp
+++ b/Engine/Plugins/Experimental/ChaosVD/Source/ChaosVD/Private/ChaosVDHeightfieldMeshGenerator.cpp
@@ -1,6 +1,6 @@
 // Copyright Epic Games, Inc. All Rights Reserved.
 
-#include "ChaosVDHeightFieldMeshGenerator.h"
+#include "ChaosVDHeightfieldMeshGenerator.h"
 
 #include "Chaos/HeightField.h"
 
