[Chaos VD] Fix memory leaks on tool shutdown.

Before this change, objects were being held up in memory because the ChaosVDTraceManager which is owned by the Plugins's Module was keeping a reference to it, even when the tool was closed. These objects don't need to stay alive when the tool windows are closed

diff --git a/Engine/Plugins/Experimental/ChaosVD/Source/ChaosVD/Private/ChaosVDEngine.cpp b/Engine/Plugins/Experimental/ChaosVD/Source/ChaosVD/Private/ChaosVDEngine.cpp
--- a/Engine/Plugins/Experimental/ChaosVD/Source/ChaosVD/Private/ChaosVDEngine.cpp
+++ b/Engine/Plugins/Experimental/ChaosVD/Source/ChaosVD/Private/ChaosVDEngine.cpp
@@ -34,6 +34,13 @@ void FChaosVDEngine::DeInitialize()
 	CurrentScene->DeInitialize();
 	CurrentScene.Reset();
 	PlaybackController.Reset();
+
+	if (const TSharedPtr<FChaosVDTraceManager> CVDTraceManager = FChaosVDModule::Get().GetTraceManager())
+	{
+		CVDTraceManager->CloseSession(CurrentSessionName);
+	}
+
+	CollectGarbage(GARBAGE_COLLECTION_KEEPFLAGS);
 }
 
 void FChaosVDEngine::LoadRecording(const FString& FilePath)
diff --git a/Engine/Plugins/Experimental/ChaosVD/Source/ChaosVD/Private/ChaosVDPlaybackController.cpp b/Engine/Plugins/Experimental/ChaosVD/Source/ChaosVD/Private/ChaosVDPlaybackController.cpp
--- a/Engine/Plugins/Experimental/ChaosVD/Source/ChaosVD/Private/ChaosVDPlaybackController.cpp
+++ b/Engine/Plugins/Experimental/ChaosVD/Source/ChaosVD/Private/ChaosVDPlaybackController.cpp
@@ -79,18 +79,6 @@ void FChaosVDPlaybackController::UnloadCurrentRecording(const bool bBroadcastUpd
 	if (LoadedRecording.IsValid())
 	{
 		LoadedRecording->OnRecordingUpdated().RemoveAll(this);
-
-		//TODO: Add proper support to edit the trace provider
-		/*if (const TSharedPtr<const TraceServices::IAnalysisSession> TraceSession = FChaosVDModule::Get().GetTraceManager()->GetSession(LoadedRecording->SessionName))
-		{
-			TraceServices::IAnalysisSession* TraceSessionPtr = const_cast<TraceServices::IAnalysisSession*>(TraceSession.Get());
-			if (FChaosVDTraceProvider* ChaosVDProvider = TraceSessionPtr->EditProvider<FChaosVDTraceProvider>(FChaosVDTraceProvider::ProviderName))
-			{
-				TraceServices::FAnalysisSessionEditScope SessionEditScope(*TraceSessionPtr);
-				ChaosVDProvider->DeleteRecordingInstanceForSession();
-			}
-		}*/
-
 		LoadedRecording.Reset();
 	}
 
diff --git a/Engine/Plugins/Experimental/ChaosVD/Source/ChaosVD/Private/ChaosVDScene.cpp b/Engine/Plugins/Experimental/ChaosVD/Source/ChaosVD/Private/ChaosVDScene.cpp
--- a/Engine/Plugins/Experimental/ChaosVD/Source/ChaosVD/Private/ChaosVDScene.cpp
+++ b/Engine/Plugins/Experimental/ChaosVD/Source/ChaosVD/Private/ChaosVDScene.cpp
@@ -40,10 +40,15 @@ void FChaosVDScene::DeInitialize()
 
 	GeometryGenerator.Reset();
 
+	CleanUpScene();
+
 	if (PhysicsVDWorld)
 	{
 		PhysicsVDWorld->DestroyWorld(true);
+		GEngine->DestroyWorldContext(PhysicsVDWorld);
+
 		PhysicsVDWorld->MarkAsGarbage();
+		PhysicsVDWorld = nullptr;
 	}
 	
 	CollectGarbage(GARBAGE_COLLECTION_KEEPFLAGS);
@@ -135,8 +140,6 @@ void FChaosVDScene::CleanUpScene()
 	}
 
 	ParticleVDInstancesByID.Reset();
-
-	CollectGarbage(GARBAGE_COLLECTION_KEEPFLAGS);
 }
 
 AChaosVDParticleActor* FChaosVDScene::SpawnParticleFromRecordedData(const FChaosVDParticleDebugData& InParticleData, const FChaosVDSolverFrameData& InFrameData)
diff --git a/Engine/Plugins/Experimental/ChaosVD/Source/ChaosVD/Private/Trace/ChaosVDTraceManager.cpp b/Engine/Plugins/Experimental/ChaosVD/Source/ChaosVD/Private/Trace/ChaosVDTraceManager.cpp
--- a/Engine/Plugins/Experimental/ChaosVD/Source/ChaosVD/Private/Trace/ChaosVDTraceManager.cpp
+++ b/Engine/Plugins/Experimental/ChaosVD/Source/ChaosVD/Private/Trace/ChaosVDTraceManager.cpp
@@ -25,7 +25,7 @@ FChaosVDTraceManager::~FChaosVDTraceManager()
 
 FString FChaosVDTraceManager::LoadTraceFile(const FString& InTraceFilename)
 {
-	ResetSession(InTraceFilename);
+	CloseSession(InTraceFilename);
 
 	ITraceServicesModule& TraceServicesModule = FModuleManager::LoadModuleChecked<ITraceServicesModule>("TraceServices");
 	if (const TSharedPtr<TraceServices::IAnalysisService> TraceAnalysisService = TraceServicesModule.GetAnalysisService())
@@ -48,7 +48,7 @@ TSharedPtr<const TraceServices::IAnalysisSession> FChaosVDTraceManager::GetSessi
 	return nullptr;
 }
 
-void FChaosVDTraceManager::ResetSession(const FString& InSessionName)
+void FChaosVDTraceManager::CloseSession(const FString& InSessionName)
 {
 	if (const TSharedPtr<const TraceServices::IAnalysisSession>* Session = AnalysisSessionByName.Find(InSessionName))
 	{
diff --git a/Engine/Plugins/Experimental/ChaosVD/Source/ChaosVD/Public/Trace/ChaosVDTraceManager.h b/Engine/Plugins/Experimental/ChaosVD/Source/ChaosVD/Public/Trace/ChaosVDTraceManager.h
--- a/Engine/Plugins/Experimental/ChaosVD/Source/ChaosVD/Public/Trace/ChaosVDTraceManager.h
+++ b/Engine/Plugins/Experimental/ChaosVD/Source/ChaosVD/Public/Trace/ChaosVDTraceManager.h
@@ -21,9 +21,9 @@ public:
 
 	TSharedPtr<const TraceServices::IAnalysisSession> GetSession(const FString& InSessionName);
 
+	void CloseSession(const FString& InSessionName);
+
 private:
-	
-	void ResetSession(const FString& InSessionName);
 
 	/** The trace analysis session. */
 	TMap<FString, TSharedPtr<const TraceServices::IAnalysisSession>> AnalysisSessionByName;
