[Chaos VD] Improved Geometry Generation

-  Now Geometry generation is async and parallelized.  You can still interact with the tool while the meshes are generated
- Added support to generate LODs for static meshes. The main use case is for meshes generated based on Heightfield data
- Heightfield-generated meshes now generate LODs, which brings the video memory usage under control. The other feature needed to fully resolve this is a level streaming system for the tool, but that will be implemented in a future change

diff --git a/Engine/Plugins/Experimental/ChaosVD/ChaosVD.uplugin b/Engine/Plugins/Experimental/ChaosVD/ChaosVD.uplugin
--- a/Engine/Plugins/Experimental/ChaosVD/ChaosVD.uplugin
+++ b/Engine/Plugins/Experimental/ChaosVD/ChaosVD.uplugin
@@ -21,5 +21,11 @@
 			"Type": "Editor",
 			"LoadingPhase": "Default"
 		}
+	],
+	"Plugins": [
+		{
+			"Name": "GeometryProcessing",
+			"Enabled": true
+		}
 	]
 }
\ No newline at end of file
diff --git a/Engine/Plugins/Experimental/ChaosVD/Source/ChaosVD/ChaosVD.Build.cs b/Engine/Plugins/Experimental/ChaosVD/Source/ChaosVD/ChaosVD.Build.cs
--- a/Engine/Plugins/Experimental/ChaosVD/Source/ChaosVD/ChaosVD.Build.cs
+++ b/Engine/Plugins/Experimental/ChaosVD/Source/ChaosVD/ChaosVD.Build.cs
@@ -38,7 +38,8 @@ public class ChaosVD : ModuleRules
 				"GeometryCore",
 				"MeshConversion",
 				"MeshDescription",
-				"StaticMeshDescription",
+				"StaticMeshDescription", 
+				"DynamicMesh"
 			}
 			);
 
diff --git a/Engine/Plugins/Experimental/ChaosVD/Source/ChaosVD/Private/ChaosVDGeometryBuilder.cpp b/Engine/Plugins/Experimental/ChaosVD/Source/ChaosVD/Private/ChaosVDGeometryBuilder.cpp
--- a/Engine/Plugins/Experimental/ChaosVD/Source/ChaosVD/Private/ChaosVDGeometryBuilder.cpp
+++ b/Engine/Plugins/Experimental/ChaosVD/Source/ChaosVD/Private/ChaosVDGeometryBuilder.cpp
@@ -3,74 +3,251 @@
 #include "ChaosVDGeometryBuilder.h"
 
 #include "DynamicMeshToMeshDescription.h"
-#include "Engine/StaticMesh.h"
+#include "MeshConstraints.h"
+#include "MeshConstraintsUtil.h"
+#include "MeshSimplification.h"
 #include "StaticMeshAttributes.h"
 #include "UDynamicMesh.h"
-
 #include "UObject/UObjectGlobals.h"
 
+static FAutoConsoleVariable CVarChaosVDGeometryToProcessPerTick(
+	TEXT("p.Chaos.VD.Tool.GeometryToProcessPerTick"),
+	50,
+	TEXT("Number of generated geometry to process each tick when loading a teace file in the CVD tool"));
+
 void FChaosVDGeometryBuilder::AddReferencedObjects(FReferenceCollector& Collector)
 {
-	Collector.AddStableReferenceMap(MeshCacheMap);
+	Collector.AddStableReferenceMap(DynamicMeshCacheMap);
 	Collector.AddStableReferenceMap(StaticMeshCacheMap);
 }
 
+bool FChaosVDGeometryBuilder::DoesImplicitContainType(const Chaos::FImplicitObject* InImplicitObject, const Chaos::EImplicitObjectType ImplicitTypeToCheck)
+{
+	using namespace Chaos;
+	const EImplicitObjectType InnerType = GetInnerType(InImplicitObject->GetType());
+
+	switch (InnerType)
+	{
+		case ImplicitObjectType::Union:
+			{
+				const FImplicitObjectUnion* Union = InImplicitObject->template GetObject<FImplicitObjectUnion>();
+
+				for (int i = 0; i < Union->GetObjects().Num(); ++i)
+				{
+					const TUniquePtr<FImplicitObject>& UnionImplicit = Union->GetObjects()[i];
+
+					if (DoesImplicitContainType(UnionImplicit.Get(), ImplicitTypeToCheck))
+					{
+						return true;
+					}
+				}
+
+				return false;
+				break;
+			}
+		case ImplicitObjectType::Transformed:
+			{
+				const TImplicitObjectTransformed<FReal, 3>* Transformed = InImplicitObject->template GetObject<TImplicitObjectTransformed<FReal, 3>>();
+				return DoesImplicitContainType(Transformed->GetTransformedObject(), ImplicitTypeToCheck);
+				break;
+			}
+	default:
+		return InnerType == ImplicitTypeToCheck;
+	}
+
+	return false;
+}
+
 UDynamicMesh* FChaosVDGeometryBuilder::CreateAndCacheDynamicMesh(const uint32 GeometryCacheKey, UE::Geometry::FMeshShapeGenerator& MeshGenerator)
 {
-	//TODO: Make this return what is cached when the system is more robust
-	// For now this should not happen and we want to catch it and make it visually noticeable
-	if (MeshCacheMap.Contains(GeometryCacheKey))
 	{
-		ensureMsgf(false, TEXT("Tried to create a new mesh with an existing Cache key"));
-		return nullptr;
+		FReadScopeLock ReadLock(RWLock);
+		//TODO: Make this return what is cached when the system is more robust
+		// For now this should not happen and we want to catch it and make it visually noticeable
+		if (DynamicMeshCacheMap.Contains(GeometryCacheKey))
+		{
+			ensureMsgf(false, TEXT("Tried to create a new mesh with an existing Cache key"));
+			return nullptr;
+		}
 	}
 
+	TRACE_CPUPROFILER_EVENT_SCOPE(FChaosVDGeometryBuilder::CreateAndCacheDynamicMesh_BUILD);
+
 	UDynamicMesh* Mesh = NewObject<UDynamicMesh>();
 	Mesh->SetMesh(&MeshGenerator.Generate());
 
-	MeshCacheMap.Add(GeometryCacheKey, Mesh);
+	{
+		FWriteScopeLock WriteLock(RWLock);
+		DynamicMeshCacheMap.Add(GeometryCacheKey, Mesh);
+	}
 
 	return Mesh;
 }
 
-UStaticMesh* FChaosVDGeometryBuilder::CreateAndCacheStaticMesh(const uint32 GeometryCacheKey, UE::Geometry::FMeshShapeGenerator& MeshGenerator)
+UStaticMesh* FChaosVDGeometryBuilder::CreateAndCacheStaticMesh(const uint32 GeometryCacheKey, UE::Geometry::FMeshShapeGenerator& MeshGenerator, const int32 LODsToGenerateNum)
 {
-	//TODO: Make this return what is cached when the system is more robust
-	// For now this should not happen and we want to catch it and make it visually noticeable
-	if (StaticMeshCacheMap.Contains(GeometryCacheKey))
 	{
-		ensureMsgf(false, TEXT("Tried to create a new mesh with an existing Cache key"));
-		return nullptr;
+		FReadScopeLock ReadLock(RWLock);
+		//TODO: Make this return what is cached when the system is more robust
+		// For now this should not happen and we want to catch it and make it visually noticeable
+		if (StaticMeshCacheMap.Contains(GeometryCacheKey))
+		{
+			ensureMsgf(false, TEXT("Tried to create a new mesh with an existing Cache key"));
+			return nullptr;
+		}
 	}
 
 	//TODO: Instead of generating a dynamic mesh and discard it, we should
-	// Create a Mesh description directly.
+	// Create a Mesh description directly when no LODs are required.
 	// We could create a base class for our mesh Generators and add a Generate method that generates these mesh descriptions
-	UDynamicMesh* Mesh = NewObject<UDynamicMesh>();
-	UStaticMesh* StaticMesh = NewObject<UStaticMesh>();
-	StaticMesh->GetStaticMaterials().Add(FStaticMaterial());
+	UStaticMesh* MainStaticMesh = NewObject<UStaticMesh>();
+	MainStaticMesh->GetStaticMaterials().Add(FStaticMaterial());
 
-	Mesh->SetMesh(&MeshGenerator.Generate());
+	const int32 MeshDescriptionsToGenerate = LODsToGenerateNum + 1;
 
-	FMeshDescription MeshDescription;
-	FStaticMeshAttributes Attributes(MeshDescription);
-	Attributes.Register();
+	TArray<const FMeshDescription*> LODDescriptions;
+	LODDescriptions.Reserve(MeshDescriptionsToGenerate);
 
-	FDynamicMeshToMeshDescription Converter;
-	Converter.Convert(Mesh->GetMeshPtr(), MeshDescription, true);
+	MainStaticMesh->SetNumSourceModels(MeshDescriptionsToGenerate);
 
-	UStaticMesh::FBuildMeshDescriptionsParams Params;
-	Params.bUseHashAsGuid = true;
-	Params.bMarkPackageDirty = false;
-	Params.bBuildSimpleCollision = false;
-	Params.bCommitMeshDescription = false;
-	Params.bFastBuild = true;
+	FDynamicMesh3 DynamicMesh(&MeshGenerator.Generate());
 
-	StaticMesh->NaniteSettings.bEnabled = true;
-	StaticMesh->BuildFromMeshDescriptions({&MeshDescription}, Params);
+	for (int32 i = 0; i < MeshDescriptionsToGenerate; i++)
+	{
+		if (i > 0)
+		{
+			TRACE_CPUPROFILER_EVENT_SCOPE(FChaosVDGeometryBuilder::CreateAndCacheStaticMesh_LOD);
+			//TODO: Come up with a better algo for this.
+			const int32 DesiredTriangleCount = DynamicMesh.TriangleCount() / (i * 2);
+			// Simplify
+			UE::Geometry::FMeshConstraints Constraints;
+			UE::Geometry::FMeshConstraintsUtil::ConstrainAllBoundariesAndSeams(Constraints, DynamicMesh,
+																			   UE::Geometry::EEdgeRefineFlags::NoFlip, UE::Geometry::EEdgeRefineFlags::NoConstraint, UE::Geometry::EEdgeRefineFlags::NoConstraint,
+																			   false, false, true);
+			// Reduce the same previous LOD Mesh on each iteration
+			UE::Geometry::FQEMSimplification Simplifier(&DynamicMesh);
+			Simplifier.SetExternalConstraints(MoveTemp(Constraints));
+			Simplifier.SimplifyToTriangleCount(DesiredTriangleCount);
+		}
+
+		FMeshDescription* MeshDescription = new FMeshDescription();
+		FStaticMeshAttributes Attributes(*MeshDescription);
+		Attributes.Register();
+
+		FDynamicMeshToMeshDescription Converter;
+		Converter.Convert(&DynamicMesh, *MeshDescription, true);
+		LODDescriptions.Add(MeshDescription);
+	}
 
-	StaticMeshCacheMap.Add(GeometryCacheKey, StaticMesh);
+	{
+		TRACE_CPUPROFILER_EVENT_SCOPE(FChaosVDGeometryBuilder::CreateAndCacheStaticMesh_BUILD);
+		UStaticMesh::FBuildMeshDescriptionsParams Params;
+		Params.bUseHashAsGuid = true;
+		Params.bMarkPackageDirty = false;
+		Params.bBuildSimpleCollision = false;
+		Params.bCommitMeshDescription = false;
+		Params.bFastBuild = true;
+
+		MainStaticMesh->NaniteSettings.bEnabled = true;
+		MainStaticMesh->BuildFromMeshDescriptions(LODDescriptions, Params);
+
+		MainStaticMesh->bAutoComputeLODScreenSize = true;
+	}
 
-	return StaticMesh;
+	{
+		FWriteScopeLock WriteLock(RWLock);
+		StaticMeshCacheMap.Add(GeometryCacheKey, MainStaticMesh);
+	}
+
+	for (const FMeshDescription* Desc : LODDescriptions)
+	{
+		delete Desc;
+		Desc = nullptr;
+	}
+
+	LODDescriptions.Reset();
+
+	return MainStaticMesh;
 }
 
+void FChaosVDGeometryBuilder::ApplyMeshToComponentFromKey(TWeakObjectPtr<UMeshComponent> MeshComponent, const uint32 GeometryKey)
+{
+	if (!MeshComponent.IsValid())
+	{
+		return;
+	}
+
+	if (UDynamicMeshComponent* DynamicMeshComponent = Cast<UDynamicMeshComponent>(MeshComponent))
+	{
+		if (UDynamicMesh** DynamicMesh = DynamicMeshCacheMap.Find(GeometryKey))
+		{
+			DynamicMeshComponent->SetDynamicMesh(*DynamicMesh);
+		}
+	}
+	else if (UStaticMeshComponent* StaticMeshComponent = Cast<UStaticMeshComponent>(MeshComponent))
+	{
+		if (UStaticMesh** StaticMesh = StaticMeshCacheMap.Find(GeometryKey))
+		{
+			StaticMeshComponent->SetStaticMesh(*StaticMesh);
+		}
+	}
+}
+
+bool FChaosVDGeometryBuilder::GameThreadTick(float DeltaTime)
+{
+	int32 CurrentGeometryProcessed = 0;
+	while (!GeometryReadyToApplyQueue.IsEmpty() && CurrentGeometryProcessed < CVarChaosVDGeometryToProcessPerTick->GetInt())
+	{
+		uint32 GeometryKey = 0;
+		GeometryReadyToApplyQueue.Dequeue(GeometryKey);
+		CurrentGeometryProcessed++;
+
+		TArray<TWeakObjectPtr<UMeshComponent>>* MeshComponentsWaiting = nullptr;
+		{
+			FReadScopeLock ReadLock(RWLock);
+			 MeshComponentsWaiting = MeshComponentsWaitingForGeometryByKey.Find(GeometryKey);
+		}
+
+		if (MeshComponentsWaiting)
+		{
+			for (const TWeakObjectPtr<UMeshComponent>& MeshComponent : *MeshComponentsWaiting)
+			{
+				ApplyMeshToComponentFromKey(MeshComponent, GeometryKey);
+			}
+		}
+		
+		{
+			FWriteScopeLock WriteLock(RWLock);
+			MeshComponentsWaitingForGeometryByKey.Remove(GeometryKey);
+		}
+	}
+
+	return true;
+}
+
+void FChaosVDGeometryBuilder::RegisterMeshComponentWaitingForGeometry(uint32 GeometryKey, TWeakObjectPtr<UMeshComponent> MesComponent, const int32 LODsToGenerateNum)
+{
+	if (!MesComponent.IsValid())
+	{
+		return;
+	}
+
+	const bool bIsMeshTypeLODCompatible = !MesComponent->IsA(UInstancedStaticMeshComponent::StaticClass()) && !MesComponent->IsA(UDynamicMeshComponent::StaticClass());
+	if (LODsToGenerateNum > 0 && !bIsMeshTypeLODCompatible)
+	{
+		ensureMsgf(false, TEXT("LODs are currently not supported with the specified mesh component [%s] in the CVD tool | [%d] LODs were requested when 0 is expected"), *MesComponent->GetName(), LODsToGenerateNum);
+	}
+
+	{
+		FWriteScopeLock WriteLock(RWLock);
+	
+		if (TArray<TWeakObjectPtr<UMeshComponent>>* MeshComponentsWaiting = MeshComponentsWaitingForGeometryByKey.Find(GeometryKey))
+		{
+			MeshComponentsWaiting->Add(MesComponent);
+		}
+		else
+		{
+			MeshComponentsWaitingForGeometryByKey.Add(GeometryKey, {MesComponent});
+		}
+	}
+}
diff --git a/Engine/Plugins/Experimental/ChaosVD/Source/ChaosVD/Private/ChaosVDParticleActor.cpp b/Engine/Plugins/Experimental/ChaosVD/Source/ChaosVD/Private/ChaosVDParticleActor.cpp
--- a/Engine/Plugins/Experimental/ChaosVD/Source/ChaosVD/Private/ChaosVDParticleActor.cpp
+++ b/Engine/Plugins/Experimental/ChaosVD/Source/ChaosVD/Private/ChaosVDParticleActor.cpp
@@ -26,26 +26,38 @@ void AChaosVDParticleActor::UpdateFromRecordedData(const FChaosVDParticleDebugDa
 
 void AChaosVDParticleActor::UpdateGeometryData(const TSharedPtr<const Chaos::FImplicitObject>& ImplicitObject)
 {
-	if (bIsGeometryDataLoaded)
+	if (bIsGeometryDataGenerationStarted)
 	{
 		return;
 	}
 
-	TArray<UMeshComponent*> OutGeneratedMeshComponents;
-
 	if (const TSharedPtr<FChaosVDScene>& ScenePtr = OwningScene.Pin())
 	{
 		if (const TSharedPtr<FChaosVDGeometryBuilder>& GeometryGenerator = ScenePtr->GetGeometryGenerator())
 		{
+			TArray<TWeakObjectPtr<UMeshComponent>> OutGeneratedMeshComponents;
 			Chaos::FRigidTransform3 Transform;
-			GeometryGenerator->CreateMeshComponentsFromImplicit<UStaticMesh, UInstancedStaticMeshComponent>(ImplicitObject.Get(), this, OutGeneratedMeshComponents, Transform);
+
+			// Heightfields need to be created as Static meshes and use normal Static Mesh components because we need LODs for them due to their high triangle count
+			if (FChaosVDGeometryBuilder::DoesImplicitContainType(ImplicitObject.Get(), Chaos::ImplicitObjectType::HeightField))
+			{
+				constexpr int32 LODsToGenerateNum = 3;
+				constexpr int32 StartingMeshComponentIndex = 0;
+				GeometryGenerator->CreateMeshComponentsFromImplicit<UStaticMesh, UStaticMeshComponent>(ImplicitObject.Get(), this, OutGeneratedMeshComponents, Transform, StartingMeshComponentIndex, LODsToGenerateNum);
+			}
+			else
+			{
+				GeometryGenerator->CreateMeshComponentsFromImplicit<UStaticMesh, UInstancedStaticMeshComponent>(ImplicitObject.Get(), this, OutGeneratedMeshComponents, Transform);
+			}
+
+			
+
+			if (OutGeneratedMeshComponents.Num() > 0)
+			{
+				bIsGeometryDataGenerationStarted = true;
+			}
 		}
 	}
-	
-	if (OutGeneratedMeshComponents.Num() > 0)
-	{
-		bIsGeometryDataLoaded = true;
-	}
 }
 
 void AChaosVDParticleActor::SetScene(const TSharedPtr<FChaosVDScene>& InScene)
@@ -64,7 +76,6 @@ void AChaosVDParticleActor::SetScene(const TSharedPtr<FChaosVDScene>& InScene)
 	}
 }
 
-
 void AChaosVDParticleActor::BeginDestroy()
 {
 	if (const TSharedPtr<FChaosVDScene>& ScenePtr = OwningScene.Pin())
diff --git a/Engine/Plugins/Experimental/ChaosVD/Source/ChaosVD/Public/ChaosVDGeometryBuilder.h b/Engine/Plugins/Experimental/ChaosVD/Source/ChaosVD/Public/ChaosVDGeometryBuilder.h
--- a/Engine/Plugins/Experimental/ChaosVD/Source/ChaosVD/Public/ChaosVDGeometryBuilder.h
+++ b/Engine/Plugins/Experimental/ChaosVD/Source/ChaosVD/Public/ChaosVDGeometryBuilder.h
@@ -8,12 +8,15 @@
 #include "Chaos/HeightField.h"
 #include "Chaos/ImplicitObjectScaled.h"
 #include "Chaos/ImplicitObjectType.h"
+#include "Containers/Ticker.h"
 #include "Components/DynamicMeshComponent.h"
 #include "Components/InstancedStaticMeshComponent.h"
+#include "Engine/StaticMesh.h"
 #include "Generators/CapsuleGenerator.h"
 #include "Generators/MinimalBoxMeshGenerator.h"
 #include "Generators/SphereGenerator.h"
 #include "UDynamicMesh.h"
+#include "Tasks/Task.h"
 #include "UObject/GCObject.h"
 #include "UObject/UObjectGlobals.h"
 
@@ -33,10 +36,38 @@ namespace UE
 /*
  * Generates Dynamic mesh components and dynamic meshes based on Chaos implicit object data
  */
-class FChaosVDGeometryBuilder : public FGCObject
+class FChaosVDGeometryBuilder : public FGCObject, public TSharedFromThis<FChaosVDGeometryBuilder>
 {
 public:
 
+	FChaosVDGeometryBuilder()
+	{
+		GameThreadTickDelegate = FTSTicker::GetCoreTicker().AddTicker(FTickerDelegate::CreateRaw(this, &FChaosVDGeometryBuilder::GameThreadTick));
+	}
+
+	virtual ~FChaosVDGeometryBuilder() override
+	{
+		FTSTicker::GetCoreTicker().RemoveTicker(GameThreadTickDelegate);
+
+		for (const TPair<uint32, UStaticMesh*>& StaticMeshByKey : StaticMeshCacheMap)
+		{
+			if (StaticMeshByKey.Value)
+			{
+				StaticMeshByKey.Value->ClearFlags(RF_Standalone);
+				StaticMeshByKey.Value->MarkAsGarbage();
+			}
+		}
+
+		for (const TPair<uint32, UDynamicMesh*>& DynamicMeshByKey : DynamicMeshCacheMap)
+		{
+			if (DynamicMeshByKey.Value)
+			{
+				DynamicMeshByKey.Value->ClearFlags(RF_Standalone);
+				DynamicMeshByKey.Value->MarkAsGarbage();
+			}
+		}
+	}
+
 	/** Creates Dynamic Mesh components for each object within the provided Implicit object
 	 *	@param InImplicitObject : Implicit object to process
 	 *	@param Owner Actor who will own the generated components
@@ -45,14 +76,22 @@ public:
 	 *	@param Transform to apply to the generated components/geometry
 	 */
 	template<typename MeshType, typename ComponentType>
-	void CreateMeshComponentsFromImplicit(const Chaos::FImplicitObject* InImplicitObject, AActor* Owner, TArray<UMeshComponent*>& OutMeshComponents, Chaos::FRigidTransform3& Transform, int32 Index = 0);
+	void CreateMeshComponentsFromImplicit(const Chaos::FImplicitObject* InImplicitObject, AActor* Owner, TArray<TWeakObjectPtr<UMeshComponent>>& OutMeshComponents, Chaos::FRigidTransform3& Transform, const int32 Index = 0, const int32 DesiredLODCount = 0);
 	
 	virtual void AddReferencedObjects(FReferenceCollector& Collector) override;
 	virtual FString GetReferencerName() const override
 	{
-		return TEXT("FChaosVDScene");
+		return TEXT("FChaosVDGeometryBuilder");
 	}
 
+	/**
+	 * Evaluates an Implicit objects and returns true if it contains an object of the specified type 
+	 * @param InImplicitObject Object to evaluate
+	 * @param ImplicitTypeToCheck Object type to compare against
+	 * @return 
+	 */
+	static bool DoesImplicitContainType(const Chaos::FImplicitObject* InImplicitObject, const Chaos::EImplicitObjectType ImplicitTypeToCheck);
+
 private:
 
 	/** Creates a Dynamic Mesh for the provided Implicit object and generator, and then caches it to be reused later
@@ -71,10 +110,10 @@ private:
 	 * @param GeometryCacheKey Key to be used to find this geometry in the cache
 	 * @param MeshGenerator Generator class with the data and rules to create the mesh
 	 */
-	UStaticMesh* CreateAndCacheStaticMesh(const uint32 GeometryCacheKey, UE::Geometry::FMeshShapeGenerator& MeshGenerator);
+	UStaticMesh* CreateAndCacheStaticMesh(const uint32 GeometryCacheKey, UE::Geometry::FMeshShapeGenerator& MeshGenerator, const int32 LODsToGenerateNum = 0);
 
 	/**
-	 * Creates an empty DynamicMeshComponent/Static Mesh Component or Instanced mesg component and adds it to the actor
+	 * Creates an empty DynamicMeshComponent/Static Mesh Component or Instanced mesh component and adds it to the actor
 	 * @param Owner Actor who will own the component
 	 * @param Name Name of the component. It has to be unique within the components in the owner actor
 	 * @param Transform (Optional) Transform to apply as Relative Transform in the component after its creating and attachment to the provided actor
@@ -83,12 +122,20 @@ private:
 	ComponentType* CreateMeshComponent(AActor* Owner, const FString& Name, const Chaos::FRigidTransform3& Transform) const;
 
 	/**
-	 * Creates a Mesh from the provided Implicit object geometry data
+	 * Applies a mesh to a mesh component based on its type
+	 * @param MeshComponent Mesh component to apply the mesh to
+	 * @param GeometryKey Key to find the mesh to apply in the cache
+	 */
+	void ApplyMeshToComponentFromKey(TWeakObjectPtr<UMeshComponent> MeshComponent, const uint32 GeometryKey);
+
+	/**
+	 * Creates a Mesh from the provided Implicit object geometry data. This is a async operation, and the mesh will be assigned to the component once is ready
 	 * @param GeometryCacheKey Key to be used to find this geometry in the cache
 	 * @param MeshGenerator Generator class that has the data to generate the mesh
+	 * @param LODsToGenerateNum Num of LODs to Generate. Not all mesh types support this
 	 * */
 	template<typename MeshType>
-	MeshType* CreateAndCacheMeshForImplicit(const uint32 GeometryCacheKey, UE::Geometry::FMeshShapeGenerator& MeshGenerator);
+	void DispatchCreateAndCacheMeshForImplicitAsync(const uint32 GeometryCacheKey, TSharedPtr<UE::Geometry::FMeshShapeGenerator> MeshGenerator, const int32 LODsToGenerateNum = 0);
 
 	/* Process an FImplicitObject and returns de desired geometry type. Could be directly the shape or another version of the implicit */
 	template <bool bIsInstanced, typename GeometryType>
@@ -98,16 +145,69 @@ private:
 	template<typename GeometryType>
 	const GeometryType* GetGeometryBasedOnPackedType(const Chaos::FImplicitObject* InImplicitObject, Chaos::FRigidTransform3& Transform, const Chaos::EImplicitObjectType PackedType) const;
 
+	/** Tick method of this Geometry builder. Used to do everything that needs to be performed in the GT, like applying the generated meshes to mesh component */
+	bool GameThreadTick(float DeltaTime);
+
+	/**
+	 * Add a mesh component to the waiting list for Geometry. This needs to be called before dispatching a generation job for new Geometry
+	 * @param GeometryKey Key of the geometry that is being generated (or will be generated)
+	 * @param MesComponent Component where the geometry needs to be applied
+	 * @param LODsToGenerateNum How many LODs we expect the generated mesh to have. Used to check we are using the correct mesh component type
+	 */
+	void RegisterMeshComponentWaitingForGeometry(uint32 GeometryKey, TWeakObjectPtr<UMeshComponent> MesComponent, const int32 LODsToGenerateNum);
+
 	/** Map containing already generated dynamic mesh for any given implicit object */
-	TMap<uint32, UDynamicMesh*> MeshCacheMap;
+	TMap<uint32, UDynamicMesh*> DynamicMeshCacheMap;
 
 	/** Map containing already generated static mesh for any given implicit object */
 	TMap<uint32, UStaticMesh*> StaticMeshCacheMap;
+
+	/** Map containing all the meshes component waiting for geometry, by geometry key*/
+	TMap<uint32, TArray<TWeakObjectPtr<UMeshComponent>>> MeshComponentsWaitingForGeometryByKey;
+
+	/** Set of all geometry keys of the Meshes that are being generated but not ready yet */
+	TSet<uint32> GeometryBeingGeneratedByKey;
+	
+	/** Used to locks Read or Writes to the Geometry cache and in flight job tracking containers */
+	FRWLock RWLock;
+
+	/** Handle to the ticker used to ticker the Geometry Builder in the game thread*/
+	FTSTicker::FDelegateHandle GameThreadTickDelegate;
+
+	/** Queue of geometry keys already generated and waiting to be applied */
+	TQueue<uint32, EQueueMode::Mpsc> GeometryReadyToApplyQueue;
+
+	friend class FGeometryGenerationTask;
+
 };
 
+/** Used to execute each individual Geometry Generation task using the data with which was constructed.
+ * It allows the task to skip the actual generation attempt if the Geometry builder instance goes away which happens when the tool is closed
+ */
+class FGeometryGenerationTask
+{
+public:
+	FGeometryGenerationTask(const TWeakPtr<FChaosVDGeometryBuilder>& InBuilder, const TSharedPtr<UE::Geometry::FMeshShapeGenerator>& InGenerator, const uint32 GeometryKey,
+		const int32 LODsToGenerateNum)
+		: Builder(InBuilder),
+		  MeshGenerator(InGenerator),
+		  GeometryKey(GeometryKey),
+		  LODsToGenerateNum(LODsToGenerateNum)
+	{
+	}
+
+	template <typename MeshType>
+	void GenerateGeometry();
+
+private:
+	TWeakPtr<FChaosVDGeometryBuilder> Builder;
+	TSharedPtr<UE::Geometry::FMeshShapeGenerator> MeshGenerator;
+	uint32 GeometryKey;
+	int32 LODsToGenerateNum;
+};
 
 template <typename MeshType, typename ComponentType>
-void FChaosVDGeometryBuilder::CreateMeshComponentsFromImplicit(const Chaos::FImplicitObject* InImplicitObject, AActor* Owner, TArray<UMeshComponent*>& OutMeshComponents, Chaos::FRigidTransform3& Transform, int32 Index)
+void FChaosVDGeometryBuilder::CreateMeshComponentsFromImplicit(const Chaos::FImplicitObject* InImplicitObject, AActor* Owner, TArray<TWeakObjectPtr<UMeshComponent>>& OutMeshComponents, Chaos::FRigidTransform3& Transform, const int32 Index, const int32 DesiredLODCount)
 {
 	static_assert(std::is_same_v<MeshType, UStaticMesh> || std::is_same_v<MeshType, UDynamicMesh>, "CreateMeshComponentsFromImplicit Only supports DynamicMesh and Static Mesh");
 	static_assert(std::is_same_v<ComponentType, UStaticMeshComponent> || std::is_same_v<ComponentType, UInstancedStaticMeshComponent> || std::is_same_v<MeshType, UDynamicMeshComponent>, "CreateMeshComponentsFromImplicit Only supports DynamicMeshComponent, Static MeshComponent and Instanced Static Mesh Component");
@@ -130,7 +230,7 @@ void FChaosVDGeometryBuilder::CreateMeshComponentsFromImplicit(const Chaos::FImp
 		{
 			const TUniquePtr<FImplicitObject>& UnionImplicit = Union->GetObjects()[i];
 
-			CreateMeshComponentsFromImplicit<MeshType, ComponentType>(UnionImplicit.Get(), Owner, OutMeshComponents, Transform, i);	
+			CreateMeshComponentsFromImplicit<MeshType, ComponentType>(UnionImplicit.Get(), Owner, OutMeshComponents, Transform, i, DesiredLODCount);	
 		}
 
 		return;
@@ -140,7 +240,7 @@ void FChaosVDGeometryBuilder::CreateMeshComponentsFromImplicit(const Chaos::FImp
 	{
 		const TImplicitObjectTransformed<FReal, 3>* Transformed = InImplicitObject->template GetObject<TImplicitObjectTransformed<FReal, 3>>();
 		FRigidTransform3 TransformCopy = Transformed->GetTransform();
-		CreateMeshComponentsFromImplicit<MeshType, ComponentType>(Transformed->GetTransformedObject(), Owner, OutMeshComponents,TransformCopy, Index);
+		CreateMeshComponentsFromImplicit<MeshType, ComponentType>(Transformed->GetTransformedObject(), Owner, OutMeshComponents,TransformCopy, Index, DesiredLODCount);
 		return;
 	}
 
@@ -156,17 +256,23 @@ void FChaosVDGeometryBuilder::CreateMeshComponentsFromImplicit(const Chaos::FImp
 
 			MeshComponent = CreateMeshComponent<ComponentType>(Owner, Name, Transform);
 
-			Mesh = GetCachedMeshForImplicit<MeshType>(Sphere->GetTypeHash());
+			const uint32 GeometryKey = Sphere->GetTypeHash();
+			Mesh = GetCachedMeshForImplicit<MeshType>(GeometryKey);
 
-			if (!Mesh)
+			if (Mesh)
 			{
-				UE::Geometry::FSphereGenerator SphereGen;
-				SphereGen.Radius = Sphere->GetRadius();
-				SphereGen.NumTheta = 50;
-				SphereGen.NumPhi = 50;
-				SphereGen.bPolygroupPerQuad = false;
-
-				Mesh = CreateAndCacheMeshForImplicit<MeshType>(Sphere->GetTypeHash(), SphereGen);
+				ApplyMeshToComponentFromKey(MeshComponent, GeometryKey);
+			}
+			else
+			{
+				TSharedPtr<UE::Geometry::FSphereGenerator> SphereGen = MakeShared<UE::Geometry::FSphereGenerator>();
+				SphereGen->Radius = Sphere->GetRadius();
+				SphereGen->NumTheta = 50;
+				SphereGen->NumPhi = 50;
+				SphereGen->bPolygroupPerQuad = false;
+
+				RegisterMeshComponentWaitingForGeometry(GeometryKey, MeshComponent, DesiredLODCount);
+				DispatchCreateAndCacheMeshForImplicitAsync<MeshType>(GeometryKey, SphereGen);
 			}
 
 			break;
@@ -178,17 +284,23 @@ void FChaosVDGeometryBuilder::CreateMeshComponentsFromImplicit(const Chaos::FImp
 			const FString Name = FString::Format(TEXT("{0} - {1}"), {TEXT("Box"), FString::FromInt(Index)});
 			MeshComponent = CreateMeshComponent<ComponentType>(Owner, Name, Transform);
 
-			Mesh = GetCachedMeshForImplicit<MeshType>(Box->GetTypeHash());
+			const uint32 GeometryKey = Box->GetTypeHash();
+			Mesh = GetCachedMeshForImplicit<MeshType>(GeometryKey);
 
-			if (!Mesh)
+			if (Mesh)
 			{
-				UE::Geometry::FMinimalBoxMeshGenerator BoxGen;
+				ApplyMeshToComponentFromKey(MeshComponent, GeometryKey);
+			}
+			else
+			{
+				TSharedPtr<UE::Geometry::FMinimalBoxMeshGenerator> BoxGen = MakeShared<UE::Geometry::FMinimalBoxMeshGenerator>();
 				UE::Geometry::FOrientedBox3d OrientedBox;
 				OrientedBox.Frame = UE::Geometry::FFrame3d(Box->Center());
 				OrientedBox.Extents = Box->Extents() * 0.5;
-				BoxGen.Box = OrientedBox;
+				BoxGen->Box = OrientedBox;
 
-				Mesh = CreateAndCacheMeshForImplicit<MeshType>(Box->GetTypeHash(), BoxGen);
+				RegisterMeshComponentWaitingForGeometry(GeometryKey, MeshComponent, DesiredLODCount);
+				DispatchCreateAndCacheMeshForImplicitAsync<MeshType>(GeometryKey, BoxGen);
 			}
 
 			break;
@@ -211,18 +323,24 @@ void FChaosVDGeometryBuilder::CreateMeshComponentsFromImplicit(const Chaos::FImp
 			MeshComponent->SetRelativeLocation(FinalLocation);
 			MeshComponent->SetRelativeScale3D(Transform.GetScale3D());
 
-			Mesh = GetCachedMeshForImplicit<MeshType>(Capsule->GetTypeHash());
+			const uint32 GeometryKey = Capsule->GetTypeHash();
+			Mesh = GetCachedMeshForImplicit<MeshType>(GeometryKey);
 
-			if (!Mesh)
+			if (Mesh)
 			{
-				UE::Geometry::FCapsuleGenerator CapsuleGenerator;
-				CapsuleGenerator.Radius = FMath::Max(FMathf::ZeroTolerance, Capsule->GetRadius());
-				CapsuleGenerator.SegmentLength = FMath::Max(FMathf::ZeroTolerance, Capsule->GetSegment().GetLength());
-				CapsuleGenerator.NumHemisphereArcSteps = 12;
-				CapsuleGenerator.NumCircleSteps = 12;
-				CapsuleGenerator.bPolygroupPerQuad = false;
-				
-				Mesh = CreateAndCacheMeshForImplicit<MeshType>(Capsule->GetTypeHash(), CapsuleGenerator);
+				ApplyMeshToComponentFromKey(MeshComponent, GeometryKey);
+			}
+			else
+			{
+				TSharedPtr<UE::Geometry::FCapsuleGenerator> CapsuleGenerator = MakeShared<UE::Geometry::FCapsuleGenerator>();
+				CapsuleGenerator->Radius = FMath::Max(FMathf::ZeroTolerance, Capsule->GetRadius());
+				CapsuleGenerator->SegmentLength = FMath::Max(FMathf::ZeroTolerance, Capsule->GetSegment().GetLength());
+				CapsuleGenerator->NumHemisphereArcSteps = 12;
+				CapsuleGenerator->NumCircleSteps = 12;
+				CapsuleGenerator->bPolygroupPerQuad = false;
+
+				RegisterMeshComponentWaitingForGeometry(GeometryKey, MeshComponent, DesiredLODCount);
+				DispatchCreateAndCacheMeshForImplicitAsync<MeshType>(GeometryKey, CapsuleGenerator);
 			}
 			break;
 		}
@@ -238,13 +356,21 @@ void FChaosVDGeometryBuilder::CreateMeshComponentsFromImplicit(const Chaos::FImp
 			{
 				const FString Name = FString::Format(TEXT("{0} - {1}"), {TEXT("Convex"), FString::FromInt(Index)});
 				MeshComponent = CreateMeshComponent<ComponentType>(Owner, Name, Transform);
-				Mesh = GetCachedMeshForImplicit<MeshType>(Convex->GetTypeHash());
+				
+				const uint32 GeometryKey = Convex->GetTypeHash();
+				Mesh = GetCachedMeshForImplicit<MeshType>(GeometryKey);
 	
-				if (!Mesh)
+				if (Mesh)
+				{
+					ApplyMeshToComponentFromKey(MeshComponent, GeometryKey);
+				}
+				else
 				{
-					FChaosVDConvexMeshGenerator ConvexMeshGen;
-					ConvexMeshGen.GenerateFromConvex(*Convex);
-					Mesh = CreateAndCacheMeshForImplicit<MeshType>(Convex->GetTypeHash() ,ConvexMeshGen);
+					TSharedPtr<FChaosVDConvexMeshGenerator> ConvexMeshGen = MakeShared<FChaosVDConvexMeshGenerator>();
+					ConvexMeshGen->GenerateFromConvex(*Convex);
+
+					RegisterMeshComponentWaitingForGeometry(GeometryKey, MeshComponent, DesiredLODCount);
+					DispatchCreateAndCacheMeshForImplicitAsync<MeshType>(GeometryKey ,ConvexMeshGen);
 				}
 			}
 
@@ -261,13 +387,21 @@ void FChaosVDGeometryBuilder::CreateMeshComponentsFromImplicit(const Chaos::FImp
 				const FString Name = FString::Format(TEXT("{0} - {1}"), {TEXT("Trimesh"), FString::FromInt(Index)});
 				MeshComponent = CreateMeshComponent<ComponentType>(Owner, Name, Transform);
 
-				Mesh = GetCachedMeshForImplicit<MeshType>(TriangleMesh->GetTypeHash());
-				if (!Mesh)
+				const uint32 GeometryKey = TriangleMesh->GetTypeHash();
+				Mesh = GetCachedMeshForImplicit<MeshType>(GeometryKey);
+	
+				if (Mesh)
 				{
-					FChaosVDTriMeshGenerator TriMeshGen;
-					TriMeshGen.bReverseOrientation = true;
-					TriMeshGen.GenerateFromTriMesh(*TriangleMesh);		
-					Mesh = CreateAndCacheMeshForImplicit<MeshType>(TriangleMesh->GetTypeHash(), TriMeshGen);
+					ApplyMeshToComponentFromKey(MeshComponent, GeometryKey);
+				}
+				else
+				{
+					TSharedPtr<FChaosVDTriMeshGenerator> TriMeshGen = MakeShared<FChaosVDTriMeshGenerator>();
+					TriMeshGen->bReverseOrientation = true;
+					TriMeshGen->GenerateFromTriMesh(*TriangleMesh);
+
+					RegisterMeshComponentWaitingForGeometry(GeometryKey, MeshComponent, DesiredLODCount);
+					DispatchCreateAndCacheMeshForImplicitAsync<MeshType>(GeometryKey, TriMeshGen);
 				}
 			}
 				
@@ -279,16 +413,25 @@ void FChaosVDGeometryBuilder::CreateMeshComponentsFromImplicit(const Chaos::FImp
 			{
 				const FString Name = FString::Format(TEXT("{0} - {1}"), {TEXT("HeightField"), FString::FromInt(Index)});
 				MeshComponent = CreateMeshComponent<ComponentType>(Owner, Name, Transform);
-				Mesh = GetCachedMeshForImplicit<MeshType>(HeightField->GetTypeHash());
+
+				const uint32 GeometryKey = HeightField->GetTypeHash();
+				Mesh = GetCachedMeshForImplicit<MeshType>(GeometryKey);
 	
-				if (!Mesh)
+				if (Mesh)
+				{
+					ApplyMeshToComponentFromKey(MeshComponent, GeometryKey);
+				}
+				else
 				{
-					FChaosVDHeightFieldMeshGenerator HeightFieldMeshGen;
-					HeightFieldMeshGen.bReverseOrientation = false;
-					HeightFieldMeshGen.GenerateFromHeightField(*HeightField);
+					TSharedPtr<FChaosVDHeightFieldMeshGenerator> HeightFieldMeshGen = MakeShared<FChaosVDHeightFieldMeshGenerator>();
+					HeightFieldMeshGen->bReverseOrientation = false;
+					HeightFieldMeshGen->GenerateFromHeightField(*HeightField);
 
-					Mesh = CreateAndCacheMeshForImplicit<MeshType>(HeightField->GetTypeHash(), HeightFieldMeshGen);
+					RegisterMeshComponentWaitingForGeometry(GeometryKey, MeshComponent, DesiredLODCount);
+					DispatchCreateAndCacheMeshForImplicitAsync<MeshType>(GeometryKey, HeightFieldMeshGen, DesiredLODCount);
 				}
+
+				return;
 			}
 		
 			break;
@@ -297,17 +440,8 @@ void FChaosVDGeometryBuilder::CreateMeshComponentsFromImplicit(const Chaos::FImp
 			break;
 		}
 
-		if (MeshComponent != nullptr && Mesh != nullptr)
+		if (MeshComponent != nullptr)
 		{
-			if constexpr (std::is_same_v<ComponentType, UDynamicMeshComponent>)
-			{
-				MeshComponent->SetDynamicMesh(Mesh);
-			}
-			else if constexpr(std::is_same_v<ComponentType, UInstancedStaticMeshComponent> || std::is_same_v<ComponentType, UStaticMeshComponent>)
-			{
-				MeshComponent->SetStaticMesh(Mesh);
-			}
-
 			OutMeshComponents.Add(MeshComponent);
 		}
 }
@@ -317,7 +451,7 @@ MeshType* FChaosVDGeometryBuilder::GetCachedMeshForImplicit(const uint32 Geometr
 {
 	if constexpr (std::is_same_v<MeshType, UDynamicMesh>)
 	{
-		if (MeshType** MeshPtrPtr = MeshCacheMap.Find(GeometryCacheKey))
+		if (MeshType** MeshPtrPtr = DynamicMeshCacheMap.Find(GeometryCacheKey))
 		{
 			return *MeshPtrPtr;
 		}
@@ -360,18 +494,27 @@ ComponentType* FChaosVDGeometryBuilder::CreateMeshComponent(AActor* Owner, const
 }
 
 template <typename MeshType>
-MeshType* FChaosVDGeometryBuilder::CreateAndCacheMeshForImplicit(const uint32 GeometryKey, UE::Geometry::FMeshShapeGenerator& MeshGenerator)
+void FChaosVDGeometryBuilder::DispatchCreateAndCacheMeshForImplicitAsync(const uint32 GeometryKey, TSharedPtr<UE::Geometry::FMeshShapeGenerator> MeshGenerator, const int32 LODsToGenerateNum)
 {
-	if constexpr (std::is_same_v<MeshType, UDynamicMesh>)
 	{
-		return CreateAndCacheDynamicMesh(GeometryKey, MeshGenerator);
+		FReadScopeLock ReadLock(RWLock);
+		if (GeometryBeingGeneratedByKey.Contains(GeometryKey))
+		{
+			return;
+		}
 	}
-	else if constexpr (std::is_same_v<MeshType, UStaticMesh>)
+
 	{
-		return CreateAndCacheStaticMesh(GeometryKey, MeshGenerator);
+		FWriteScopeLock WriteLock(RWLock);
+		GeometryBeingGeneratedByKey.Add(GeometryKey);
 	}
-
-	return nullptr;
+	
+	TSharedPtr<FGeometryGenerationTask> GenerationTask = MakeShared<FGeometryGenerationTask>(AsWeak(), MeshGenerator, GeometryKey,LODsToGenerateNum);
+	UE::Tasks::Launch(TEXT("GeometryGeneration"),
+		[GenerationTask]()
+		{
+			GenerationTask->GenerateGeometry<MeshType>();
+		});
 }
 
 template <bool bIsInstanced, typename GeometryType>
@@ -420,3 +563,29 @@ const GeometryType* FChaosVDGeometryBuilder::GetGeometryBasedOnPackedType(const
 
 	return nullptr;
 }
+
+template <typename MeshType>
+void FGeometryGenerationTask::GenerateGeometry()
+{
+	if (const TSharedPtr<FChaosVDGeometryBuilder> BuilderPtr = Builder.Pin())
+	{
+		if constexpr (std::is_same_v<MeshType, UDynamicMesh>)
+		{
+			ensureMsgf(LODsToGenerateNum == 0, TEXT("LOD Generation is only suppoted with static meshes | [%d] LODs were requested for a dynamic mesh when 0 is expected"), LODsToGenerateNum);
+			BuilderPtr->CreateAndCacheDynamicMesh(GeometryKey, *MeshGenerator.Get());
+
+			BuilderPtr->GeometryReadyToApplyQueue.Enqueue(GeometryKey);
+		}
+		else if constexpr (std::is_same_v<MeshType, UStaticMesh>)
+		{
+			BuilderPtr->CreateAndCacheStaticMesh(GeometryKey, *MeshGenerator.Get(), LODsToGenerateNum);
+
+			BuilderPtr->GeometryReadyToApplyQueue.Enqueue(GeometryKey);
+		}
+
+		{
+			FWriteScopeLock WriteLock(BuilderPtr->RWLock);
+			BuilderPtr->GeometryBeingGeneratedByKey.Remove(GeometryKey);
+		}
+	}
+}
diff --git a/Engine/Plugins/Experimental/ChaosVD/Source/ChaosVD/Public/ChaosVDParticleActor.h b/Engine/Plugins/Experimental/ChaosVD/Source/ChaosVD/Public/ChaosVDParticleActor.h
--- a/Engine/Plugins/Experimental/ChaosVD/Source/ChaosVD/Public/ChaosVDParticleActor.h
+++ b/Engine/Plugins/Experimental/ChaosVD/Source/ChaosVD/Public/ChaosVDParticleActor.h
@@ -36,7 +36,7 @@ protected:
 	UPROPERTY(VisibleAnywhere, Category="Chaos Visual Debugger Data")
 	FChaosVDParticleDebugData RecordedDebugData;
 
-	bool bIsGeometryDataLoaded = false;
+	bool bIsGeometryDataGenerationStarted = false;
 
 	TWeakPtr<FChaosVDScene> OwningScene;
 
