[Chaos VD] Adding support to recognize and load multiple solvers running at the same time on a ChaosVD recording.

Before this change, only the latest loaded solver data was shown in the UI to control.

Now, the tool provides a single timeline representation for Game Frames (Game thread ticks), and when you scrub that timeline, solvers are loaded/unloaded, and scrub to the closest "solver frame available" for that game frame.

A new tab was added that contains the old playback controls, but this time one for each available solver at the current game frame. These controls can be collapsed as desired.

The per-solver playback controls can be used to scrub in a more precise way to a specific solver frame, or solver step. When scrubbing a solver control, the game frame timeline gets automatically scrubbed to the correct game frame (these are not in sync, so doing it exactly is not possible).

All the approximate synchronization between timelines is done using the recorded platform cycles by using them as the lower bound for a binary search for the target "track spacetime". This is possible because Trace recordings and analysis are serial, so when we populate Game frame data and Solver frame data during Trace analysis, it is guaranteed to be populated in order based on the recorded platform cycles.

The API exposed in the ChaosVDRecording object to do these searches and conversion between "solver spacetimes" could be improved, but this is a starting point.

Not supported yet:

- Scrubbing on a solver timeline will update the game timeline for reference, but not other solvers' timelines, and therefore the scene data for these will not be updated on screen. But scrubbing the Game Frames timeline will bring everything in sync again.

diff --git a/Engine/Plugins/Experimental/ChaosVD/Source/ChaosVD/Private/ChaosVDPlaybackController.cpp b/Engine/Plugins/Experimental/ChaosVD/Source/ChaosVD/Private/ChaosVDPlaybackController.cpp
--- a/Engine/Plugins/Experimental/ChaosVD/Source/ChaosVD/Private/ChaosVDPlaybackController.cpp
+++ b/Engine/Plugins/Experimental/ChaosVD/Source/ChaosVD/Private/ChaosVDPlaybackController.cpp
@@ -3,13 +3,13 @@
 #include "ChaosVDPlaybackController.h"
 
 #include "ChaosVDModule.h"
+#include "ChaosVDPlaybackControllerInstigator.h"
 #include "ChaosVDRecording.h"
 #include "ChaosVDScene.h"
 #include "Trace/ChaosVDTraceManager.h"
 #include "Trace/ChaosVDTraceProvider.h"
 #include "TraceServices/Model/AnalysisSession.h"
 
-
 FChaosVDPlaybackController::FChaosVDPlaybackController(const TWeakPtr<FChaosVDScene>& InSceneToControl)
 {
 	SceneToControl = InSceneToControl;
@@ -17,8 +17,7 @@ FChaosVDPlaybackController::FChaosVDPlaybackController(const TWeakPtr<FChaosVDSc
 
 FChaosVDPlaybackController::~FChaosVDPlaybackController()
 {
-	constexpr bool bBroadcastControllerUpdate = false;
-	UnloadCurrentRecording(bBroadcastControllerUpdate);
+	UnloadCurrentRecording(EChaosVDUnloadRecordingFlags::Silent);
 }
 
 bool FChaosVDPlaybackController::LoadChaosVDRecordingFromTraceSession(const FString& InSessionName)
@@ -48,9 +47,13 @@ bool FChaosVDPlaybackController::LoadChaosVDRecordingFromTraceSession(const FStr
 
 	LoadedRecording->OnRecordingUpdated().AddRaw(this, &FChaosVDPlaybackController::HandleCurrentRecordingUpdated);
 
+	HandleCurrentRecordingUpdated();
+
 	for (TMap<int32, TArray<FChaosVDSolverFrameData>>::TConstIterator SolversIterator = LoadedRecording->GetAvailableSolvers().CreateConstIterator(); SolversIterator; ++SolversIterator)
 	{
-		GoToRecordedStep(SolversIterator->Key, 0, 0);
+		constexpr int32 FrameNumber = 0;
+		constexpr int32 StepNumber = 0;
+		GoToTrackFrame(IChaosVDPlaybackControllerInstigator::InvalidGuid, EChaosVDTrackType::Solver, SolversIterator->Key, FrameNumber, StepNumber);
 	}
 
 	LoadedRecording->OnGeometryDataLoaded().AddLambda([this](const TSharedPtr<const Chaos::FImplicitObject>& NewGeometry, const int32 GeometryID)
@@ -66,22 +69,21 @@ bool FChaosVDPlaybackController::LoadChaosVDRecordingFromTraceSession(const FStr
 		ScenePtr->LoadedRecording = LoadedRecording;
 	}
 	
-	OnControllerUpdated().ExecuteIfBound(AsWeak());
+	OnDataUpdated().Broadcast(AsWeak());
 
 	return true;
 }
 
-void FChaosVDPlaybackController::UnloadCurrentRecording(const bool bBroadcastUpdate)
+void FChaosVDPlaybackController::UnloadCurrentRecording(EChaosVDUnloadRecordingFlags UnloadOptions)
 {
-	CurrentFramePerTrack.Reset();
-	CurrentStepPerTrack.Reset();
-
 	if (LoadedRecording.IsValid())
 	{
 		LoadedRecording->OnRecordingUpdated().RemoveAll(this);
 		LoadedRecording.Reset();
 	}
 
+	TrackInfoPerType.Reset();
+
 	if (const TSharedPtr<FChaosVDScene> SceneToControlSharedPtr = SceneToControl.Pin())
 	{
 		if (SceneToControlSharedPtr->IsInitialized())
@@ -90,141 +92,338 @@ void FChaosVDPlaybackController::UnloadCurrentRecording(const bool bBroadcastUpd
 		}	
 	}
 
-	if (bBroadcastUpdate)
+	if (EnumHasAnyFlags(UnloadOptions, EChaosVDUnloadRecordingFlags::BroadcastChanges))
 	{
 		const TWeakPtr<FChaosVDPlaybackController> ThisWeakPtr = DoesSharedInstanceExist() ? AsWeak() : nullptr;
-		OnControllerUpdated().ExecuteIfBound(ThisWeakPtr);
+		OnDataUpdated().Broadcast(ThisWeakPtr);
 	}
 }
 
-void FChaosVDPlaybackController::GoToRecordedStep(const int32 InTrackID, const int32 FrameNumber, const int32 Step)
+void FChaosVDPlaybackController::GoToRecordedSolverStep(const int32 InTrackID, const int32 FrameNumber, const int32 Step, FGuid InstigatorID)
 {
 	if (const TSharedPtr<FChaosVDScene> SceneToControlSharedPtr = SceneToControl.Pin())
 	{
 		if (ensure(LoadedRecording.IsValid()))
 		{
-			const EChaosVDFrameLoadState FrameState = LoadedRecording->GetFrameState(InTrackID, FrameNumber);
-			if (FrameState == EChaosVDFrameLoadState::Unknown)
+			const TSharedPtr<const TraceServices::IAnalysisSession> TraceSession = FChaosVDModule::Get().GetTraceManager()->GetSession(LoadedRecording->SessionName);
+			if (!ensure(TraceSession))
 			{
-				// Invalid Frame Number
-				return;
+				return;	
 			}
 
-			{
-				const TSharedPtr<const TraceServices::IAnalysisSession> TraceSession = FChaosVDModule::Get().GetTraceManager()->GetSession(LoadedRecording->SessionName);
-				if (!ensure(TraceSession))
-				{
-					return;	
-				}
+			// TODO: This will not cover a future case where the recording is not own/populated by Trace
+			// If in the future we decide implement the CVD format standalone with streaming support again,
+			// we will need to add a lock to the file. A feature that might need that is Recording Clips
+			TraceServices::FAnalysisSessionReadScope SessionReadScope(*TraceSession);
 
-				// TODO: This will not cover a future case where the recording is not own/populated by Trace
-				// If in the future we decide implement the CVD format standalone with streaming support again,
-				// we will need to add a lock to the file. A feature that might need that is Recording Clips
-				TraceServices::FAnalysisSessionReadScope SessionReadScope(*TraceSession);
+			const FChaosVDSolverFrameData* SolverFrameData = LoadedRecording->GetSolverFrameData(InTrackID, FrameNumber);
+			if (SolverFrameData && SolverFrameData->SolverSteps.IsValidIndex(Step))
+			{
+				SceneToControlSharedPtr->UpdateFromRecordedStepData(InTrackID, SolverFrameData->DebugName, SolverFrameData->SolverSteps[Step], *SolverFrameData);
+			}
 
-				const FChaosVDSolverFrameData* SolverFrameData = LoadedRecording->GetFrameForSolver(InTrackID, FrameNumber);
-				if (SolverFrameData && ensure(SolverFrameData->SolverSteps.IsValidIndex(Step)))
+			if (TrackInfoByIDMap* TrackInfoByID = TrackInfoPerType.Find(EChaosVDTrackType::Solver))
+			{
+				if (TSharedPtr<FChaosVDTrackInfo>* TrackInfo = TrackInfoByID->Find(InTrackID))
 				{
-					SceneToControlSharedPtr->UpdateFromRecordedStepData(InTrackID, SolverFrameData->DebugName, SolverFrameData->SolverSteps[Step], *SolverFrameData);
+					TrackInfo->Get()->CurrentFrame = FrameNumber;
+					TrackInfo->Get()->CurrentStep = Step;
+
+					OnTrackFrameUpdated().Broadcast(AsWeak(), TrackInfo->Get(), InstigatorID);
 				}
 			}
+		}
+	}
+	else
+	{
+		ensureMsgf(false, TEXT("GoToRecordedStep Called without a valid scene to control"));	
+	}
+}
 
-			if (int32* CurrentFrameNumber = CurrentFramePerTrack.Find(InTrackID))
-			{
-				 *CurrentFrameNumber = FrameNumber;
-			}
-			else
+void FChaosVDPlaybackController::GoToRecordedGameFrame(const int32 FrameNumber, FGuid InstigatorID)
+{
+	if (const TSharedPtr<FChaosVDScene> SceneToControlSharedPtr = SceneToControl.Pin())
+	{
+		if (ensure(LoadedRecording.IsValid()))
+		{
+			if (const FChaosVDGameFrameData* FrameData = LoadedRecording->GetGameFrameData(FrameNumber))
 			{
-				CurrentFramePerTrack.Add(InTrackID, FrameNumber);
+				if (TrackInfoByIDMap* TrackInfoByID = TrackInfoPerType.Find(EChaosVDTrackType::Game))
+				{
+					TArray<int32> AvailableSolversID;
+					LoadedRecording->GetAvailableSolverIDsAtGameFrameNumber(FrameNumber, AvailableSolversID);
+
+					SceneToControlSharedPtr->HandleEnterNewGameFrame(FrameNumber, AvailableSolversID);
+
+					for (int32 SolverID : AvailableSolversID)
+					{
+						constexpr int32 StepNumber = 0;
+
+						// When Scrubbing the timeline by game frames instead of solvers, try to go to the first solver frame on the first platform cycle of the game frame.
+						// Game Frames are not in sync with Solver Frames and Solver steps.
+						int32 SolverFrameNumber = LoadedRecording->GetLowestSolverFrameNumberAtCycle(SolverID, FrameData->FirstCycle);
+						GoToTrackFrame(InstigatorID, EChaosVDTrackType::Solver, SolverID, SolverFrameNumber, StepNumber);
+					}
+
+					TSharedPtr<FChaosVDTrackInfo>* TrackInfoPtrPtr = TrackInfoByID->Find(GameTrackID);
+					if (TSharedPtr<FChaosVDTrackInfo> TrackInfoSharedPtr = TrackInfoPtrPtr? *TrackInfoPtrPtr : nullptr)
+					{
+						TrackInfoSharedPtr->CurrentFrame = FrameNumber;
+						OnTrackFrameUpdated().Broadcast(AsWeak(), TrackInfoSharedPtr.Get(), InstigatorID);
+					}
+				}
 			}
+		}
+	}
+}
+
+void FChaosVDPlaybackController::GoToTrackFrame(FGuid InstigatorID, EChaosVDTrackType TrackType, int32 InTrackID, int32 FrameNumber, int32 Step)
+{
+	switch (TrackType)
+	{
+	case EChaosVDTrackType::Game:
+		GoToRecordedGameFrame(FrameNumber, InstigatorID);
+		break;
+	case EChaosVDTrackType::Solver:
+		GoToRecordedSolverStep(InTrackID, FrameNumber, Step, InstigatorID);
+		break;
+	default:
+		ensure(false);
+		break;
+	}
+}
 
-			if (int32* CurrentStepNumber = CurrentStepPerTrack.Find(InTrackID))
+int32 FChaosVDPlaybackController::GetTrackStepsAtFrame(EChaosVDTrackType TrackType, const int32 InTrackID, const int32 FrameNumber) const
+{
+	if (!LoadedRecording.IsValid())
+	{
+		return INDEX_NONE;
+	}
+	
+	switch (TrackType)
+	{
+		case EChaosVDTrackType::Game:
+			// Game Tracks do not have steps
+			return 0;
+			break;
+		case EChaosVDTrackType::Solver:
+		{
+			if (const FChaosVDSolverFrameData* FrameData = LoadedRecording->GetSolverFrameData(InTrackID, FrameNumber))
 			{
-				*CurrentStepNumber = Step;
+				return FrameData->SolverSteps.Num() > 0 ? FrameData->SolverSteps.Num() : INDEX_NONE;
 			}
 			else
 			{
-				CurrentStepPerTrack.Add(InTrackID, Step);
+				return INDEX_NONE;
 			}
+			break;
 		}
+	default:
+		return INDEX_NONE;
+		break;
 	}
-	else
+}
+
+int32 FChaosVDPlaybackController::GetTrackFramesNumber(EChaosVDTrackType TrackType, const int32 InTrackID) const
+{
+	if (!LoadedRecording.IsValid())
 	{
-		ensureMsgf(false, TEXT("GoToRecordedStep Called without a valid scene to control"));	
+		return INDEX_NONE;
+	}
+	
+	switch (TrackType)
+	{
+		case EChaosVDTrackType::Game:
+		{
+			// There is only one game track so no ID is needed
+			int32 GameFrames = LoadedRecording->GetAvailableGameFramesNumber();
+			return GameFrames > 0 ? GameFrames : INDEX_NONE;
+			break;
+		}
+
+		case EChaosVDTrackType::Solver:
+		{
+			int32 SolverFrames = LoadedRecording->GetAvailableSolverFramesNumber(InTrackID);
+			return  SolverFrames > 0 ? SolverFrames : INDEX_NONE;
+			break;
+		}
+		default:
+			return INDEX_NONE;
+			break;
 	}
 }
 
-int32 FChaosVDPlaybackController::GetStepsForFrame(const int32 InTrackID, const int32 FrameNumber) const
+int32 FChaosVDPlaybackController::ConvertCurrentFrameToOtherTrackFrame(FChaosVDTrackInfo FromTrack, FChaosVDTrackInfo ToTrack)
 {
-	if (LoadedRecording.IsValid())
+	// Each track is on a different "space time", because it's source data ticked at different rates when was recorded, and some start/end at different points on time
+	// But all the recorded frame data on all of them use Platform Cycles as timestamps
+	// This method wraps specialized methods in the recording object to convert between these spaces. For example Game frame 1500 could be frame 5 on a specific solver
+	// And Frame 5 of that solver could be frame 30 on another solver.
+
+	const bool bBothTracksHaveSameID = FromTrack.TrackID == ToTrack.TrackID;
+	const bool bBothTracksHaveSameType = FromTrack.TrackType == ToTrack.TrackType;
+	if (bBothTracksHaveSameType && bBothTracksHaveSameID)
+	{
+		return FromTrack.CurrentFrame;
+	}
+
+	switch (FromTrack.TrackType)
 	{
-		if (const FChaosVDSolverFrameData* FrameData = LoadedRecording->GetFrameForSolver(InTrackID, FrameNumber))
+	case EChaosVDTrackType::Game:
 		{
-			return FrameData->SolverSteps.Num();
+			// Convert from Game Frame to Solver Frame
+			return LoadedRecording->GetLowestSolverFrameNumberGameFrame(FromTrack.TrackID, FromTrack.CurrentFrame);
+			break;
 		}
-		else
+		
+	case EChaosVDTrackType::Solver:
 		{
-			return INDEX_NONE;
+			if (ToTrack.TrackType == EChaosVDTrackType::Solver)
+			{
+				ensureMsgf(false, TEXT("Frame conversion between solver tracks is not supported yet"));
+				return INDEX_NONE;
+			}
+			return LoadedRecording->GetLowestGameFrameAtSolverFrameNumber(FromTrack.TrackID, FromTrack.CurrentFrame);
+			break;
 		}
+	default:
+		ensure(false);
+		return INDEX_NONE;
+		break;
 	}
-
-	return	INDEX_NONE;
 }
 
-int32 FChaosVDPlaybackController::GetAvailableFramesNumber(const int32 InTrackID) const
+int32 FChaosVDPlaybackController::GetTrackCurrentFrame(EChaosVDTrackType TrackType, const int32 InTrackID) const
 {
-	if (!LoadedRecording.IsValid())
+	if (const TrackInfoByIDMap* TrackInfoByID = TrackInfoPerType.Find(TrackType))
 	{
-		return INDEX_NONE;
+		const TSharedPtr<FChaosVDTrackInfo>* TrackInfoPtrPtr = TrackInfoByID->Find(InTrackID);
+		if (const TSharedPtr<FChaosVDTrackInfo> TrackInfoSharedPtr = TrackInfoPtrPtr ? *TrackInfoPtrPtr : nullptr)
+		{
+			return TrackInfoSharedPtr->CurrentFrame;
+		}
 	}
 
-	return LoadedRecording->GetAvailableFramesNumber(InTrackID);
+	return INDEX_NONE;
 }
 
-int32 FChaosVDPlaybackController::GetAvailableSolversNumber() const
+int32 FChaosVDPlaybackController::GetTrackCurrentStep(EChaosVDTrackType TrackType, const int32 InTrackID) const
 {
-	if (!LoadedRecording.IsValid())
+	if (const TrackInfoByIDMap* TrackInfoByID = TrackInfoPerType.Find(TrackType))
 	{
-		return INDEX_NONE;
+		const TSharedPtr<FChaosVDTrackInfo>* TrackInfoPtrPtr = TrackInfoByID->Find(InTrackID);
+		if (const TSharedPtr<FChaosVDTrackInfo> TrackInfoSharedPtr = TrackInfoPtrPtr ? *TrackInfoPtrPtr : nullptr)
+		{
+			return TrackInfoSharedPtr->CurrentStep;
+		}
 	}
 
-	return LoadedRecording->GetAvailableSolvers().Num();
+	return INDEX_NONE;
 }
 
-int32 FChaosVDPlaybackController::GetActiveSolverTrackID() const
+const FChaosVDTrackInfo* FChaosVDPlaybackController::GetTrackInfo(EChaosVDTrackType TrackType, int32 TrackID)
 {
-	//For Now use the Last Solver if any
-	int32 SolverID = INDEX_NONE;
-	for (const TPair<int32, TArray<FChaosVDSolverFrameData>>& SolverData : LoadedRecording->GetAvailableSolvers())
+	if (const TrackInfoByIDMap* TrackInfoByID = TrackInfoPerType.Find(TrackType))
 	{
-		SolverID = SolverData.Key;
+		const TSharedPtr<FChaosVDTrackInfo>* TrackInfoPtrPtr = TrackInfoByID->Find(TrackID);
+		if (const TSharedPtr<FChaosVDTrackInfo> TrackInfoSharedPtr = TrackInfoPtrPtr ? *TrackInfoPtrPtr : nullptr)
+		{
+			return TrackInfoSharedPtr.Get();
+		}
 	}
 
-	return SolverID;
+	return nullptr;
 }
 
-int32 FChaosVDPlaybackController::GetCurrentFrame(const int32 InTrackID) const
+void FChaosVDPlaybackController::GetAvailableTracks(EChaosVDTrackType TrackType, TArray<TSharedPtr<FChaosVDTrackInfo>>& OutTrackInfo)
 {
-	if (const int32* FrameNumber = CurrentFramePerTrack.Find(InTrackID))
+	OutTrackInfo.Reset();
+	TrackInfoPerType.FindOrAdd(TrackType).GenerateValueArray(OutTrackInfo);
+}
+
+void FChaosVDPlaybackController::GetAvailableTrackInfosAtTrackFrame(EChaosVDTrackType TrackTypeToFind, TArray<TSharedPtr<FChaosVDTrackInfo>>& OutTrackInfo, const FChaosVDTrackInfo* TrackFrameInfo)
+{
+	OutTrackInfo.Reset();
+
+	if (!LoadedRecording.IsValid())
 	{
-		return *FrameNumber;
+		return;
 	}
 
-	return INDEX_NONE;
+	if (TrackFrameInfo == nullptr)
+	{
+		return;
+	}
+
+	int32 CorrectedFrameNumber = INDEX_NONE;
+	switch (TrackFrameInfo->TrackType)
+	{
+		case EChaosVDTrackType::Game:
+			{
+				CorrectedFrameNumber = TrackFrameInfo->CurrentFrame;
+			}
+			break;
+		case EChaosVDTrackType::Solver:
+			{
+				CorrectedFrameNumber = LoadedRecording->GetLowestGameFrameAtSolverFrameNumber(TrackFrameInfo->TrackID, TrackFrameInfo->CurrentFrame);
+				break;
+			}
+		default:
+			ensure(false);
+			break;
+	}
+
+	TArray<int32> AvailableSolversID;
+	LoadedRecording->GetAvailableSolverIDsAtGameFrameNumber(CorrectedFrameNumber, AvailableSolversID);
+	
+	TrackInfoByIDMap& TrackInfoMap = TrackInfoPerType.FindOrAdd(TrackTypeToFind);
+	for (int32 SolverID : AvailableSolversID)
+	{
+		OutTrackInfo.Add(TrackInfoMap.FindChecked(SolverID));
+	}
 }
 
-int32 FChaosVDPlaybackController::GetCurrentStep(const int32 InTrackID) const
+void FChaosVDPlaybackController::UpdateSolverTracksData()
 {
-	if (const int32* StepNumber = CurrentStepPerTrack.Find(InTrackID))
+	const TMap<int32, TArray<FChaosVDSolverFrameData>>& SolversByID = LoadedRecording->GetAvailableSolvers();
+	for (const TPair<int32, TArray<FChaosVDSolverFrameData>>& SolverIDPair : SolversByID)
 	{
-		return *StepNumber;
-	}
+		TSharedPtr<FChaosVDTrackInfo>& SolverTrackInfo = TrackInfoPerType[EChaosVDTrackType::Solver].FindOrAdd(SolverIDPair.Key);
 
-	return INDEX_NONE;
+		if (!SolverTrackInfo.IsValid())
+		{
+			SolverTrackInfo = MakeShared<FChaosVDTrackInfo>();
+			SolverTrackInfo->CurrentFrame = 0;
+			SolverTrackInfo->CurrentStep = 0;
+		};
+		
+		SolverTrackInfo->TrackID = SolverIDPair.Key;
+		SolverTrackInfo->MaxFrames = GetTrackFramesNumber(EChaosVDTrackType::Solver, SolverIDPair.Key);
+		SolverTrackInfo->TrackName = LoadedRecording->GetSolverName(SolverIDPair.Key);
+		SolverTrackInfo->TrackType = EChaosVDTrackType::Solver;
+	}
 }
 
 void FChaosVDPlaybackController::HandleCurrentRecordingUpdated()
 {
-	OnControllerUpdated().ExecuteIfBound(AsWeak());
-}
+	// These two tracks should always exist
+	TrackInfoPerType.FindOrAdd(EChaosVDTrackType::Game);
+	TrackInfoPerType.FindOrAdd(EChaosVDTrackType::Solver);
+
+	// Same for the Game Track, needs to always exist
+	TSharedPtr<FChaosVDTrackInfo>& GameTrackInfo = TrackInfoPerType[EChaosVDTrackType::Game].FindOrAdd(GameTrackID);
+	if (!GameTrackInfo.IsValid())
+	{
+		GameTrackInfo = MakeShared<FChaosVDTrackInfo>();
+		GameTrackInfo->CurrentFrame = 0;
+		GameTrackInfo->CurrentStep = 0;
+	};
 
+	GameTrackInfo->MaxFrames = LoadedRecording->GetAvailableGameFrames().Num();
+	GameTrackInfo->TrackType = EChaosVDTrackType::Game;
+
+	// Each time the recording is updated, populate or update the existing solver tracks data
+	UpdateSolverTracksData();
+
+	OnDataUpdated().Broadcast(AsWeak());
+}
diff --git a/Engine/Plugins/Experimental/ChaosVD/Source/ChaosVD/Private/ChaosVDPlaybackControllerObserver.cpp b/Engine/Plugins/Experimental/ChaosVD/Source/ChaosVD/Private/ChaosVDPlaybackControllerObserver.cpp
new file mode 100644
--- /dev/null
+++ b/Engine/Plugins/Experimental/ChaosVD/Source/ChaosVD/Private/ChaosVDPlaybackControllerObserver.cpp
@@ -0,0 +1,34 @@
+﻿// Copyright Epic Games, Inc. All Rights Reserved.
+
+#include "ChaosVDPlaybackControllerObserver.h"
+#include "ChaosVDPlaybackController.h"
+
+FChaosVDPlaybackControllerObserver::~FChaosVDPlaybackControllerObserver()
+{
+	if (const TSharedPtr<FChaosVDPlaybackController> CurrentPlaybackControllerPtr = PlaybackController.Pin())
+	{
+		CurrentPlaybackControllerPtr->OnDataUpdated().RemoveAll(this);
+		CurrentPlaybackControllerPtr->OnTrackFrameUpdated().RemoveAll(this);
+	}
+}
+
+void FChaosVDPlaybackControllerObserver::RegisterNewController(TWeakPtr<FChaosVDPlaybackController> NewController)
+{
+	if (PlaybackController != NewController)
+	{
+		if (const TSharedPtr<FChaosVDPlaybackController> CurrentPlaybackControllerPtr = PlaybackController.Pin())
+		{
+			CurrentPlaybackControllerPtr->OnDataUpdated().RemoveAll(this);
+			CurrentPlaybackControllerPtr->OnTrackFrameUpdated().RemoveAll(this);
+		}
+
+		PlaybackController = NewController;
+
+		if (const TSharedPtr<FChaosVDPlaybackController> NewPlaybackControllerPtr = PlaybackController.Pin())
+		{	
+			NewPlaybackControllerPtr->OnDataUpdated().AddRaw(this, &FChaosVDPlaybackControllerObserver::HandlePlaybackControllerDataUpdated);
+			NewPlaybackControllerPtr->OnTrackFrameUpdated().AddRaw(this, &FChaosVDPlaybackControllerObserver::HandleControllerTrackFrameUpdated);
+			HandlePlaybackControllerDataUpdated(PlaybackController);
+		}
+	}
+}
diff --git a/Engine/Plugins/Experimental/ChaosVD/Source/ChaosVD/Private/ChaosVDScene.cpp b/Engine/Plugins/Experimental/ChaosVD/Source/ChaosVD/Private/ChaosVDScene.cpp
--- a/Engine/Plugins/Experimental/ChaosVD/Source/ChaosVD/Private/ChaosVDScene.cpp
+++ b/Engine/Plugins/Experimental/ChaosVD/Source/ChaosVD/Private/ChaosVDScene.cpp
@@ -7,6 +7,8 @@
 #include "ChaosVDParticleActor.h"
 #include "Chaos/ImplicitObject.h"
 #include "ChaosVDRecording.h"
+#include "EditorActorFolders.h"
+#include "WorldPersistentFolders.h"
 #include "Engine/Engine.h"
 #include "Engine/LevelStreamingDynamic.h"
 #include "Engine/World.h"
@@ -59,7 +61,6 @@ void FChaosVDScene::DeInitialize()
 void FChaosVDScene::AddReferencedObjects(FReferenceCollector& Collector)
 {
 	Collector.AddReferencedObject(PhysicsVDWorld);
-	Collector.AddStableReferenceMap(ParticleVDInstancesByID);
 }
 
 void FChaosVDScene::UpdateFromRecordedStepData(const int32 SolverID, const FString& SolverName, const FChaosVDStepData& InRecordedStepData, const FChaosVDSolverFrameData& InFrameData)
@@ -67,13 +68,15 @@ void FChaosVDScene::UpdateFromRecordedStepData(const int32 SolverID, const FStri
 	TSet<int32> ParticlesIDsInRecordedStepData;
 	ParticlesIDsInRecordedStepData.Reserve(InRecordedStepData.RecordedParticles.Num());
 
+	FChaosVDParticlesByIDMap& SolverParticlesByID = ParticlesBySolverID.FindChecked(SolverID);
+
 	// Go over existing Particle VD Instances and update them or create them if needed 
 	for (const FChaosVDParticleDebugData& Particle : InRecordedStepData.RecordedParticles)
 	{
 		const int32 ParticleVDInstanceID = GetIDForRecordedParticleData(Particle);
 		ParticlesIDsInRecordedStepData.Add(ParticleVDInstanceID);
 
-		if (AChaosVDParticleActor** ExistingParticleVDInstancePtrPtr = ParticleVDInstancesByID.Find(ParticleVDInstanceID))
+		if (AChaosVDParticleActor** ExistingParticleVDInstancePtrPtr = SolverParticlesByID.Find(ParticleVDInstanceID))
 		{
 			if (AChaosVDParticleActor* ExistingParticleVDInstancePtr = *ExistingParticleVDInstancePtrPtr)
 			{
@@ -95,7 +98,7 @@ void FChaosVDScene::UpdateFromRecordedStepData(const int32 SolverID, const FStri
 				NewParticleVDInstance->SetFolderPath(FolderPath);
 
 				// TODO: Precalculate the max num of entries we would see in the loaded file, and use that number to pre-allocate this map
-				ParticleVDInstancesByID.Add(ParticleVDInstanceID, NewParticleVDInstance);
+				SolverParticlesByID.Add(ParticleVDInstanceID, NewParticleVDInstance);
 			}
 			else
 			{
@@ -105,8 +108,10 @@ void FChaosVDScene::UpdateFromRecordedStepData(const int32 SolverID, const FStri
 		}
 	}
 
+	// TODO: This will not work once the recording has what changes between solver steps
+	// So it needs to be update to remove particles based on a Particles Removed event
 	int32 AmountRemoved = 0;
-	for (TMap<int32, AChaosVDParticleActor*>::TIterator RemoveIterator = ParticleVDInstancesByID.CreateIterator(); RemoveIterator; ++RemoveIterator)
+	for (FChaosVDParticlesByIDMap::TIterator RemoveIterator = SolverParticlesByID.CreateIterator(); RemoveIterator; ++RemoveIterator)
 	{
 		if (!ParticlesIDsInRecordedStepData.Contains(RemoveIterator.Key()))
 		{
@@ -129,17 +134,62 @@ void FChaosVDScene::HandleNewGeometryData(const TSharedPtr<const Chaos::FImplici
 	NewGeometryAvailableDelegate.Broadcast(GeometryData, GeometryID);
 }
 
+void FChaosVDScene::HandleEnterNewGameFrame(int32 FrameNumber, const TArray<int32>& AvailableSolversIds)
+{
+	// Currently the particle actors from all the solvers are in the same level, and we manage them by keeping track
+	// of to which solvers they belong using maps.
+	// Using Level instead or a Sub ChaosVDScene could be a better solution
+	// I'm intentionally not making that change right now until the "level streaming" solution for the tool is defined
+	// As that would impose restriction on how levels could be used. For now the map approach is simpler and will be easier to refactor later on.
+
+	TSet<int32> AvailableSolversSet;
+	AvailableSolversSet.Reserve(AvailableSolversIds.Num());
+
+	for (int32 SolverID : AvailableSolversIds)
+	{
+		AvailableSolversSet.Add(SolverID);
+		if (!ParticlesBySolverID.Contains(SolverID))
+		{
+			ParticlesBySolverID.Add(SolverID);
+		}
+	}
+
+	int32 AmountRemoved = 0;
+	for (TMap<int32, FChaosVDParticlesByIDMap>::TIterator RemoveIterator = ParticlesBySolverID.CreateIterator(); RemoveIterator; ++RemoveIterator)
+	{
+		if (!AvailableSolversSet.Contains(RemoveIterator.Key()))
+		{
+			for (const TPair<int32, AChaosVDParticleActor*>& ParticleVDInstanceWithID : RemoveIterator.Value())
+			{
+				PhysicsVDWorld->DestroyActor(ParticleVDInstanceWithID.Value);
+			}
+
+			RemoveIterator.RemoveCurrent();
+
+			AmountRemoved++;
+		}
+	}
+
+	if (AmountRemoved > 0)
+	{
+		CollectGarbage(GARBAGE_COLLECTION_KEEPFLAGS);
+	}
+}
+
 void FChaosVDScene::CleanUpScene()
 {
 	if (PhysicsVDWorld)
 	{
-		for (const TPair<int32, AChaosVDParticleActor*>& ParticleVDInstanceWithID : ParticleVDInstancesByID)
+		for (const TPair<int32, FChaosVDParticlesByIDMap>& SolverParticleVDInstanceWithID : ParticlesBySolverID)
 		{
-			PhysicsVDWorld->DestroyActor(ParticleVDInstanceWithID.Value);
+			for (const TPair<int32, AChaosVDParticleActor*>& ParticleVDInstanceWithID : SolverParticleVDInstanceWithID.Value)
+			{
+				PhysicsVDWorld->DestroyActor(ParticleVDInstanceWithID.Value);
+			}
 		}
 	}
 
-	ParticleVDInstancesByID.Reset();
+	ParticlesBySolverID.Reset();
 }
 
 AChaosVDParticleActor* FChaosVDScene::SpawnParticleFromRecordedData(const FChaosVDParticleDebugData& InParticleData, const FChaosVDSolverFrameData& InFrameData)
diff --git a/Engine/Plugins/Experimental/ChaosVD/Source/ChaosVD/Private/ChaosVDSolversTracksTab.cpp b/Engine/Plugins/Experimental/ChaosVD/Source/ChaosVD/Private/ChaosVDSolversTracksTab.cpp
new file mode 100644
--- /dev/null
+++ b/Engine/Plugins/Experimental/ChaosVD/Source/ChaosVD/Private/ChaosVDSolversTracksTab.cpp
@@ -0,0 +1,32 @@
+﻿// Copyright Epic Games, Inc. All Rights Reserved.
+
+#include "ChaosVDSolversTracksTab.h"
+
+#include "ChaosVDEngine.h"
+#include "ChaosVDStyle.h"
+#include "Widgets/Docking/SDockTab.h"
+#include "Widgets/SChaosVDMainTab.h"
+#include "Widgets/SChaosVDSolverTracks.h"
+
+#define LOCTEXT_NAMESPACE "ChaosVisualDebugger"
+
+TSharedRef<SDockTab> FChaosVDSolversTracksTab::HandleTabSpawned(const FSpawnTabArgs& Args)
+{
+	TSharedRef<SDockTab> ViewportTab =
+	SNew(SDockTab)
+	.TabRole(ETabRole::PanelTab)
+	.Label(LOCTEXT("SolverTracksTabLabel", "Available Solvers"))
+	.ToolTipText(LOCTEXT("SolverTracksTabToolTip", "Playback controls for the available solvers on the current Frame"));
+
+	ViewportTab->SetContent
+	(
+		// TODO: Handle Null cases to not crash the Editor
+		SAssignNew(SolverTracksWidget, SChaosVDSolverTracks, OwningTabWidget->GetChaosVDEngineInstance()->GetPlaybackController())
+	);
+
+	ViewportTab->SetTabIcon(FChaosVDStyle::Get().GetBrush("TabIconPlaybackViewport"));
+
+	return ViewportTab;
+}
+
+#undef LOCTEXT_NAMESPACE
diff --git a/Engine/Plugins/Experimental/ChaosVD/Source/ChaosVD/Private/ChaosVDTabsIDs.cpp b/Engine/Plugins/Experimental/ChaosVD/Source/ChaosVD/Private/ChaosVDTabsIDs.cpp
--- a/Engine/Plugins/Experimental/ChaosVD/Source/ChaosVD/Private/ChaosVDTabsIDs.cpp
+++ b/Engine/Plugins/Experimental/ChaosVD/Source/ChaosVD/Private/ChaosVDTabsIDs.cpp
@@ -7,4 +7,5 @@ const FName FChaosVDTabID::ChaosVisualDebuggerTab = TEXT("ChaosVisualDebuggerTab
 const FName FChaosVDTabID::PlaybackViewport = TEXT("PlaybackViewport");
 const FName FChaosVDTabID::WorldOutliner = TEXT("WorldOutliner");
 const FName FChaosVDTabID::DetailsPanel = TEXT("DetailsPanel");
-const FName FChaosVDTabID::OutputLog = TEXT("ChaosVDOutputLog");
\ No newline at end of file
+const FName FChaosVDTabID::OutputLog = TEXT("ChaosVDOutputLog");
+const FName FChaosVDTabID::SolversTrack = TEXT("SolversTrack");
\ No newline at end of file
diff --git a/Engine/Plugins/Experimental/ChaosVD/Source/ChaosVD/Private/Trace/ChaosVDTraceAnalyzer.cpp b/Engine/Plugins/Experimental/ChaosVD/Source/ChaosVD/Private/Trace/ChaosVDTraceAnalyzer.cpp
--- a/Engine/Plugins/Experimental/ChaosVD/Source/ChaosVD/Private/Trace/ChaosVDTraceAnalyzer.cpp
+++ b/Engine/Plugins/Experimental/ChaosVD/Source/ChaosVD/Private/Trace/ChaosVDTraceAnalyzer.cpp
@@ -23,6 +23,9 @@ void FChaosVDTraceAnalyzer::OnAnalysisBegin(const FOnAnalysisContext& Context)
 	Builder.RouteEvent(RouteId_ChaosVDBinaryDataEnd, "ChaosVDLogger", "ChaosVDBinaryDataEnd");
 	Builder.RouteEvent(RouteId_ChaosVDSolverSimulationSpace, "ChaosVDLogger", "ChaosVDSolverSimulationSpace");
 
+	Builder.RouteEvent(RouteId_BeginFrame, "Misc", "BeginFrame");
+	Builder.RouteEvent(RouteId_EndFrame, "Misc", "EndFrame");
+
 	TraceServices::FAnalysisSessionEditScope _(Session);
 	ChaosVDTraceProvider->CreateRecordingInstanceForSession(Session.GetName());
 }
@@ -42,18 +45,44 @@ bool FChaosVDTraceAnalyzer::OnEvent(uint16 RouteId, EStyle Style, const FOnEvent
 	
 	switch (RouteId)
 	{
+	case RouteId_BeginFrame:
+		{
+			uint8 FrameType = EventData.GetValue<uint8>("FrameType");
+			if (static_cast<ETraceFrameType>(FrameType) == TraceFrameType_Game)
+			{
+				FChaosVDGameFrameData FrameData;
+				FrameData.FirstCycle = EventData.GetValue<uint64>("Cycle");
+				ChaosVDTraceProvider->AddGameFrame(MoveTemp(FrameData));
+			}
+
+			break;
+		}
+
+	case RouteId_EndFrame:
+		{
+			uint8 FrameType = EventData.GetValue<uint8>("FrameType");
+			if (static_cast<ETraceFrameType>(FrameType) == TraceFrameType_Game)
+			{
+				if (FChaosVDGameFrameData* CurrentFrameData = ChaosVDTraceProvider->GetLastGameFrame())
+				{
+					CurrentFrameData->LastCycle = EventData.GetValue<uint64>("Cycle");
+				}
+			}
+			break;
+		}
 	case RouteId_ChaosVDSolverFrameStart:
 		{
 			FChaosVDSolverFrameData NewFrameData;
 
 			NewFrameData.SolverID = EventData.GetValue<int32>("SolverID");
+			NewFrameData.FrameCycle = EventData.GetValue<uint64>("Cycle");
 
 			FWideStringView DebugNameView;
 			EventData.GetString("DebugName", DebugNameView);
 			NewFrameData.DebugName = DebugNameView;
 
 			// Add an empty frame. It will be filled out by the solver trace events
-			ChaosVDTraceProvider->AddFrame(NewFrameData.SolverID, MoveTemp(NewFrameData));
+			ChaosVDTraceProvider->AddSolverFrame(NewFrameData.SolverID, MoveTemp(NewFrameData));
 	
 			break;
 		}
@@ -66,7 +95,7 @@ bool FChaosVDTraceAnalyzer::OnEvent(uint16 RouteId, EStyle Style, const FOnEvent
 			const int32 SolverID = EventData.GetValue<int32>("SolverID");
 
 			// This can be null if the recording started Mid-Frame. In this case we just discard the data for now
-			if (FChaosVDSolverFrameData* FrameData  = ChaosVDTraceProvider->GetLastFrame(SolverID))
+			if (FChaosVDSolverFrameData* FrameData  = ChaosVDTraceProvider->GetLastSolverFrame(SolverID))
 			{
 				// Add an empty step. It will be filled out by the particle (and later on other objects/elements) events
 				FrameData->SolverSteps.AddDefaulted();
@@ -86,7 +115,7 @@ bool FChaosVDTraceAnalyzer::OnEvent(uint16 RouteId, EStyle Style, const FOnEvent
 			ReadParticleDataFromEvent(EventData, ParticleData);
 
 			// This can be null if the recording started Mid-Frame. In this case we just discard the data for now
-			if (FChaosVDSolverFrameData* FrameData = ChaosVDTraceProvider->GetLastFrame(SolverID))
+			if (FChaosVDSolverFrameData* FrameData = ChaosVDTraceProvider->GetLastSolverFrame(SolverID))
 			{
 				if (FrameData->SolverSteps.Num() > 0)
 				{
@@ -146,7 +175,7 @@ bool FChaosVDTraceAnalyzer::OnEvent(uint16 RouteId, EStyle Style, const FOnEvent
 		Rotation.W = EventData.GetValue<float>("RotationW");
 
 		// This can be null if the recording started Mid-Frame. In this case we just discard the data for now
-		if (FChaosVDSolverFrameData* FrameData = ChaosVDTraceProvider->GetLastFrame(SolverID))
+		if (FChaosVDSolverFrameData* FrameData = ChaosVDTraceProvider->GetLastSolverFrame(SolverID))
 		{
 			FrameData->SimulationTransform.SetLocation(Position);
 			FrameData->SimulationTransform.SetRotation(Rotation);
diff --git a/Engine/Plugins/Experimental/ChaosVD/Source/ChaosVD/Private/Trace/ChaosVDTraceProvider.cpp b/Engine/Plugins/Experimental/ChaosVD/Source/ChaosVD/Private/Trace/ChaosVDTraceProvider.cpp
--- a/Engine/Plugins/Experimental/ChaosVD/Source/ChaosVD/Private/Trace/ChaosVDTraceProvider.cpp
+++ b/Engine/Plugins/Experimental/ChaosVD/Source/ChaosVD/Private/Trace/ChaosVDTraceProvider.cpp
@@ -26,7 +26,7 @@ void FChaosVDTraceProvider::DeleteRecordingInstanceForSession()
 	InternalRecording.Reset();
 }
 
-void FChaosVDTraceProvider::AddFrame(const int32 InSolverID, FChaosVDSolverFrameData&& FrameData)
+void FChaosVDTraceProvider::AddSolverFrame(const int32 InSolverID, FChaosVDSolverFrameData&& FrameData)
 {
 	if (InternalRecording.IsValid())
 	{
@@ -34,26 +34,61 @@ void FChaosVDTraceProvider::AddFrame(const int32 InSolverID, FChaosVDSolverFrame
 	}
 }
 
-FChaosVDSolverFrameData* FChaosVDTraceProvider::GetFrame(const int32 InSolverID, const int32 FrameNumber) const
+void FChaosVDTraceProvider::AddGameFrame(FChaosVDGameFrameData&& FrameData)
 {
-	return InternalRecording.IsValid() ? InternalRecording->GetFrameForSolver(InSolverID, FrameNumber) : nullptr;
+	if (InternalRecording.IsValid())
+	{
+		// In PIE, we can have a lot of empty frames at the beginning, so we discard them here
+		if (InternalRecording->GetAvailableSolvers().IsEmpty())
+		{
+			if (FChaosVDGameFrameData* GameFrame = InternalRecording->GetLastGameFrameData())
+			{
+				GameFrame = &FrameData;
+				return;
+			}
+		}
+
+		InternalRecording->AddGameFrameData(MoveTemp(FrameData));
+	}
+}
+
+FChaosVDSolverFrameData* FChaosVDTraceProvider::GetSolverFrame(const int32 InSolverID, const int32 FrameNumber) const
+{
+	return InternalRecording.IsValid() ? InternalRecording->GetSolverFrameData(InSolverID, FrameNumber) : nullptr;
 }
 
-FChaosVDSolverFrameData* FChaosVDTraceProvider::GetLastFrame(const int32 InSolverID) const
+FChaosVDSolverFrameData* FChaosVDTraceProvider::GetLastSolverFrame(const int32 InSolverID) const
 {
-	if (InternalRecording.IsValid() && InternalRecording->GetAvailableFramesNumber(InSolverID) > 0)
+	if (InternalRecording.IsValid() && InternalRecording->GetAvailableSolverFramesNumber(InSolverID) > 0)
 	{
-		const int32 AvailableFramesNumber = InternalRecording->GetAvailableFramesNumber(InSolverID);
+		const int32 AvailableFramesNumber = InternalRecording->GetAvailableSolverFramesNumber(InSolverID);
 
 		if (AvailableFramesNumber != INDEX_NONE)
 		{
-			return GetFrame(InSolverID, InternalRecording->GetAvailableFramesNumber(InSolverID) - 1);
+			return GetSolverFrame(InSolverID, InternalRecording->GetAvailableSolverFramesNumber(InSolverID) - 1);
 		}
 	}
 
 	return nullptr;
 }
 
+FChaosVDGameFrameData* FChaosVDTraceProvider::GetSolverFrame(uint64 FrameStartCycle) const
+{
+	FChaosVDGameFrameData* FoundFrameData = nullptr;
+
+	if (InternalRecording.IsValid() && InternalRecording->GetAvailableGameFrames().Num() > 0)
+	{
+		return InternalRecording->GetGameFrameDataAtCycle(FrameStartCycle);
+	}
+	
+	return FoundFrameData;
+}
+
+FChaosVDGameFrameData* FChaosVDTraceProvider::GetLastGameFrame() const
+{
+	return InternalRecording.IsValid() ? InternalRecording->GetLastGameFrameData() : nullptr;
+}
+
 FChaosVDBinaryDataContainer& FChaosVDTraceProvider::FindOrAddUnprocessedData(const int32 DataID)
 {
 	if (const TSharedPtr<FChaosVDBinaryDataContainer>* UnprocessedData = UnprocessedDataByID.Find(DataID))
diff --git a/Engine/Plugins/Experimental/ChaosVD/Source/ChaosVD/Private/Widgets/SChaosVDMainTab.cpp b/Engine/Plugins/Experimental/ChaosVD/Source/ChaosVD/Private/Widgets/SChaosVDMainTab.cpp
--- a/Engine/Plugins/Experimental/ChaosVD/Source/ChaosVD/Private/Widgets/SChaosVDMainTab.cpp
+++ b/Engine/Plugins/Experimental/ChaosVD/Source/ChaosVD/Private/Widgets/SChaosVDMainTab.cpp
@@ -7,6 +7,7 @@
 #include "ChaosVDOutputLogTab.h"
 #include "ChaosVDPlaybackController.h"
 #include "ChaosVDPlaybackViewportTab.h"
+#include "ChaosVDSolversTracksTab.h"
 #include "ChaosVDTabsIDs.h"
 #include "ChaosVDWorldOutlinerTab.h"
 #include "DesktopPlatformModule.h"
@@ -34,6 +35,7 @@ void SChaosVDMainTab::Construct(const FArguments& InArgs, TSharedPtr<FChaosVDEng
 	ObjectDetailsTab = MakeShared<FChaosVDObjectDetailsTab>(FChaosVDTabID::DetailsPanel, TabManager, this);
 	OutputLogTab = MakeShared<FChaosVDOutputLogTab>(FChaosVDTabID::OutputLog, TabManager, this);
 	PlaybackViewportTab = MakeShared<FChaosVDPlaybackViewportTab>(FChaosVDTabID::PlaybackViewport, TabManager, this);
+	SolversTracksTab = MakeShared<FChaosVDSolversTracksTab>(FChaosVDTabID::SolversTrack, TabManager, this);
 
 	GenerateMainWindowMenu();
 
@@ -138,8 +140,11 @@ void SChaosVDMainTab::Construct(const FArguments& InArgs, TSharedPtr<FChaosVDEng
 		.Padding(FMargin(0.0f,5.0f,0.0f,0.0f))
 		[
 			TabManager->RestoreFrom(GenerateMainLayout(), TabManager->GetOwnerTab()->GetParentWindow()).ToSharedRef()
+			
 		]
 	];
+
+	TabManager->TryInvokeTab(FChaosVDTabID::SolversTrack);
 }
 
 TSharedRef<FTabManager::FLayout> SChaosVDMainTab::GenerateMainLayout()
@@ -166,6 +171,7 @@ TSharedRef<FTabManager::FLayout> SChaosVDMainTab::GenerateMainLayout()
 				(
 					FTabManager::NewStack()
 					->SetSizeCoefficient(0.3f)
+					->AddTab(FChaosVDTabID::SolversTrack, ETabState::OpenedTab)
 					->AddTab(FChaosVDTabID::OutputLog, ETabState::OpenedTab)
 				)
 			)
diff --git a/Engine/Plugins/Experimental/ChaosVD/Source/ChaosVD/Private/Widgets/SChaosVDPlaybackViewport.cpp b/Engine/Plugins/Experimental/ChaosVD/Source/ChaosVD/Private/Widgets/SChaosVDPlaybackViewport.cpp
--- a/Engine/Plugins/Experimental/ChaosVD/Source/ChaosVD/Private/Widgets/SChaosVDPlaybackViewport.cpp
+++ b/Engine/Plugins/Experimental/ChaosVD/Source/ChaosVD/Private/Widgets/SChaosVDPlaybackViewport.cpp
@@ -3,6 +3,7 @@
 #include "Widgets/SChaosVDPlaybackViewport.h"
 
 #include "ChaosVDPlaybackController.h"
+#include "ChaosVDScene.h"
 #include "Framework/Application/SlateApplication.h"
 #include "LevelEditorViewport.h"
 #include "Slate/SceneViewport.h"
@@ -14,11 +15,6 @@
 
 SChaosVDPlaybackViewport::~SChaosVDPlaybackViewport()
 {
-	if (const TSharedPtr<FChaosVDPlaybackController> CurrentPlaybackControllerPtr = PlaybackController.Pin())
-	{
-		CurrentPlaybackControllerPtr->OnControllerUpdated().Unbind();
-	}
-
 	LevelViewportClient->Viewport = nullptr;
 	LevelViewportClient.Reset();
 }
@@ -78,54 +74,29 @@ void SChaosVDPlaybackViewport::Construct(const FArguments& InArgs, const UWorld*
 		.Padding(16.0f, 16.0f, 16.0f, 16.0f)
 		.FillHeight(0.1f)
 		[
-			SNew(SHorizontalBox)
-			+SHorizontalBox::Slot()
-			.FillWidth(0.7f)
-			[
 			SNew(SVerticalBox)
-				+SVerticalBox::Slot()
-				.AutoHeight()
-				.Padding(0.0f, 0.0f, 0.0f, 2.0f)
-				[
-					SNew(STextBlock)
-					.Justification(ETextJustify::Center)
-					.Text(LOCTEXT("PlaybackViewportWidgetPhysicsFramesLabel", "Physics Frames" ))
-				]
-				+SVerticalBox::Slot()
-				[
-					SAssignNew(FramesTimelineWidget, SChaosVDTimelineWidget)
-						.HidePlayStopButtons(false)
-						.OnFrameChanged_Raw(this, &SChaosVDPlaybackViewport::OnFrameSelectionUpdated)
-						.MaxFrames(0)
-				]
+			+SVerticalBox::Slot()
+			.AutoHeight()
+			.Padding(0.0f, 0.0f, 0.0f, 2.0f)
+			[
+				SNew(STextBlock)
+				.Justification(ETextJustify::Center)
+				.Text(LOCTEXT("PlaybackViewportWidgetPhysicsFramesLabel", "Game Frames" ))
 			]
-			+SHorizontalBox::Slot()
-			.FillWidth(0.3f)
+			+SVerticalBox::Slot()
 			[
-				SNew(SVerticalBox)
-				+SVerticalBox::Slot()
-				.AutoHeight()
-				.Padding(0.0f, 0.0f, 0.0f, 2.0f)
-				[
-					SNew(STextBlock)
-					.Justification(ETextJustify::Center)
-					.Text(LOCTEXT("PlaybackViewportWidgetStepsLabel", "Solver Steps" ))
-				]
-				+SVerticalBox::Slot()
-				[
-					SAssignNew(StepsTimelineWidget, SChaosVDTimelineWidget)
-					.HidePlayStopButtons(true)
-					.OnFrameChanged_Raw(this, &SChaosVDPlaybackViewport::OnStepSelectionUpdated)
-					.MaxFrames(0)
-				]	
+				SAssignNew(GameFramesTimelineWidget, SChaosVDTimelineWidget)
+				.HidePlayStopButtons(false)
+				.OnFrameChanged_Raw(this, &SChaosVDPlaybackViewport::OnFrameSelectionUpdated)
+				.MaxFrames(0)
 			]
-		]
+		]	
 	];
 
 	RegisterNewController(InPlaybackController);
 }
 
-void SChaosVDPlaybackViewport::OnPlaybackControllerUpdated(TWeakPtr<FChaosVDPlaybackController> InController)
+void SChaosVDPlaybackViewport::HandlePlaybackControllerDataUpdated(TWeakPtr<FChaosVDPlaybackController> InController)
 {
 	if (PlaybackController != InController)
 	{
@@ -135,43 +106,57 @@ void SChaosVDPlaybackViewport::OnPlaybackControllerUpdated(TWeakPtr<FChaosVDPlay
 	const TSharedPtr<FChaosVDPlaybackController> ControllerSharedPtr = PlaybackController.Pin();
 	if (ControllerSharedPtr.IsValid() && ControllerSharedPtr->IsRecordingLoaded())
 	{
-		const int32 TrackID = ControllerSharedPtr->GetActiveSolverTrackID();
-		
-		const int32 AvailableFrames = ControllerSharedPtr->GetAvailableFramesNumber(TrackID);
-		const int32 AvailableSteps = ControllerSharedPtr->GetStepsForFrame(TrackID, ControllerSharedPtr->GetCurrentFrame(TrackID));
-
-		// Max is inclusive and we use this to request as the index on the recorded frames/steps arrays so we need to -1 to the available frames/steps
-		FramesTimelineWidget->UpdateMinMaxValue(0, AvailableFrames != INDEX_NONE ? AvailableFrames -1  : 0);
-
-		//TODO: This will show steps 0/0 if only one step is recorded, we need to add a way to override that functionality
-		// or just set the slider to start from 1 and handle the offset later 
-		StepsTimelineWidget->UpdateMinMaxValue(0,AvailableSteps != INDEX_NONE ? AvailableSteps -1 : 0);
+		if (const FChaosVDTrackInfo* TrackInfo = ControllerSharedPtr->GetTrackInfo(EChaosVDTrackType::Game, FChaosVDPlaybackController::GameTrackID))
+		{
+			// Max is inclusive and we use this to request as the index on the recorded frames/steps arrays so we need to -1 to the available frames/steps
+			GameFramesTimelineWidget->UpdateMinMaxValue(0, TrackInfo->MaxFrames != INDEX_NONE ? TrackInfo->MaxFrames -1  : 0);
+		}
 	}
 	else
 	{
-		FramesTimelineWidget->UpdateMinMaxValue(0, 0);
-		FramesTimelineWidget->ResetTimeline();
-		StepsTimelineWidget->UpdateMinMaxValue(0,0);
-		StepsTimelineWidget->ResetTimeline();
+		GameFramesTimelineWidget->UpdateMinMaxValue(0, 0);
+		GameFramesTimelineWidget->ResetTimeline();
 	}
 
 	LevelViewportClient->bNeedsRedraw = true;
 }
 
+void SChaosVDPlaybackViewport::HandleControllerTrackFrameUpdated(TWeakPtr<FChaosVDPlaybackController> InController, const FChaosVDTrackInfo* UpdatedTrackInfo, FGuid InstigatorGuid)
+{
+	if (InstigatorGuid == GetInstigatorID())
+	{
+		// Ignore the update if we initiated it
+		return;
+	}
+
+	GameFramesTimelineWidget->SetCurrentTimelineFrame(UpdatedTrackInfo->CurrentFrame, EChaosVDSetTimelineFrameFlags::None);
+}
+
+void SChaosVDPlaybackViewport::OnPlaybackSceneUpdated()
+{
+	LevelViewportClient->bNeedsRedraw = true;
+}
+
 void SChaosVDPlaybackViewport::RegisterNewController(TWeakPtr<FChaosVDPlaybackController> NewController)
 {
 	if (PlaybackController != NewController)
 	{
 		if (const TSharedPtr<FChaosVDPlaybackController> CurrentPlaybackControllerPtr = PlaybackController.Pin())
 		{
-			CurrentPlaybackControllerPtr->OnControllerUpdated().Unbind();
+			if (TSharedPtr<FChaosVDScene> ScenePtr = CurrentPlaybackControllerPtr->GetControllerScene().Pin())
+			{
+				ScenePtr->OnSceneUpdated().RemoveAll(this);
+			}
 		}
 
-		PlaybackController = NewController;
+		FChaosVDPlaybackControllerObserver::RegisterNewController(NewController);
 
 		if (const TSharedPtr<FChaosVDPlaybackController> NewPlaybackControllerPtr = PlaybackController.Pin())
 		{
-			NewPlaybackControllerPtr->OnControllerUpdated().BindRaw(this, &SChaosVDPlaybackViewport::OnPlaybackControllerUpdated);
+			if (TSharedPtr<FChaosVDScene> ScenePtr = NewPlaybackControllerPtr->GetControllerScene().Pin())
+			{
+				ScenePtr->OnSceneUpdated().AddRaw(this, &SChaosVDPlaybackViewport::OnPlaybackSceneUpdated);
+			}
 		}
 	}
 }
@@ -180,27 +165,9 @@ void SChaosVDPlaybackViewport::OnFrameSelectionUpdated(int32 NewFrameIndex) cons
 {
 	if (const TSharedPtr<FChaosVDPlaybackController> PlaybackControllerPtr = PlaybackController.Pin())
 	{
-		const int32 TrackID = PlaybackControllerPtr->GetActiveSolverTrackID();
-		
-		const int32 AvailableSteps = PlaybackControllerPtr->GetStepsForFrame(TrackID, PlaybackControllerPtr->GetCurrentFrame(TrackID));
-		StepsTimelineWidget->UpdateMinMaxValue(0,AvailableSteps != INDEX_NONE ? AvailableSteps -1 : 0);
-
-		// On Frame updates, always use Step 0
-		StepsTimelineWidget->SetCurrentTimelineFrame(0);
-		PlaybackControllerPtr->GoToRecordedStep(TrackID, NewFrameIndex, 0);
-
-		LevelViewportClient->bNeedsRedraw = true;
-	}
-}
+		constexpr int32 StepNumber = 0;
+		PlaybackControllerPtr->GoToTrackFrame(GetInstigatorID(), EChaosVDTrackType::Game, FChaosVDPlaybackController::GameTrackID, NewFrameIndex, StepNumber);
 
-void SChaosVDPlaybackViewport::OnStepSelectionUpdated(int32 NewStepIndex) const
-{
-	if (const TSharedPtr<FChaosVDPlaybackController> PlaybackControllerPtr = PlaybackController.Pin())
-	{
-		const int32 TrackID = PlaybackControllerPtr->GetActiveSolverTrackID();
-		// On Steps updates. Always use the current Frame
-		PlaybackControllerPtr->GoToRecordedStep(TrackID, PlaybackControllerPtr->GetCurrentFrame(TrackID), NewStepIndex);
-	
 		LevelViewportClient->bNeedsRedraw = true;
 	}
 }
diff --git a/Engine/Plugins/Experimental/ChaosVD/Source/ChaosVD/Private/Widgets/SChaosVDSolverPlaybackControls.cpp b/Engine/Plugins/Experimental/ChaosVD/Source/ChaosVD/Private/Widgets/SChaosVDSolverPlaybackControls.cpp
new file mode 100644
--- /dev/null
+++ b/Engine/Plugins/Experimental/ChaosVD/Source/ChaosVD/Private/Widgets/SChaosVDSolverPlaybackControls.cpp
@@ -0,0 +1,156 @@
+﻿// Copyright Epic Games, Inc. All Rights Reserved.
+
+#include "Widgets/SChaosVDSolverPlaybackControls.h"
+
+#include "ChaosVDPlaybackController.h"
+#include "Widgets/SBoxPanel.h"
+#include "Widgets/SChaosVDPlaybackViewport.h"
+#include "Widgets/SChaosVDTimelineWidget.h"
+#include "Widgets/Text/STextBlock.h"
+
+#define LOCTEXT_NAMESPACE "ChaosVisualDebugger"
+
+void SChaosVDSolverPlaybackControls::Construct(const FArguments& InArgs, int32 InSolverID, const TWeakPtr<FChaosVDPlaybackController>& InPlaybackController)
+{
+	SolverID = InSolverID;
+
+	ChildSlot
+	[
+		SNew(SVerticalBox)
+		// Playback controls
+		// TODO: Now that the tool is In-Editor, see if we can/is worth use the Sequencer widgets
+		// instead of these custom ones
+		+SVerticalBox::Slot()
+		[
+			SNew(SHorizontalBox)
+			+SHorizontalBox::Slot()
+			.FillWidth(0.7f)
+			[
+			SNew(SVerticalBox)
+				+SVerticalBox::Slot()
+				.AutoHeight()
+				.Padding(0.0f, 0.0f, 0.0f, 2.0f)
+				[
+					SNew(STextBlock)
+					.Justification(ETextJustify::Center)
+					.Text(LOCTEXT("PlaybackViewportWidgetPhysicsFramesLabel", "Physics Frames" ))
+				]
+				+SVerticalBox::Slot()
+				[
+					SAssignNew(FramesTimelineWidget, SChaosVDTimelineWidget)
+						.HidePlayStopButtons(false)
+						.OnFrameChanged_Raw(this, &SChaosVDSolverPlaybackControls::OnFrameSelectionUpdated)
+						.MaxFrames(0)
+				]
+			]
+			+SHorizontalBox::Slot()
+			.FillWidth(0.3f)
+			[
+				SNew(SVerticalBox)
+				+SVerticalBox::Slot()
+				.AutoHeight()
+				.Padding(0.0f, 0.0f, 0.0f, 2.0f)
+				[
+					SNew(STextBlock)
+					.Justification(ETextJustify::Center)
+					.Text(LOCTEXT("PlaybackViewportWidgetStepsLabel", "Solver Steps" ))
+				]
+				+SVerticalBox::Slot()
+				[
+					SAssignNew(StepsTimelineWidget, SChaosVDTimelineWidget)
+					.HidePlayStopButtons(true)
+					.OnFrameChanged_Raw(this, &SChaosVDSolverPlaybackControls::OnStepSelectionUpdated)
+					.MaxFrames(0)
+				]	
+			]
+		]
+	];
+
+	RegisterNewController(InPlaybackController);
+
+	if (const TSharedPtr<FChaosVDPlaybackController> CurrentPlaybackControllerPtr = InPlaybackController.Pin())
+	{
+		if (const FChaosVDTrackInfo* SolverTrackInfo = CurrentPlaybackControllerPtr->GetTrackInfo(EChaosVDTrackType::Solver, SolverID))
+		{
+			HandleControllerTrackFrameUpdated(InPlaybackController, SolverTrackInfo, InvalidGuid);
+		}
+	}
+}
+
+void SChaosVDSolverPlaybackControls::HandlePlaybackControllerDataUpdated(TWeakPtr<FChaosVDPlaybackController> InController)
+{
+	if (PlaybackController != InController)
+	{
+		RegisterNewController(InController);
+	}
+
+	const TSharedPtr<FChaosVDPlaybackController> ControllerSharedPtr = PlaybackController.Pin();
+	if (ControllerSharedPtr.IsValid() && ControllerSharedPtr->IsRecordingLoaded())
+	{	
+		const int32 AvailableFrames = ControllerSharedPtr->GetTrackFramesNumber(EChaosVDTrackType::Solver, SolverID);
+		const int32 AvailableSteps = ControllerSharedPtr->GetTrackStepsAtFrame(EChaosVDTrackType::Solver, SolverID, ControllerSharedPtr->GetTrackCurrentFrame(EChaosVDTrackType::Solver, SolverID));
+
+		// Max is inclusive and we use this to request as the index on the recorded frames/steps arrays so we need to -1 to the available frames/steps
+		FramesTimelineWidget->UpdateMinMaxValue(0, AvailableFrames != INDEX_NONE ? AvailableFrames -1 : 0);
+
+		//TODO: This will show steps 0/0 if only one step is recorded, we need to add a way to override that functionality
+		// or just set the slider to start from 1 and handle the offset later 
+		StepsTimelineWidget->UpdateMinMaxValue(0,AvailableSteps != INDEX_NONE ? AvailableSteps -1 : 0);
+	}
+	else
+	{
+		FramesTimelineWidget->UpdateMinMaxValue(0, 0);
+		FramesTimelineWidget->ResetTimeline();
+		StepsTimelineWidget->UpdateMinMaxValue(0,0);
+		StepsTimelineWidget->ResetTimeline();
+	}
+}
+
+void SChaosVDSolverPlaybackControls::HandleControllerTrackFrameUpdated(TWeakPtr<FChaosVDPlaybackController> InController, const FChaosVDTrackInfo* UpdatedTrackInfo, FGuid InstigatorGuid)
+{
+	if (InstigatorGuid == GetInstigatorID())
+	{
+		// Ignore the update if we initiated it
+		return;
+	}
+
+	if (const TSharedPtr<FChaosVDPlaybackController> CurrentPlaybackControllerPtr = InController.Pin())
+	{
+		if (const FChaosVDTrackInfo* SolverTrackInfo = CurrentPlaybackControllerPtr->GetTrackInfo(EChaosVDTrackType::Solver, SolverID))
+		{
+			FramesTimelineWidget->SetCurrentTimelineFrame(SolverTrackInfo->CurrentFrame, EChaosVDSetTimelineFrameFlags::None);
+
+			const int32 AvailableSteps = CurrentPlaybackControllerPtr->GetTrackStepsAtFrame(EChaosVDTrackType::Solver, SolverID, SolverTrackInfo->CurrentFrame);
+			StepsTimelineWidget->UpdateMinMaxValue(0,AvailableSteps != INDEX_NONE ? AvailableSteps -1 : 0);
+
+			// On Frame updates, always use Step 0
+			StepsTimelineWidget->SetCurrentTimelineFrame(0, EChaosVDSetTimelineFrameFlags::None);
+		}
+	}
+}
+
+void SChaosVDSolverPlaybackControls::OnFrameSelectionUpdated(int32 NewFrameIndex) const
+{
+	if (const TSharedPtr<FChaosVDPlaybackController> PlaybackControllerPtr = PlaybackController.Pin())
+	{
+		const int32 AvailableSteps = PlaybackControllerPtr->GetTrackStepsAtFrame(EChaosVDTrackType::Solver, SolverID, PlaybackControllerPtr->GetTrackCurrentFrame(EChaosVDTrackType::Solver, SolverID));
+		StepsTimelineWidget->UpdateMinMaxValue(0,AvailableSteps != INDEX_NONE ? AvailableSteps - 1 : 0);
+
+		// On Frame updates, always use Step 0
+		constexpr int32 StepNumber = 0;
+		StepsTimelineWidget->SetCurrentTimelineFrame(StepNumber, EChaosVDSetTimelineFrameFlags::None);
+		PlaybackControllerPtr->GoToTrackFrame(GetInstigatorID(), EChaosVDTrackType::Solver, SolverID, NewFrameIndex, StepNumber);
+	}
+}
+
+void SChaosVDSolverPlaybackControls::OnStepSelectionUpdated(int32 NewStepIndex) const
+{
+	if (const TSharedPtr<FChaosVDPlaybackController> PlaybackControllerPtr = PlaybackController.Pin())
+	{
+		// On Steps updates. Always use the current Frame
+		int32 CurrentFrame = PlaybackControllerPtr->GetTrackCurrentFrame(EChaosVDTrackType::Solver, SolverID);
+		PlaybackControllerPtr->GoToTrackFrame(GetInstigatorID(), EChaosVDTrackType::Solver, SolverID, CurrentFrame, NewStepIndex);
+	}
+}
+
+#undef LOCTEXT_NAMESPACE
\ No newline at end of file
diff --git a/Engine/Plugins/Experimental/ChaosVD/Source/ChaosVD/Private/Widgets/SChaosVDSolverTracks.cpp b/Engine/Plugins/Experimental/ChaosVD/Source/ChaosVD/Private/Widgets/SChaosVDSolverTracks.cpp
new file mode 100644
--- /dev/null
+++ b/Engine/Plugins/Experimental/ChaosVD/Source/ChaosVD/Private/Widgets/SChaosVDSolverTracks.cpp
@@ -0,0 +1,132 @@
+﻿// Copyright Epic Games, Inc. All Rights Reserved.
+
+#include "Widgets/SChaosVDSolverTracks.h"
+
+#include "ChaosVDPlaybackController.h"
+#include "Widgets/SChaosVDPlaybackViewport.h"
+#include "Widgets/SChaosVDSolverPlaybackControls.h"
+#include "Widgets/Layout/SExpandableArea.h"
+
+void SChaosVDSolverTracks::Construct(const FArguments& InArgs, TWeakPtr<FChaosVDPlaybackController> InPlaybackController)
+{
+	RegisterNewController(InPlaybackController);
+
+	ChildSlot
+	[
+		SAssignNew(SolverTracksListWidget, SListView<TSharedPtr<FChaosVDTrackInfo>>)
+			.ListItemsSource(&CachedTrackInfoArray)
+			.SelectionMode( ESelectionMode::None )
+			.ListViewStyle(&FAppStyle::Get().GetWidgetStyle<FTableViewStyle>("SimpleListView"))
+			.OnGenerateRow(this, &SChaosVDSolverTracks::MakeSolverTrackControlsFromTrackInfo)
+	];
+	
+	if (const TSharedPtr<FChaosVDPlaybackController> CurrentPlaybackControllerPtr = InPlaybackController.Pin())
+	{
+		if (const FChaosVDTrackInfo* GameTrackInfo = CurrentPlaybackControllerPtr->GetTrackInfo(EChaosVDTrackType::Game, FChaosVDPlaybackController::GameTrackID))
+		{
+			HandleControllerTrackFrameUpdated(InPlaybackController, GameTrackInfo, InvalidGuid);
+		}
+	}
+}
+
+void SChaosVDSolverTracks::HandlePlaybackControllerDataUpdated(TWeakPtr<FChaosVDPlaybackController> InPlaybackController)
+{
+	if (PlaybackController != InPlaybackController)
+	{
+		RegisterNewController(InPlaybackController);
+	}
+
+	if (const TSharedPtr<FChaosVDPlaybackController> CurrentPlaybackControllerPtr = InPlaybackController.Pin())
+	{
+		// If the controller data was updated, need to update our cache track info data as it could have been changed.
+		// For example this can happen when we load another recording. We use the GameTrack info for that as it is the one that is always valid
+		const FChaosVDTrackInfo* GameTrackInfo = CurrentPlaybackControllerPtr->GetTrackInfo(EChaosVDTrackType::Game, FChaosVDPlaybackController::GameTrackID);
+		UpdatedCachedTrackInfoData(InPlaybackController, GameTrackInfo);
+	}
+}
+
+void SChaosVDSolverTracks::UpdatedCachedTrackInfoData(TWeakPtr<FChaosVDPlaybackController> InPlaybackController, const FChaosVDTrackInfo* UpdatedTrackInfo)
+{
+	if (const TSharedPtr<FChaosVDPlaybackController> CurrentPlaybackControllerPtr = InPlaybackController.Pin())
+	{
+		TArray<TSharedPtr<FChaosVDTrackInfo>> TrackInfoArray;
+		CurrentPlaybackControllerPtr->GetAvailableTrackInfosAtTrackFrame(EChaosVDTrackType::Solver, TrackInfoArray, UpdatedTrackInfo);
+
+		if (TrackInfoArray != CachedTrackInfoArray)
+		{
+			CachedTrackInfoArray = MoveTemp(TrackInfoArray);
+			SolverTracksListWidget->RebuildList();
+		}
+	}
+	else
+	{
+		CachedTrackInfoArray.Empty();
+		SolverTracksListWidget->RebuildList();
+	}
+}
+
+void SChaosVDSolverTracks::HandleControllerTrackFrameUpdated(TWeakPtr<FChaosVDPlaybackController> InPlaybackController, const FChaosVDTrackInfo* UpdatedTrackInfo, FGuid InstigatorGuid)
+{
+	if (InstigatorGuid == GetInstigatorID())
+	{
+		// Ignore the update if we initiated it
+		return;
+	}
+
+	if (UpdatedTrackInfo == nullptr)
+	{
+		return;
+	}
+
+	if (UpdatedTrackInfo->TrackType == EChaosVDTrackType::Solver)
+	{
+		return;
+	}
+
+	UpdatedCachedTrackInfoData(InPlaybackController, UpdatedTrackInfo);
+}
+
+TSharedRef<ITableRow> SChaosVDSolverTracks::MakeSolverTrackControlsFromTrackInfo(TSharedPtr<FChaosVDTrackInfo> TrackInfo, const TSharedRef<STableViewBase>& OwnerTable)
+{
+	TSharedPtr<SVerticalBox> PlaybackControlsContainer;
+	SAssignNew(PlaybackControlsContainer, SVerticalBox)
+			+SVerticalBox::Slot()
+			.AutoHeight()
+			.Padding(20.0f, 1.0f, 20.0f, 1.0f)
+			[
+				SNew(SExpandableArea)
+					.InitiallyCollapsed(false)
+					.BorderBackgroundColor(FLinearColor::White)
+					.Padding(FMargin(8.f))
+					.HeaderContent()
+					[
+						SNew(SHorizontalBox)
+						+SHorizontalBox::Slot()
+						.HAlign(HAlign_Left)
+						.VAlign(VAlign_Center)
+						.AutoWidth()
+						.Padding(0.f, 0.f, 0.f, 0.f)
+						[
+							SNew(STextBlock)
+							.Text(FText::FromString(TrackInfo->TrackName))
+							.Font(FCoreStyle::Get().GetFontStyle("ExpandableArea.TitleFont"))
+						]
+					]
+					.BodyContent()
+					[
+						SNew(SHorizontalBox)
+						+SHorizontalBox::Slot()
+						.Padding(2.f, 4.f, 2.f, 12.f)
+						[
+							SNew(SChaosVDSolverPlaybackControls, TrackInfo->TrackID, PlaybackController)
+						]
+					]
+			];
+
+	TSharedPtr<SWidget> RowWidget = PlaybackControlsContainer;
+	return
+		SNew(STableRow< TSharedPtr<FString> >, OwnerTable)
+		[
+			RowWidget.ToSharedRef()
+		];
+}
diff --git a/Engine/Plugins/Experimental/ChaosVD/Source/ChaosVD/Private/Widgets/SChaosVDTimelineWidget.cpp b/Engine/Plugins/Experimental/ChaosVD/Source/ChaosVD/Private/Widgets/SChaosVDTimelineWidget.cpp
--- a/Engine/Plugins/Experimental/ChaosVD/Source/ChaosVD/Private/Widgets/SChaosVDTimelineWidget.cpp
+++ b/Engine/Plugins/Experimental/ChaosVD/Source/ChaosVD/Private/Widgets/SChaosVDTimelineWidget.cpp
@@ -25,7 +25,7 @@ void SChaosVDTimelineWidget::Construct(const FArguments& InArgs)
 			+SHorizontalBox::Slot()
 			.HAlign(HAlign_Center)
 			.VAlign(VAlign_Center)
-			.FillWidth(0.2)
+			.FillWidth(0.2f)
 			[
 				SNew(SHorizontalBox)
 				+SHorizontalBox::Slot()
@@ -89,18 +89,18 @@ void SChaosVDTimelineWidget::Construct(const FArguments& InArgs)
 			]
 			+SHorizontalBox::Slot()
 			.VAlign(VAlign_Center)
-			.FillWidth(0.75)
+			.FillWidth(0.7f)
 			[
 			  SAssignNew(TimelineSlider, SSlider)
 			  .ToolTipText_Lambda([this]()-> FText{ return FText::AsNumber(CurrentFrame); })
 			  .Value(CurrentFrame)
-			  .OnValueChanged_Raw(this, &SChaosVDTimelineWidget::SetCurrentTimelineFrame)
+			  .OnValueChanged_Raw(this, &SChaosVDTimelineWidget::SetCurrentTimelineFrame, EChaosVDSetTimelineFrameFlags::BroadcastChange)
 			  .StepSize(1)
 			  .MaxValue(MaxFrames)
 			  .MinValue(0)
 			]
 			+SHorizontalBox::Slot()
-			.FillWidth(0.05)
+			.FillWidth(0.1f)
 			.HAlign(HAlign_Center)
 			.VAlign(VAlign_Center)
 			[
@@ -135,11 +135,18 @@ void SChaosVDTimelineWidget::ResetTimeline()
 	CurrentFrame = MinFrames;
 }
 
-void SChaosVDTimelineWidget::SetCurrentTimelineFrame(float FrameNumber)
+void SChaosVDTimelineWidget::SetCurrentTimelineFrame(float FrameNumber, EChaosVDSetTimelineFrameFlags Options)
 {
 	CurrentFrame = static_cast<int32>(FrameNumber);
-	TimelineSlider->SetValue(CurrentFrame);
-	FrameChangedDelegate.ExecuteIfBound(CurrentFrame);
+
+	if (TimelineSlider.IsValid())
+	{
+		TimelineSlider->SetValue(CurrentFrame);
+		if (EnumHasAnyFlags(Options, EChaosVDSetTimelineFrameFlags::BroadcastChange))
+		{
+			FrameChangedDelegate.ExecuteIfBound(CurrentFrame);
+		}
+	}
 }
 
 FReply SChaosVDTimelineWidget::Play()
diff --git a/Engine/Plugins/Experimental/ChaosVD/Source/ChaosVD/Public/ChaosVDPlaybackController.h b/Engine/Plugins/Experimental/ChaosVD/Source/ChaosVD/Public/ChaosVDPlaybackController.h
--- a/Engine/Plugins/Experimental/ChaosVD/Source/ChaosVD/Public/ChaosVDPlaybackController.h
+++ b/Engine/Plugins/Experimental/ChaosVD/Source/ChaosVD/Public/ChaosVDPlaybackController.h
@@ -2,58 +2,184 @@
 
 #pragma once
 
+#include "ChaosVDRecording.h"
 #include "Templates/SharedPointer.h"
 #include "Delegates/DelegateCombinations.h"
 #include "Delegates/Delegate.h"
 
+struct FChaosVDTrackInfo;
 class FChaosVDScene;
 struct FChaosVDRecording;
 class FChaosVDPlaybackController;
 class FString;
 
-DECLARE_DELEGATE_OneParam(FChaosVDPlaybackControllerUpdated, TWeakPtr<FChaosVDPlaybackController>)
+DECLARE_MULTICAST_DELEGATE_OneParam(FChaosVDPlaybackControllerUpdated, TWeakPtr<FChaosVDPlaybackController>)
+DECLARE_MULTICAST_DELEGATE_ThreeParams(FChaosVDPlaybackControllerFrameUpdated, TWeakPtr<FChaosVDPlaybackController>, const FChaosVDTrackInfo*, FGuid);
+
+/** Enum with the available game track types*/
+enum class EChaosVDTrackType : int32
+{
+	Invalid,
+	Game,
+	Solver
+};
+
+/** Data that represents the current state of a track and ID info*/
+struct FChaosVDTrackInfo
+{
+	int32 TrackID = INDEX_NONE;
+	EChaosVDTrackType TrackType = EChaosVDTrackType::Invalid;
+	int32 CurrentFrame = INDEX_NONE;
+	int32 CurrentStep = INDEX_NONE;
+	int32 MaxFrames = INDEX_NONE;
+	FString TrackName;
+};
+
+/** Flags used to control how the unload of a recording is performed */
+enum class EChaosVDUnloadRecordingFlags : uint8
+{
+	None = 0, 
+	BroadcastChanges = 1 << 0,
+	Silent = 1 << 1
+};
+ENUM_CLASS_FLAGS(EChaosVDUnloadRecordingFlags)
+
+typedef TMap<int32, TSharedPtr<FChaosVDTrackInfo>> TrackInfoByIDMap;
 
 /** Loads,unloads and owns a Chaos VD recording file */
 class FChaosVDPlaybackController : public TSharedFromThis<FChaosVDPlaybackController>
 {
 public:
 
+	/** ID used for the Game Track */
+	static constexpr int32 GameTrackID  = 0 ;
+
 	FChaosVDPlaybackController(const TWeakPtr<FChaosVDScene>& InSceneToControl);
 	~FChaosVDPlaybackController();
-	
+
+	/** Loads a recording using a Trace Session Name */
 	bool LoadChaosVDRecordingFromTraceSession(const FString& InSessionName);
-	
-	void UnloadCurrentRecording(const bool bBroadcastUpdate = true);
 
-	void GoToRecordedStep(const int32 InTrackID, const int32 FrameNumber, const int32 Step);
+	/** Unloads the currently loaded recording
+	 * @param UnloadOptions Options flags to change the steps performed during the unload
+	 */
+	void UnloadCurrentRecording(EChaosVDUnloadRecordingFlags UnloadOptions = EChaosVDUnloadRecordingFlags::BroadcastChanges);
 
-	TWeakPtr<FChaosVDRecording> GetCurrentRecording() { return LoadedRecording; }
+	/** Returns true if the controller has a valid recording loaded*/
+	bool IsRecordingLoaded() const { return LoadedRecording.IsValid(); }
 
-	int32 GetStepsForFrame(const int32 InTrackID, const int32 FrameNumber) const;
-	int32 GetAvailableFramesNumber(const int32 InTrackID) const;
-	int32 GetAvailableSolversNumber() const;
+	/** Returns a weak ptr to the Scene this controller is controlling during playback */
+	TWeakPtr<FChaosVDScene> GetControllerScene() { return SceneToControl; }
 
-	int32 GetActiveSolverTrackID() const;
+	/**
+	 * Moves a track of the recording to the specified step and frame numbers
+	 * @param InstigatorID ID of the ControllerInstigator that requested the move
+	 * @param TrackType Type of the track to to move
+	 * @param InTrackID ID of the track to move
+	 * @param FrameNumber Frame number to go
+	 * @param Step Step number to go
+	 */
+	void GoToTrackFrame(FGuid InstigatorID, EChaosVDTrackType TrackType, int32 InTrackID, int32 FrameNumber, int32 Step);
+
+	/**
+	 * Gets the number of available steps in a track at the specified frame
+	 * @param TrackType Type of the track to evaluate
+	 * @param InTrackID ID of the track to evaluate
+	 * @param FrameNumber Frame number to evaluate
+	 * @return Number of available steps
+	 */
+	int32 GetTrackStepsAtFrame(EChaosVDTrackType TrackType, int32 InTrackID, int32 FrameNumber) const;
+
+	/**
+	 * Gets the number of available frames for the specified track
+	 * @param TrackType Type of the track to evaluate
+	 * @param InTrackID ID of the track to evaluate
+	 * @return Number of available frames
+	 */
+	int32 GetTrackFramesNumber(EChaosVDTrackType TrackType, int32 InTrackID) const;
+	
+	/**
+	 * Gets the current frame number at which the specified track is at
+	 * @param TrackType Type of the track to evaluate
+	 * @param InTrackID ID of the track to evaluate
+	 * @return Current frame number for the track
+	 */
+	int32 GetTrackCurrentFrame(EChaosVDTrackType TrackType, int32 InTrackID) const;
+
+	/**
+	 * Gets the number of available frames for the specified track
+	 * @param TrackType Type of the track to evaluate
+	 * @param InTrackID ID of the track to evaluate
+	 * @return Number of available frames
+	 */
+	int32 GetTrackCurrentStep(EChaosVDTrackType TrackType, int32 InTrackID) const;
+
+	/** Converts the current frame number of a track, to a frame number in other tracks space time
+	 * @param FromTrack Track info with the current frame number we want to convert
+	 * @param ToTrack Track info we want to use to convert the frame to
+	 * @return Converted Frame Number
+	 */
+	int32 ConvertCurrentFrameToOtherTrackFrame(FChaosVDTrackInfo FromTrack, FChaosVDTrackInfo ToTrack);
+
+	/**
+	 * Gets all the ids of the tracks, of the specified type, that are available available on the loaded recording
+	 * @param TrackType Type of the tracks we are interested in
+	 * @param OutTrackInfo Array where any found track info data will be added
+	 */
+
+	void GetAvailableTracks(EChaosVDTrackType TrackType, TArray<TSharedPtr<FChaosVDTrackInfo>>& OutTrackInfo);
+	
+	/**
+	 * Gets all the ids of the tracks, of the specified type, that are available available on the loaded recording, at a specified frame
+	 * @param TrackTypeToFind Type of the tracks we are interested in
+	 * @param OutTrackInfo Array where any found track info data will be added
+	 * @param TrackFrameInfo Ptr to the track info with the current frame to evaluate
+	 */
+	void GetAvailableTrackInfosAtTrackFrame(EChaosVDTrackType TrackTypeToFind, TArray<TSharedPtr<FChaosVDTrackInfo>>& OutTrackInfo, const FChaosVDTrackInfo* TrackFrameInfo);
+
+	/**
+	 * Gets the track info of the specified type with the specified ID
+	 * @param TrackType Type of the track to find
+	 * @param TrackID ID of the track to find
+	 * @return Ptr to the found track info data - Null if nothing was found
+	 */
+	const FChaosVDTrackInfo* GetTrackInfo(EChaosVDTrackType TrackType, int32 TrackID);
+
+	/** Returns a weak ptr pointer to the loaded recording */
+	TWeakPtr<FChaosVDRecording> GetCurrentRecording() { return LoadedRecording; }
 
-	int32 GetCurrentFrame(const int32 InTrackID) const;
-	int32 GetCurrentStep(const int32 InTrackID) const;
+	/** Called when data on the recording being controlled gets updated internally or externally (for example, during Trace Analysis)*/
+	FChaosVDPlaybackControllerUpdated& OnDataUpdated() { return ControllerUpdatedDelegate; }
 
-	bool IsRecordingLoaded() const { return LoadedRecording.IsValid(); }
+	/** Called when a frame on a track is updated */
+	FChaosVDPlaybackControllerFrameUpdated& OnTrackFrameUpdated() { return ControllerFrameUpdatedDelegate; }
 
-	TWeakPtr<FChaosVDScene> GetControllerScene() { return SceneToControl; }
+protected:
 
-	FChaosVDPlaybackControllerUpdated& OnControllerUpdated() { return ControllerUpdatedDelegate; }
+	/** Updates (or adds) solvers data from the loaded recording to the solver tracks */
+	void UpdateSolverTracksData();
 
-protected:
+	/** Updates the controlled scene with the loaded data at specified game frame */
+	void GoToRecordedGameFrame(int32 FrameNumber, FGuid InstigatorID);
+
+	/** Updates the controlled scene with the loaded data at specified solver frame and solver step */
+	void GoToRecordedSolverStep(int32 InTrackID, int32 FrameNumber, int32 Step, FGuid InstigatorID);
 
+	/** Handles any data changes on the loaded recording - Usually called during Trace analysis */
 	void HandleCurrentRecordingUpdated();
 
-	TMap<int32,int32> CurrentFramePerTrack;
-	TMap<int32,int32> CurrentStepPerTrack;
-	
+	/** Map containing all track info, by track type*/
+	TMap<EChaosVDTrackType, TrackInfoByIDMap> TrackInfoPerType;
+
+	/** Ptr to the loaded recording */
 	TSharedPtr<FChaosVDRecording> LoadedRecording;
 
+	/**Ptr to the current Chaos VD Scene this controller controls*/
 	TWeakPtr<FChaosVDScene> SceneToControl;
 
+	/** Delegate called when the data on the loaded recording changes */
 	FChaosVDPlaybackControllerUpdated ControllerUpdatedDelegate;
+
+	/** Delegate called when the a in a track changes */
+	FChaosVDPlaybackControllerFrameUpdated ControllerFrameUpdatedDelegate;
 };
diff --git a/Engine/Plugins/Experimental/ChaosVD/Source/ChaosVD/Public/ChaosVDPlaybackControllerInstigator.cpp b/Engine/Plugins/Experimental/ChaosVD/Source/ChaosVD/Public/ChaosVDPlaybackControllerInstigator.cpp
new file mode 100644
--- /dev/null
+++ b/Engine/Plugins/Experimental/ChaosVD/Source/ChaosVD/Public/ChaosVDPlaybackControllerInstigator.cpp
@@ -0,0 +1,5 @@
+﻿// Copyright Epic Games, Inc. All Rights Reserved.
+
+#include "ChaosVDPlaybackControllerInstigator.h"
+
+FGuid IChaosVDPlaybackControllerInstigator::InvalidGuid = FGuid(0,0,0,0);
\ No newline at end of file
diff --git a/Engine/Plugins/Experimental/ChaosVD/Source/ChaosVD/Public/ChaosVDPlaybackControllerInstigator.h b/Engine/Plugins/Experimental/ChaosVD/Source/ChaosVD/Public/ChaosVDPlaybackControllerInstigator.h
new file mode 100644
--- /dev/null
+++ b/Engine/Plugins/Experimental/ChaosVD/Source/ChaosVD/Public/ChaosVDPlaybackControllerInstigator.h
@@ -0,0 +1,19 @@
+﻿// Copyright Epic Games, Inc. All Rights Reserved.
+
+#pragma once
+#include "Misc/Guid.h"
+
+/** Common interface for any object that alters a ChaosVDPlaybackController */
+class IChaosVDPlaybackControllerInstigator
+{
+public:
+	IChaosVDPlaybackControllerInstigator() : ID(FGuid::NewGuid())
+	{
+	}
+
+	static FGuid InvalidGuid;
+
+	const FGuid& GetInstigatorID() const { return ID;};
+private:
+	FGuid ID;
+};
diff --git a/Engine/Plugins/Experimental/ChaosVD/Source/ChaosVD/Public/ChaosVDPlaybackControllerObserver.h b/Engine/Plugins/Experimental/ChaosVD/Source/ChaosVD/Public/ChaosVDPlaybackControllerObserver.h
new file mode 100644
--- /dev/null
+++ b/Engine/Plugins/Experimental/ChaosVD/Source/ChaosVD/Public/ChaosVDPlaybackControllerObserver.h
@@ -0,0 +1,26 @@
+﻿// Copyright Epic Games, Inc. All Rights Reserved.
+
+#pragma once
+
+#include "Misc/Guid.h"
+#include "Templates/SharedPointer.h"
+
+struct FChaosVDTrackInfo;
+class FChaosVDPlaybackController;
+
+
+/**
+ * Class to be used as base to any object that needs to Process Player Controller changes
+ */
+class FChaosVDPlaybackControllerObserver
+{
+public:
+	virtual ~FChaosVDPlaybackControllerObserver();
+
+protected:
+	virtual void RegisterNewController(TWeakPtr<FChaosVDPlaybackController> NewController);
+	virtual void HandlePlaybackControllerDataUpdated(TWeakPtr<FChaosVDPlaybackController> InController) {};
+	virtual void HandleControllerTrackFrameUpdated(TWeakPtr<FChaosVDPlaybackController> InController, const FChaosVDTrackInfo* UpdatedTrackInfo, FGuid InstigatorGuid){};
+	
+	TWeakPtr<FChaosVDPlaybackController> PlaybackController;
+};
diff --git a/Engine/Plugins/Experimental/ChaosVD/Source/ChaosVD/Public/ChaosVDScene.h b/Engine/Plugins/Experimental/ChaosVD/Source/ChaosVD/Public/ChaosVDScene.h
--- a/Engine/Plugins/Experimental/ChaosVD/Source/ChaosVD/Public/ChaosVDScene.h
+++ b/Engine/Plugins/Experimental/ChaosVD/Source/ChaosVD/Public/ChaosVDScene.h
@@ -13,6 +13,8 @@ class AChaosVDParticleActor;
 class FReferenceCollector;
 class UWorld;
 
+typedef TMap<int32, AChaosVDParticleActor*> FChaosVDParticlesByIDMap;
+
 DECLARE_MULTICAST_DELEGATE(FChaosVDSceneUpdatedDelegate)
 
 /** Recreates a UWorld from a recorded Chaos VD Frame */
@@ -39,6 +41,8 @@ public:
 
 	void HandleNewGeometryData(const TSharedPtr<const Chaos::FImplicitObject>&, const int32 GeometryID) const;
 
+	void HandleEnterNewGameFrame(int32 FrameNumber, const TArray<int32>& AvailableSolversIds);
+
 	/** Deletes all actors of the Scene and underlying UWorld */
 	void CleanUpScene();
 
@@ -65,7 +69,7 @@ private:
 	UWorld* CreatePhysicsVDWorld() const;
 
 	/** Map of ID-ChaosVDParticle Actor. Used to keep track of actor instances and be able to modify them as needed*/
-	TMap<int32, AChaosVDParticleActor*> ParticleVDInstancesByID;
+	TMap<int32, FChaosVDParticlesByIDMap> ParticlesBySolverID;
 
 	/** UWorld instance used to represent the recorded debug data */
 	UWorld* PhysicsVDWorld = nullptr;
diff --git a/Engine/Plugins/Experimental/ChaosVD/Source/ChaosVD/Public/ChaosVDSolversTracksTab.h b/Engine/Plugins/Experimental/ChaosVD/Source/ChaosVD/Public/ChaosVDSolversTracksTab.h
new file mode 100644
--- /dev/null
+++ b/Engine/Plugins/Experimental/ChaosVD/Source/ChaosVD/Public/ChaosVDSolversTracksTab.h
@@ -0,0 +1,22 @@
+﻿// Copyright Epic Games, Inc. All Rights Reserved.
+
+#pragma once
+#include "ChaosVDTabSpawnerBase.h"
+
+class SChaosVDSolverTracks;
+
+/** Spawns and handles and instance for the visual debugger Solvers Tracks tab
+ * Which provides playback controls for each recorded solver
+ */
+class FChaosVDSolversTracksTab : public FChaosVDTabSpawnerBase, public TSharedFromThis<FChaosVDSolversTracksTab>
+{
+public:
+	FChaosVDSolversTracksTab(const FName& InTabID, TSharedPtr<FTabManager> InTabManager, SChaosVDMainTab* InOwningTabWidget) : FChaosVDTabSpawnerBase(InTabID, InTabManager, InOwningTabWidget)
+	{	
+	}
+	
+	virtual TSharedRef<SDockTab> HandleTabSpawned(const FSpawnTabArgs& Args) override;
+
+protected:
+	TSharedPtr<SChaosVDSolverTracks> SolverTracksWidget;
+};
diff --git a/Engine/Plugins/Experimental/ChaosVD/Source/ChaosVD/Public/ChaosVDTabsIDs.h b/Engine/Plugins/Experimental/ChaosVD/Source/ChaosVD/Public/ChaosVDTabsIDs.h
--- a/Engine/Plugins/Experimental/ChaosVD/Source/ChaosVD/Public/ChaosVDTabsIDs.h
+++ b/Engine/Plugins/Experimental/ChaosVD/Source/ChaosVD/Public/ChaosVDTabsIDs.h
@@ -13,4 +13,5 @@ public:
 	static const FName WorldOutliner;
 	static const FName DetailsPanel;
 	static const FName OutputLog;
+	static const FName SolversTrack;
 };
diff --git a/Engine/Plugins/Experimental/ChaosVD/Source/ChaosVD/Public/Trace/ChaosVDTraceAnalyzer.h b/Engine/Plugins/Experimental/ChaosVD/Source/ChaosVD/Public/Trace/ChaosVDTraceAnalyzer.h
--- a/Engine/Plugins/Experimental/ChaosVD/Source/ChaosVD/Public/Trace/ChaosVDTraceAnalyzer.h
+++ b/Engine/Plugins/Experimental/ChaosVD/Source/ChaosVD/Public/Trace/ChaosVDTraceAnalyzer.h
@@ -44,7 +44,9 @@ private:
 		RouteId_ChaosVDBinaryDataStart,
 		RouteId_ChaosVDBinaryDataContent,
 		RouteId_ChaosVDBinaryDataEnd,
-		RouteId_ChaosVDSolverSimulationSpace
+		RouteId_ChaosVDSolverSimulationSpace,
+		RouteId_BeginFrame,
+		RouteId_EndFrame
 	};
 
 	void ReadParticleDataFromEvent(const FEventData& InEventData, FChaosVDParticleDebugData& OutParticleData);
diff --git a/Engine/Plugins/Experimental/ChaosVD/Source/ChaosVD/Public/Trace/ChaosVDTraceProvider.h b/Engine/Plugins/Experimental/ChaosVD/Source/ChaosVD/Public/Trace/ChaosVDTraceProvider.h
--- a/Engine/Plugins/Experimental/ChaosVD/Source/ChaosVD/Public/Trace/ChaosVDTraceProvider.h
+++ b/Engine/Plugins/Experimental/ChaosVD/Source/ChaosVD/Public/Trace/ChaosVDTraceProvider.h
@@ -10,6 +10,7 @@
 
 #include "Chaos/ChaosArchive.h"
 
+struct FChaosVDGameFrameData;
 class FChaosVDEngine;
 struct FChaosVDSolverFrameData;
 struct FChaosVDRecording;
@@ -52,9 +53,13 @@ public:
 
 	void CreateRecordingInstanceForSession(const FString& InSessionName);
 	void DeleteRecordingInstanceForSession();
-	void AddFrame(const int32 InSolverGUID, FChaosVDSolverFrameData&& FrameData);
-	FChaosVDSolverFrameData* GetFrame(const int32 InSolverGUID, const int32 FrameNumber) const;
-	FChaosVDSolverFrameData* GetLastFrame(const int32 InSolverGUID) const;
+	void AddSolverFrame(const int32 InSolverGUID, FChaosVDSolverFrameData&& FrameData);
+	void AddGameFrame(FChaosVDGameFrameData&& FrameData);
+	FChaosVDSolverFrameData* GetSolverFrame(const int32 InSolverGUID, const int32 FrameNumber) const;
+	FChaosVDSolverFrameData* GetLastSolverFrame(const int32 InSolverGUID) const;
+
+	FChaosVDGameFrameData* GetSolverFrame(uint64 FrameStartCycle) const;
+	FChaosVDGameFrameData* GetLastGameFrame() const;
 
 	FChaosVDBinaryDataContainer& FindOrAddUnprocessedData(const int32 DataID);
 
diff --git a/Engine/Plugins/Experimental/ChaosVD/Source/ChaosVD/Public/Widgets/SChaosVDMainTab.h b/Engine/Plugins/Experimental/ChaosVD/Source/ChaosVD/Public/Widgets/SChaosVDMainTab.h
--- a/Engine/Plugins/Experimental/ChaosVD/Source/ChaosVD/Public/Widgets/SChaosVDMainTab.h
+++ b/Engine/Plugins/Experimental/ChaosVD/Source/ChaosVD/Public/Widgets/SChaosVDMainTab.h
@@ -6,6 +6,7 @@
 #include "Templates/SharedPointer.h"
 #include "Widgets/SCompoundWidget.h"
 
+class FChaosVDSolversTracksTab;
 class FChaosVDEngine;
 class FChaosVDOutputLogTab;
 class FChaosVDPlaybackViewportTab;
@@ -39,6 +40,7 @@ private:
 	TSharedPtr<FChaosVDWorldOutlinerTab> WorldOutlinerTab;
 	TSharedPtr<FChaosVDObjectDetailsTab> ObjectDetailsTab;
 	TSharedPtr<FChaosVDPlaybackViewportTab> PlaybackViewportTab;
+	TSharedPtr<FChaosVDSolversTracksTab> SolversTracksTab;
 	TSharedPtr<FChaosVDOutputLogTab> OutputLogTab;
 
 	TSharedPtr<FTabManager> TabManager;
diff --git a/Engine/Plugins/Experimental/ChaosVD/Source/ChaosVD/Public/Widgets/SChaosVDPlaybackViewport.h b/Engine/Plugins/Experimental/ChaosVD/Source/ChaosVD/Public/Widgets/SChaosVDPlaybackViewport.h
--- a/Engine/Plugins/Experimental/ChaosVD/Source/ChaosVD/Public/Widgets/SChaosVDPlaybackViewport.h
+++ b/Engine/Plugins/Experimental/ChaosVD/Source/ChaosVD/Public/Widgets/SChaosVDPlaybackViewport.h
@@ -2,9 +2,12 @@
 
 #pragma once
 
+#include "ChaosVDPlaybackControllerInstigator.h"
+#include "ChaosVDPlaybackControllerObserver.h"
 #include "Templates/SharedPointer.h"
 #include "Widgets/SCompoundWidget.h"
 
+class SChaosVDSolverPlaybackControls;
 class FChaosVDPlaybackController;
 class SChaosVDTimelineWidget;
 struct FChaosVDRecording;
@@ -14,7 +17,7 @@ class FSceneViewport;
 class SViewport;
 
 /* Widget that contains the 3D viewport and playback controls */
-class SChaosVDPlaybackViewport : public SCompoundWidget
+class SChaosVDPlaybackViewport : public SCompoundWidget, public FChaosVDPlaybackControllerObserver, public IChaosVDPlaybackControllerInstigator
 {
 public:
 
@@ -29,19 +32,16 @@ protected:
 
 	TSharedPtr<FLevelEditorViewportClient> CreateViewportClient() const;
 
-	void OnPlaybackControllerUpdated(TWeakPtr<FChaosVDPlaybackController> InController);
-
-	void RegisterNewController(TWeakPtr<FChaosVDPlaybackController> NewController);
+	virtual void RegisterNewController(TWeakPtr<FChaosVDPlaybackController> NewController )override;
+	virtual void HandlePlaybackControllerDataUpdated(TWeakPtr<FChaosVDPlaybackController> InController) override;
+	virtual void HandleControllerTrackFrameUpdated(TWeakPtr<FChaosVDPlaybackController> InController, const FChaosVDTrackInfo* UpdatedTrackInfo, FGuid InstigatorGuid) override;
+	void OnPlaybackSceneUpdated();
 
 	void OnFrameSelectionUpdated(int32 NewFrameIndex) const;
-	void OnStepSelectionUpdated(int32 NewStepIndex) const;
 	
-	TSharedPtr<SChaosVDTimelineWidget> FramesTimelineWidget;
-	TSharedPtr<SChaosVDTimelineWidget> StepsTimelineWidget;
+	TSharedPtr<SChaosVDTimelineWidget> GameFramesTimelineWidget;
 
 	TSharedPtr<FLevelEditorViewportClient> LevelViewportClient;
 	TSharedPtr<SViewport> ViewportWidget;
 	TSharedPtr<FSceneViewport> SceneViewport;
-
-	TWeakPtr<FChaosVDPlaybackController> PlaybackController;
 };
diff --git a/Engine/Plugins/Experimental/ChaosVD/Source/ChaosVD/Public/Widgets/SChaosVDSolverPlaybackControls.h b/Engine/Plugins/Experimental/ChaosVD/Source/ChaosVD/Public/Widgets/SChaosVDSolverPlaybackControls.h
new file mode 100644
--- /dev/null
+++ b/Engine/Plugins/Experimental/ChaosVD/Source/ChaosVD/Public/Widgets/SChaosVDSolverPlaybackControls.h
@@ -0,0 +1,36 @@
+﻿// Copyright Epic Games, Inc. All Rights Reserved.
+
+#pragma once
+
+#include "ChaosVDPlaybackControllerInstigator.h"
+#include "ChaosVDPlaybackControllerObserver.h"
+#include "Templates/SharedPointer.h"
+#include "Widgets/SCompoundWidget.h"
+
+struct FChaosVDTrackInfo;
+class FChaosVDPlaybackController;
+class SChaosVDTimelineWidget;
+
+/** Widget that Generates playback controls for solvers
+ * Which are two timelines, one for physics frames and other for steps
+ */
+class SChaosVDSolverPlaybackControls : public SCompoundWidget, public FChaosVDPlaybackControllerObserver, public IChaosVDPlaybackControllerInstigator
+{
+public:
+	SLATE_BEGIN_ARGS( SChaosVDSolverPlaybackControls ){}
+	SLATE_END_ARGS()
+
+	void Construct(const FArguments& InArgs, int32 InSolverID, const TWeakPtr<FChaosVDPlaybackController>& InPlaybackController);
+
+private:
+
+	void OnFrameSelectionUpdated(int32 NewFrameIndex) const;
+	void OnStepSelectionUpdated(int32 NewStepIndex) const;
+
+	virtual void HandlePlaybackControllerDataUpdated(TWeakPtr<FChaosVDPlaybackController> InController) override;
+	virtual void HandleControllerTrackFrameUpdated(TWeakPtr<FChaosVDPlaybackController> InController, const FChaosVDTrackInfo* UpdatedTrackInfo, FGuid InstigatorGuid) override;
+
+	int32 SolverID = INDEX_NONE;
+	TSharedPtr<SChaosVDTimelineWidget> FramesTimelineWidget;
+	TSharedPtr<SChaosVDTimelineWidget> StepsTimelineWidget;
+};
diff --git a/Engine/Plugins/Experimental/ChaosVD/Source/ChaosVD/Public/Widgets/SChaosVDSolverTracks.h b/Engine/Plugins/Experimental/ChaosVD/Source/ChaosVD/Public/Widgets/SChaosVDSolverTracks.h
new file mode 100644
--- /dev/null
+++ b/Engine/Plugins/Experimental/ChaosVD/Source/ChaosVD/Public/Widgets/SChaosVDSolverTracks.h
@@ -0,0 +1,37 @@
+﻿// Copyright Epic Games, Inc. All Rights Reserved.
+
+#pragma once
+
+#include "ChaosVDPlaybackControllerInstigator.h"
+#include "ChaosVDPlaybackControllerObserver.h"
+#include "Templates/SharedPointer.h"
+#include "Widgets/SCompoundWidget.h"
+#include "Widgets/Views/SListView.h"
+
+struct FChaosVDTrackInfo;
+class FChaosVDPlaybackController;
+
+/** Widget that Generates a expandable list of solver controls, based on the existing solver data
+ * on the ChaosVDPlaybackController
+ */
+class SChaosVDSolverTracks : public SCompoundWidget, public FChaosVDPlaybackControllerObserver, public IChaosVDPlaybackControllerInstigator
+{
+public:
+
+	SLATE_BEGIN_ARGS( SChaosVDSolverTracks ){}
+	SLATE_END_ARGS()
+
+	void Construct(const FArguments& InArgs, TWeakPtr<FChaosVDPlaybackController> InPlaybackController);
+
+private:
+
+	void UpdatedCachedTrackInfoData(TWeakPtr<FChaosVDPlaybackController> InPlaybackController, const FChaosVDTrackInfo* UpdatedTrackInfo);
+	virtual void HandlePlaybackControllerDataUpdated(TWeakPtr<FChaosVDPlaybackController> InPlaybackController) override;
+	virtual void HandleControllerTrackFrameUpdated(TWeakPtr<FChaosVDPlaybackController> InPlaybackController, const FChaosVDTrackInfo* UpdatedTrackInfo, FGuid InstigatorGuid) override;
+
+	TSharedRef<ITableRow> MakeSolverTrackControlsFromTrackInfo(TSharedPtr<FChaosVDTrackInfo> TrackInfo, const TSharedRef<STableViewBase>& OwnerTable);
+
+	TSharedPtr<SListView<TSharedPtr<FChaosVDTrackInfo>>> SolverTracksListWidget;
+
+	TArray<TSharedPtr<FChaosVDTrackInfo>> CachedTrackInfoArray;
+};
diff --git a/Engine/Plugins/Experimental/ChaosVD/Source/ChaosVD/Public/Widgets/SChaosVDTimelineWidget.h b/Engine/Plugins/Experimental/ChaosVD/Source/ChaosVD/Public/Widgets/SChaosVDTimelineWidget.h
--- a/Engine/Plugins/Experimental/ChaosVD/Source/ChaosVD/Public/Widgets/SChaosVDTimelineWidget.h
+++ b/Engine/Plugins/Experimental/ChaosVD/Source/ChaosVD/Public/Widgets/SChaosVDTimelineWidget.h
@@ -11,6 +11,15 @@ struct FSlateBrush;
 
 DECLARE_DELEGATE_OneParam(FChaosVDFrameChangedDelegate, int32)
 
+enum class EChaosVDSetTimelineFrameFlags
+{
+	None = 0,
+	BroadcastChange = 1 << 0,
+	Silent = 1 << 1,
+};
+ENUM_CLASS_FLAGS(EChaosVDSetTimelineFrameFlags)
+
+
 /** Simple timeline control widget */
 class SChaosVDTimelineWidget : public SCompoundWidget
 {
@@ -26,7 +35,7 @@ public:
 
 	void UpdateMinMaxValue(float NewMin, float NewMax);
 
-	void SetCurrentTimelineFrame(float FrameNumber);
+	void SetCurrentTimelineFrame(float FrameNumber, EChaosVDSetTimelineFrameFlags Options = EChaosVDSetTimelineFrameFlags::BroadcastChange);
 
 	/** Brings back the state of the timeline to its original state*/
 	void ResetTimeline();
diff --git a/Engine/Source/Runtime/Experimental/ChaosVDData/Private/ChaosVDRecording.cpp b/Engine/Source/Runtime/Experimental/ChaosVDData/Private/ChaosVDRecording.cpp
--- a/Engine/Source/Runtime/Experimental/ChaosVDData/Private/ChaosVDRecording.cpp
+++ b/Engine/Source/Runtime/Experimental/ChaosVDData/Private/ChaosVDRecording.cpp
@@ -3,7 +3,7 @@
 #include "ChaosVDRecording.h"
 #include "Chaos/ImplicitObject.h"
 
-int32 FChaosVDRecording::GetAvailableFramesNumber(const int32 SolverID) const
+int32 FChaosVDRecording::GetAvailableSolverFramesNumber(const int32 SolverID) const
 {
 	if (const TArray<FChaosVDSolverFrameData>* SolverData = RecordedFramesDataPerSolver.Find(SolverID))
 	{
@@ -13,7 +13,24 @@ int32 FChaosVDRecording::GetAvailableFramesNumber(const int32 SolverID) const
 	return INDEX_NONE;
 }
 
-FChaosVDSolverFrameData* FChaosVDRecording::GetFrameForSolver(const int32 SolverID, const int32 FrameNumber)
+FString FChaosVDRecording::GetSolverName(int32 SolverID)
+{
+	static FString DefaultName(TEXT("Invalid"));
+
+	// Currently we don't create an entry per solver, so we need to get the name from the frame data
+	// TODO: Record Solver specific data per instance and not per frame
+	if (TArray<FChaosVDSolverFrameData>* SolverFramesData = RecordedFramesDataPerSolver.Find(SolverID))
+	{
+		if(!SolverFramesData->IsEmpty())
+		{
+			return (*SolverFramesData)[0].DebugName;
+		}
+	}
+
+	return DefaultName;
+}
+
+FChaosVDSolverFrameData* FChaosVDRecording::GetSolverFrameData(const int32 SolverID, const int32 FrameNumber)
 {
 	if (TArray<FChaosVDSolverFrameData>* SolverFrames = RecordedFramesDataPerSolver.Find(SolverID))
 	{
@@ -24,6 +41,61 @@ FChaosVDSolverFrameData* FChaosVDRecording::GetFrameForSolver(const int32 Solver
 	return nullptr;
 }
 
+FChaosVDSolverFrameData* FChaosVDRecording::GetSolverFrameDataAtCycle(int32 SolverID, uint64 Cycle)
+{
+	if (TArray<FChaosVDSolverFrameData>* SolverFramesPtr = RecordedFramesDataPerSolver.Find(SolverID))
+	{
+		TArray<FChaosVDSolverFrameData>& SolverFrames = *SolverFramesPtr;
+		int32 FrameIndex = Algo::BinarySearchBy(SolverFrames, Cycle, &FChaosVDSolverFrameData::FrameCycle);
+		return SolverFrames.IsValidIndex(FrameIndex) ? &SolverFrames[FrameIndex] : nullptr;
+	}
+
+	return nullptr;
+}
+
+int32 FChaosVDRecording::GetLowestSolverFrameNumberAtCycle(int32 SolverID, uint64 Cycle)
+{
+	if (TArray<FChaosVDSolverFrameData>* SolverFramesPtr = RecordedFramesDataPerSolver.Find(SolverID))
+	{
+		TArray<FChaosVDSolverFrameData>& SolverFrames = *SolverFramesPtr;
+		return Algo::LowerBoundBy(SolverFrames, Cycle, &FChaosVDSolverFrameData::FrameCycle);
+	}
+
+	return INDEX_NONE;
+}
+
+int32 FChaosVDRecording::GetLowestSolverFrameNumberGameFrame(int32 SolverID, int32 GameFrame)
+{
+	if (!GameFrames.IsValidIndex(GameFrame))
+	{
+		return INDEX_NONE;
+	}
+
+	if (TArray<FChaosVDSolverFrameData>* SolverFramesPtr = RecordedFramesDataPerSolver.Find(SolverID))
+	{
+		TArray<FChaosVDSolverFrameData>& SolverFrames = *SolverFramesPtr;
+		
+		return Algo::LowerBoundBy(SolverFrames,GameFrames[GameFrame].FirstCycle, &FChaosVDSolverFrameData::FrameCycle);
+	}
+
+	return INDEX_NONE;
+}
+
+int32 FChaosVDRecording::GetLowestGameFrameAtSolverFrameNumber(int32 SolverID, int32 SolverFrame)
+{
+	if (TArray<FChaosVDSolverFrameData>* SolverFramesPtr = RecordedFramesDataPerSolver.Find(SolverID))
+	{
+		TArray<FChaosVDSolverFrameData>& SolverFrames = *SolverFramesPtr;
+
+		if (SolverFrames.IsValidIndex(SolverFrame))
+		{
+			return Algo::LowerBoundBy(GameFrames, SolverFrames[SolverFrame].FrameCycle, &FChaosVDGameFrameData::FirstCycle);
+		}	
+	}
+
+	return INDEX_NONE;
+}
+
 void FChaosVDRecording::AddFrameForSolver(const int32 SolverID, FChaosVDSolverFrameData&& InFrameData)
 {
 	if (TArray<FChaosVDSolverFrameData>* SolverFrames = RecordedFramesDataPerSolver.Find(SolverID))
@@ -38,19 +110,68 @@ void FChaosVDRecording::AddFrameForSolver(const int32 SolverID, FChaosVDSolverFr
 	OnRecordingUpdated().Broadcast();
 }
 
-EChaosVDFrameLoadState FChaosVDRecording::GetFrameState(const int32 SolverID, const int32 FrameNumber)
+void FChaosVDRecording::AddGameFrameData(FChaosVDGameFrameData&& InFrameData)
+{
+	GameFrames.Add(MoveTemp(InFrameData));
+}
+
+FChaosVDGameFrameData* FChaosVDRecording::GetGameFrameDataAtCycle(uint64 Cycle)
 {
-	const TArray<FChaosVDSolverFrameData>* SolverFrames = RecordedFramesDataPerSolver.Find(SolverID);
-	if (SolverFrames && SolverFrames->IsValidIndex(FrameNumber))
+	int32 FrameIndex = Algo::BinarySearchBy(GameFrames, Cycle, &FChaosVDGameFrameData::FirstCycle);
+
+	if (FrameIndex != INDEX_NONE)
 	{
-		return EChaosVDFrameLoadState::Loaded;
+		return &GameFrames[FrameIndex];
 	}
-	//TODO: Implement "Buffering" state handling
-	// I already implemented this on the version where we didn't use trace and CVD files were streamed from disk
-	// But removed the code for this CL - Leaving this getter here as it is used by the controller and I also left part
-	// of the implementation to handle different states there
 
-	return EChaosVDFrameLoadState::Unknown;
+	return nullptr;
+}
+
+FChaosVDGameFrameData* FChaosVDRecording::GetGameFrameData(int32 FrameNumber)
+{
+	return GameFrames.IsValidIndex(FrameNumber) ? &GameFrames[FrameNumber] : nullptr;
+}
+
+FChaosVDGameFrameData* FChaosVDRecording::GetLastGameFrameData()
+{
+	return GameFrames.Num() > 0 ? &GameFrames.Last() : nullptr;
+}
+
+int32 FChaosVDRecording::GetLowestGameFrameNumberAtCycle(uint64 Cycle)
+{
+	return Algo::LowerBoundBy(GameFrames, Cycle, &FChaosVDGameFrameData::FirstCycle);
+}
+
+void FChaosVDRecording::GetAvailableSolverIDsAtGameFrameNumber(int32 FrameNumber, TArray<int32>& OutSolversID)
+{
+	if (!GameFrames.IsValidIndex(FrameNumber))
+	{
+		return;
+	}
+	
+	FChaosVDGameFrameData& FrameData = GameFrames[FrameNumber];
+
+	OutSolversID.Reserve(RecordedFramesDataPerSolver.Num());
+
+	for (const TPair<int32, TArray<FChaosVDSolverFrameData>>& SolverFramesWithIDPair : RecordedFramesDataPerSolver)
+	{
+		if (SolverFramesWithIDPair.Value.IsEmpty())
+		{
+			continue;
+		}
+
+		if (SolverFramesWithIDPair.Value.Num() == 1 && SolverFramesWithIDPair.Value[0].FrameCycle < FrameData.FirstCycle)
+		{
+			OutSolversID.Add(SolverFramesWithIDPair.Key);
+		}
+		else
+		{
+			if (FrameData.FirstCycle > SolverFramesWithIDPair.Value[0].FrameCycle && FrameData.FirstCycle < SolverFramesWithIDPair.Value.Last().FrameCycle)
+			{
+				OutSolversID.Add(SolverFramesWithIDPair.Key);
+			}
+		}	
+	}
 }
 
 void FChaosVDRecording::AddImplicitObject(const int32 ID, const TSharedPtr<Chaos::FImplicitObject>& InImplicitObject)
diff --git a/Engine/Source/Runtime/Experimental/ChaosVDData/Public/ChaosVDRecording.h b/Engine/Source/Runtime/Experimental/ChaosVDData/Public/ChaosVDRecording.h
--- a/Engine/Source/Runtime/Experimental/ChaosVDData/Public/ChaosVDRecording.h
+++ b/Engine/Source/Runtime/Experimental/ChaosVDData/Public/ChaosVDRecording.h
@@ -82,16 +82,15 @@ struct CHAOSVDDATA_API FChaosVDSolverFrameData
 {
 	FString DebugName;
 	int32 SolverID;
+	uint64 FrameCycle;
 	Chaos::FRigidTransform3 SimulationTransform;
 	TArray<FChaosVDStepData, TInlineAllocator<16>> SolverSteps;
 };
 
-enum class EChaosVDFrameLoadState
+struct FChaosVDGameFrameData
 {
-	Unloaded,
-	Loaded,
-	Buffering,
-	Unknown
+	uint64 FirstCycle;
+	uint64 LastCycle;
 };
 
 /**
@@ -99,28 +98,122 @@ enum class EChaosVDFrameLoadState
  * It is currently populated while analyzing a Trace session
  */
 struct CHAOSVDDATA_API FChaosVDRecording
-{
-	/** Returns the current available recorded frames per recorded solver */
-	int32 GetAvailableFramesNumber(const int32 SolverID) const;
+{	
 	/** Returns the current available recorded solvers number */
-	int32 GetAvailableSolversNumber() const { return RecordedFramesDataPerSolver.Num();}
+	int32 GetAvailableSolversNumber() const { return RecordedFramesDataPerSolver.Num(); }
 	
-	const TMap<int32, TArray<FChaosVDSolverFrameData>>& GetAvailableSolvers() const { return RecordedFramesDataPerSolver; }
+	/** Returns the current available Game Frames */
+	int32 GetAvailableGameFramesNumber() const { return GameFrames.Num(); }
 
-	/** Returns a ptr to the recorded frame data for a specific solver -  Do not store as it is a reference to the element in the array */
-	FChaosVDSolverFrameData* GetFrameForSolver(const int32 SolverID, const int32 FrameNumber);
+	/** Returns a reference to the array holding all the available game frames */
+	const TArray<FChaosVDGameFrameData>& GetAvailableGameFrames() const { return GameFrames; }
+
+	/** Returns a reference to the map containing the available solver data */
+	const TMap<int32, TArray<FChaosVDSolverFrameData>>& GetAvailableSolvers() const { return RecordedFramesDataPerSolver; }
 
-	/** Adds a Frame Data for a specific Solver ID. Creates a solver entry if it does not exist */
+	/**
+	 * Returns the number of available frame data for the specified solver ID
+	 * @param SolverID ID of the solver 
+	 */
+	int32 GetAvailableSolverFramesNumber(int32 SolverID) const;
+	
+	/**
+	 * Returns the name of the specified solver id
+	 * @param SolverID ID of the solver 
+	 */
+	FString GetSolverName(int32 SolverID);
+
+	/**
+	 * Return a ptr to the existing solver frame data from the specified ID and Frame number
+	 * @param SolverID ID of the solver
+	 * @param FrameNumber Frame number
+	 * @return Ptr to the existing solver frame data from the specified ID and Frame number - It is a ptr to the array element, Do not store
+	 */
+	FChaosVDSolverFrameData* GetSolverFrameData(int32 SolverID, int32 FrameNumber);
+	
+	/**
+	 * Return a ptr to the existing solver frame data from the specified ID and Frame number
+	 * @param SolverID ID if the solver
+	 * @param Cycle Platform cycle at which the solver frame was recorded
+	 * @return Ptr to the existing solver frame data from the specified ID and Frame number - It is a ptr to the array element, Do not store
+	 */
+	FChaosVDSolverFrameData* GetSolverFrameDataAtCycle(int32 SolverID, uint64 Cycle);
+
+	/**
+	 * Searches and returns the lowest frame number of a solver at the specified cycle
+	 * @param SolverID ID if the solver
+	 * @param Cycle Platform cycle to use as lower bound
+	 * @return Found frame number. INDEX_NONE if no frame is found for the specified cycle
+	 */
+	int32 GetLowestSolverFrameNumberAtCycle(int32 SolverID, uint64 Cycle);
+	
+	/**
+	 * Searches and returns the lowest frame number of a solver at the specified cycle
+	 * @param SolverID ID if the solver
+	 * @param GameFrame Platform cycle to use as lower bound
+	 * @return Found frame number. INDEX_NONE if no frame is found for the specified cycle
+	 */
+	int32 GetLowestSolverFrameNumberGameFrame(int32 SolverID, int32 GameFrame);
+	
+	/**
+	 * Searches and returns the lowest game frame number at the specified solver frame
+	 * @param SolverID ID if the solver to evaluate
+	 * @param SolverFrame Frame number of the solver to evaluate
+	 * @return Found Game frame number. INDEX_NONE if no frame is found for the specified cycle
+	 */
+	int32 GetLowestGameFrameAtSolverFrameNumber(int32 SolverID, int32 SolverFrame);
+
+	/**
+	 * Adds a Solver Frame Data entry for a specific Solver ID. Creates a solver entry if it does not exist 
+	 * @param SolverID ID of the solver to add
+	 * @param InFrameData Reference to the frame data we want to add
+	 */
 	void AddFrameForSolver(const int32 SolverID, FChaosVDSolverFrameData&& InFrameData);
 
-	/** Returns the current frame state of a frame used to determine if it is ready for use*/
-	EChaosVDFrameLoadState GetFrameState(const int32 SolverID, const int32 FrameNumber);
+	/**
+	 * Adds a Game Frame Data entry. Creates a solver entry if it does not exist 
+	 * @param InFrameData Reference to the frame data we want to add
+	 */
+	void AddGameFrameData(FChaosVDGameFrameData&& InFrameData);
 
 	/** Called each time the recording changes - Mainly when a new frame is added from the Trace analysis */
 	FChaosVDRecordingUpdated& OnRecordingUpdated() { return RecordingUpdatedDelegate; };
 
+	/** Called each time new geometry data becomes available in the recording - Mainly when a new frame is added from the Trace analysis */
 	FChaosVDGeometryDataLoaded& OnGeometryDataLoaded() { return GeometryDataLoaded; };
 
+	/**
+	 * Searches for a recorded Game frame at the specified cycle 
+	 * @param Cycle Platform Cycle to be used in the search
+	 * @return A ptr to the recorded game frame data - This is a ptr to the array element. Do not store
+	 */
+	FChaosVDGameFrameData* GetGameFrameDataAtCycle(uint64 Cycle);
+
+	/**
+	 * Searches for a recorded Game frame at the specified cycle 
+	 * @param FrameNumber Frame Number
+	 * @return A ptr to the recorded game frame data - This is a ptr to the array element. Do not store
+	 */
+	FChaosVDGameFrameData* GetGameFrameData(int32 FrameNumber);
+
+	/** Returns a ptr to the last recorded game frame - This is a ptr to the array element. Do not store */
+	FChaosVDGameFrameData* GetLastGameFrameData();
+
+	/**
+	 * Searches and returns the lowest game frame number at the specified cycle
+	 * @param Cycle Platform Cycle to be used in the search as lower bound
+	 * @return Found Game frame number. INDEX_NONE if no frame is found for the specified cycle
+	 */
+	int32 GetLowestGameFrameNumberAtCycle(uint64 Cycle);
+
+	/**
+     * Gathers all available solvers IDs at the given Game frame number
+     * @param FrameNumber Game Frame number to evaluate
+     * @param OutSolversID Solver's ID array to be filled with any IDs found
+     */
+	void GetAvailableSolverIDsAtGameFrameNumber(int32 FrameNumber, TArray<int32>& OutSolversID);
+
+	/** Returns a reference to the GeometryID-ImplicitObject map of this recording */
 	const TMap<int32, TSharedPtr<const Chaos::FImplicitObject>>& GetGeometryDataMap() const { return ImplicitObjects; };
 
 	/** Adds a shared Implicit Object to the recording */
@@ -136,8 +229,8 @@ protected:
 	
 	void AddImplicitObject_Internal(const int32 ID, const TSharedPtr<const Chaos::FImplicitObject>& InImplicitObject);
 	
-	TMap<int32, EChaosVDFrameLoadState> AvailableFramesState;
 	TMap<int32, TArray<FChaosVDSolverFrameData>> RecordedFramesDataPerSolver;
+	TArray<FChaosVDGameFrameData> GameFrames;
 	FChaosVDRecordingUpdated RecordingUpdatedDelegate;
 	FChaosVDGeometryDataLoaded GeometryDataLoaded;
 
