[ChaosVD] Fix Localization key conflict on SChaosVDPlaybackViewport

diff --git a/Engine/Plugins/Experimental/ChaosVD/Source/ChaosVD/Private/Widgets/SChaosVDPlaybackViewport.cpp b/Engine/Plugins/Experimental/ChaosVD/Source/ChaosVD/Private/Widgets/SChaosVDPlaybackViewport.cpp
--- a/Engine/Plugins/Experimental/ChaosVD/Source/ChaosVD/Private/Widgets/SChaosVDPlaybackViewport.cpp
+++ b/Engine/Plugins/Experimental/ChaosVD/Source/ChaosVD/Private/Widgets/SChaosVDPlaybackViewport.cpp
@@ -81,7 +81,7 @@ void SChaosVDPlaybackViewport::Construct(const FArguments& InArgs, const UWorld*
 			[
 				SNew(STextBlock)
 				.Justification(ETextJustify::Center)
-				.Text(LOCTEXT("PlaybackViewportWidgetPhysicsFramesLabel", "Game Frames" ))
+				.Text(LOCTEXT("PlaybackViewportWidgetGameFramesLabel", "Game Frames" ))
 			]
 			+SVerticalBox::Slot()
 			[
