[Chaos VD] Fix for mirrored geometry appears with inverted normals

If the scale of the implicit object is negative on any component, we need to force-reverse the culling of the instanced mesh component. This works because we have one component per geometry instance.

We might need to revisit this in the future if we try to combine them. In which case we will need separate components for instances with negative and positive scale.

diff --git a/Engine/Plugins/Experimental/ChaosVD/Source/ChaosVD/Private/ChaosVDGeometryBuilder.cpp b/Engine/Plugins/Experimental/ChaosVD/Source/ChaosVD/Private/ChaosVDGeometryBuilder.cpp
--- a/Engine/Plugins/Experimental/ChaosVD/Source/ChaosVD/Private/ChaosVDGeometryBuilder.cpp
+++ b/Engine/Plugins/Experimental/ChaosVD/Source/ChaosVD/Private/ChaosVDGeometryBuilder.cpp
@@ -58,6 +58,12 @@ bool FChaosVDGeometryBuilder::DoesImplicitContainType(const Chaos::FImplicitObje
 	return false;
 }
 
+bool FChaosVDGeometryBuilder::HasNegativeScale(const Chaos::FRigidTransform3& InTransform) const
+{
+	FVector ScaleSignVector = InTransform.GetScale3D().GetSignVector();
+	return ScaleSignVector.X * ScaleSignVector.Y * ScaleSignVector.Z < 0;
+}
+
 UDynamicMesh* FChaosVDGeometryBuilder::CreateAndCacheDynamicMesh(const uint32 GeometryCacheKey, UE::Geometry::FMeshShapeGenerator& MeshGenerator)
 {
 	{
diff --git a/Engine/Plugins/Experimental/ChaosVD/Source/ChaosVD/Public/ChaosVDGeometryBuilder.h b/Engine/Plugins/Experimental/ChaosVD/Source/ChaosVD/Public/ChaosVDGeometryBuilder.h
--- a/Engine/Plugins/Experimental/ChaosVD/Source/ChaosVD/Public/ChaosVDGeometryBuilder.h
+++ b/Engine/Plugins/Experimental/ChaosVD/Source/ChaosVD/Public/ChaosVDGeometryBuilder.h
@@ -94,6 +94,12 @@ public:
 
 private:
 
+	/**
+	 * Evaluates the provided transform's scale, and returns true if the scale has a negative component
+	 * @param InTransform Transform to evaluate
+	 */
+	bool HasNegativeScale(const Chaos::FRigidTransform3& InTransform) const;
+
 	/** Creates a Dynamic Mesh for the provided Implicit object and generator, and then caches it to be reused later
 	 * @param GeometryCacheKey Key to be used to find this geometry in the cache
 	 * @param MeshGenerator Generator class with the data and rules to create the mesh
@@ -487,6 +493,8 @@ ComponentType* FChaosVDGeometryBuilder::CreateMeshComponent(AActor* Owner, const
 	}
 	else if constexpr (std::is_same_v<ComponentType, UInstancedStaticMeshComponent>)
 	{
+		// If we have negative scale we need to force reverse the culling mode in this component otherwise the faces will be inverted
+		MeshComponent->SetReverseCulling(HasNegativeScale(Transform));
 		MeshComponent->AddInstance(Transform);
 	}
 
