[ChaosVD] Implemented local selection system

- Created a Dedicated EditorViewport class to handle the clicks done in the viewport
- Created a custom scene outliner mode to show only the info we need, and to make it use the local selection system
- Added a UTypedElementSelectionSet to ChaosVDScene to hold the selection state

With this change, all selection events do not affect the rest of the editor anymore and everything is contained within the tool. Also, now we can properly select from the viewport to the outliner and from the outliner to the viewport.

diff --git a/Engine/Plugins/Experimental/ChaosVD/Source/ChaosVD/ChaosVD.Build.cs b/Engine/Plugins/Experimental/ChaosVD/Source/ChaosVD/ChaosVD.Build.cs
--- a/Engine/Plugins/Experimental/ChaosVD/Source/ChaosVD/ChaosVD.Build.cs
+++ b/Engine/Plugins/Experimental/ChaosVD/Source/ChaosVD/ChaosVD.Build.cs
@@ -14,7 +14,7 @@ public class ChaosVD : ModuleRules
 				"Engine",
 				"ChaosVDRuntime", 
 				"TraceServices",
-				"ChaosVDData"
+				"ChaosVDData",
 			});
 
 		PrivateDependencyModuleNames.AddRange(
@@ -28,8 +28,7 @@ public class ChaosVD : ModuleRules
 				"ToolMenus",
 				"Slate",
 				"SlateCore", 
-				"OutputLog", 
-				"SceneOutliner", 
+				"OutputLog",
 				"WorkspaceMenuStructure", 
 				"TraceAnalysis", 
 				"TraceInsights",
@@ -39,7 +38,10 @@ public class ChaosVD : ModuleRules
 				"MeshConversion",
 				"MeshDescription",
 				"StaticMeshDescription", 
-				"DynamicMesh"
+				"DynamicMesh",
+				"SceneOutliner",
+				"TypedElementRuntime",
+				"TypedElementFramework"
 			}
 			);
 
diff --git a/Engine/Plugins/Experimental/ChaosVD/Source/ChaosVD/Private/ChaosVDObjectDetailsTab.cpp b/Engine/Plugins/Experimental/ChaosVD/Source/ChaosVD/Private/ChaosVDObjectDetailsTab.cpp
--- a/Engine/Plugins/Experimental/ChaosVD/Source/ChaosVD/Private/ChaosVDObjectDetailsTab.cpp
+++ b/Engine/Plugins/Experimental/ChaosVD/Source/ChaosVD/Private/ChaosVDObjectDetailsTab.cpp
@@ -2,27 +2,17 @@
 
 #include "ChaosVDObjectDetailsTab.h"
 
+#include "ChaosVDScene.h"
 #include "ChaosVDStyle.h"
-#include "ChaosVDTabsIDs.h"
 #include "Modules/ModuleManager.h"
 #include "PropertyEditorModule.h"
-#include "Selection.h"
 #include "Templates/SharedPointer.h"
 #include "Editor.h"
 #include "Widgets/Docking/SDockTab.h"
+#include "Elements/Framework/TypedElementSelectionSet.h"
 
 #define LOCTEXT_NAMESPACE "ChaosVisualDebugger"
 
-FChaosVDObjectDetailsTab::~FChaosVDObjectDetailsTab()
-{
-	if (GEditor == nullptr || GEditor->GetSelectedActors() == nullptr)
-	{
-		return;
-	}
-
-	GEditor->GetSelectedActors()->SelectionChangedEvent.Remove(SelectionDelegateHandle);
-}
-
 TSharedRef<SDockTab> FChaosVDObjectDetailsTab::HandleTabSpawned(const FSpawnTabArgs& Args)
 {
 	FPropertyEditorModule& PropertyEditorModule = FModuleManager::GetModuleChecked<FPropertyEditorModule>("PropertyEditor");
@@ -30,23 +20,19 @@ TSharedRef<SDockTab> FChaosVDObjectDetailsTab::HandleTabSpawned(const FSpawnTabA
 	DetailsViewArgs.bAllowSearch = false;
 	DetailsViewArgs.NameAreaSettings = FDetailsViewArgs::HideNameArea;
 
-	DetailsPanel = PropertyEditorModule.CreateDetailView(DetailsViewArgs);
+	TSharedPtr<FChaosVDScene> ScenePtr = GetChaosVDScene().Pin();
+	check(ScenePtr);
 
-	//TODO: This is just for testing. We will not use global selection events in the final version as these affect the entire editor
-	SelectionDelegateHandle = GEditor->GetSelectedActors()->SelectionChangedEvent.AddLambda(
-	[this](UObject* Object)
-	{
-		TArray<AActor*> SelectedActors;
-		GEditor->GetSelectedActors()->GetSelectedObjects<AActor>(SelectedActors);
+	DetailsPanel = PropertyEditorModule.CreateDetailView(DetailsViewArgs);
 
-		UpdateSelectedObject(SelectedActors.Num() > 0 ? SelectedActors[0] : nullptr);
-	} );
+	RegisterSelectionSetObject(ScenePtr->GetElementSelectionSet());
 
 	TSharedRef<SDockTab> DetailsPanelTab =
 		SNew(SDockTab)
 		.TabRole(ETabRole::MajorTab)
 		.Label(LOCTEXT("DetailsPanel", "Details"))
 		.ToolTipText(LOCTEXT("DetailsPanelToolTip", "See the details of the selected object"));
+
 	DetailsPanelTab->SetContent
 	(
 		DetailsPanel.ToSharedRef()
@@ -57,9 +43,17 @@ TSharedRef<SDockTab> FChaosVDObjectDetailsTab::HandleTabSpawned(const FSpawnTabA
 	return DetailsPanelTab;
 }
 
-void FChaosVDObjectDetailsTab::UpdateSelectedObject(AActor* NewObject) const
+void FChaosVDObjectDetailsTab::HandlePostSelectionChange(const UTypedElementSelectionSet* ChangedSelectionSet)
 {
-	DetailsPanel->SetObject(NewObject, true);
+	TArray<AActor*> SelectedActors = ChangedSelectionSet->GetSelectedObjects<AActor>();
+
+	if (SelectedActors.Num() > 0)
+	{
+		// We don't support multi selection yet
+		ensure(SelectedActors.Num() == 1);
+
+		DetailsPanel->SetObject(SelectedActors[0], true);
+	}
 }
 
 #undef LOCTEXT_NAMESPACE
diff --git a/Engine/Plugins/Experimental/ChaosVD/Source/ChaosVD/Private/ChaosVDParticleActor.cpp b/Engine/Plugins/Experimental/ChaosVD/Source/ChaosVD/Private/ChaosVDParticleActor.cpp
--- a/Engine/Plugins/Experimental/ChaosVD/Source/ChaosVD/Private/ChaosVDParticleActor.cpp
+++ b/Engine/Plugins/Experimental/ChaosVD/Source/ChaosVD/Private/ChaosVDParticleActor.cpp
@@ -133,3 +133,18 @@ void AChaosVDParticleActor::BeginDestroy()
 
 	Super::BeginDestroy();
 }
+
+#if WITH_EDITOR
+bool AChaosVDParticleActor::IsSelectedInEditor() const
+{
+	// The implementation of this method in UObject, used a global edit callback,
+	// but as we don't use the global editor selection system, we need to re-route it.
+	if (TSharedPtr<FChaosVDScene> ScenePtr = OwningScene.Pin())
+	{
+		return ScenePtr->IsObjectSelected(this);
+	}
+
+	return false;
+}
+#endif
+
diff --git a/Engine/Plugins/Experimental/ChaosVD/Source/ChaosVD/Private/ChaosVDPlaybackControllerObserver.cpp b/Engine/Plugins/Experimental/ChaosVD/Source/ChaosVD/Private/ChaosVDPlaybackControllerObserver.cpp
--- a/Engine/Plugins/Experimental/ChaosVD/Source/ChaosVD/Private/ChaosVDPlaybackControllerObserver.cpp
+++ b/Engine/Plugins/Experimental/ChaosVD/Source/ChaosVD/Private/ChaosVDPlaybackControllerObserver.cpp
@@ -2,6 +2,8 @@
 
 #include "ChaosVDPlaybackControllerObserver.h"
 #include "ChaosVDPlaybackController.h"
+#include "ChaosVDScene.h"
+
 
 FChaosVDPlaybackControllerObserver::~FChaosVDPlaybackControllerObserver()
 {
@@ -25,7 +27,12 @@ void FChaosVDPlaybackControllerObserver::RegisterNewController(TWeakPtr<FChaosVD
 		PlaybackController = NewController;
 
 		if (const TSharedPtr<FChaosVDPlaybackController> NewPlaybackControllerPtr = PlaybackController.Pin())
-		{	
+		{
+			if (TSharedPtr<FChaosVDScene> ScenePtr = NewPlaybackControllerPtr->GetControllerScene().Pin())
+			{
+				RegisterSelectionSetObject(ScenePtr->GetElementSelectionSet());
+			}
+	
 			NewPlaybackControllerPtr->OnDataUpdated().AddRaw(this, &FChaosVDPlaybackControllerObserver::HandlePlaybackControllerDataUpdated);
 			NewPlaybackControllerPtr->OnTrackFrameUpdated().AddRaw(this, &FChaosVDPlaybackControllerObserver::HandleControllerTrackFrameUpdated);
 			HandlePlaybackControllerDataUpdated(PlaybackController);
diff --git a/Engine/Plugins/Experimental/ChaosVD/Source/ChaosVD/Private/ChaosVDPlaybackViewportClient.cpp b/Engine/Plugins/Experimental/ChaosVD/Source/ChaosVD/Private/ChaosVDPlaybackViewportClient.cpp
new file mode 100644
--- /dev/null
+++ b/Engine/Plugins/Experimental/ChaosVD/Source/ChaosVD/Private/ChaosVDPlaybackViewportClient.cpp
@@ -0,0 +1,60 @@
+﻿// Copyright Epic Games, Inc. All Rights Reserved.
+
+#include "ChaosVDPlaybackViewportClient.h"
+
+#include "ChaosVDScene.h"
+#include "EngineUtils.h"
+
+FChaosVDPlaybackViewportClient::FChaosVDPlaybackViewportClient() : FEditorViewportClient(nullptr), CVDWorld(nullptr)
+{
+}
+
+FChaosVDPlaybackViewportClient::~FChaosVDPlaybackViewportClient()
+{
+	if (ObjectFocusedDelegateHandle.IsValid())
+	{
+		if (TSharedPtr<FChaosVDScene> ScenePtr = CVDScene.Pin())
+		{
+			ScenePtr->OnObjectFocused().Remove(ObjectFocusedDelegateHandle);
+		}
+	}
+}
+
+void FChaosVDPlaybackViewportClient::ProcessClick(FSceneView& View, HHitProxy* HitProxy, FKey Key, EInputEvent Event, uint32 HitX, uint32 HitY)
+{
+	if (HitProxy == nullptr)
+	{
+		return;
+	}
+
+	if (TSharedPtr<FChaosVDScene> ScenePtr = CVDScene.Pin())
+	{
+		if (HitProxy->IsA(HActor::StaticGetType()))
+		{
+			HActor* ActorHitProxy = static_cast<HActor*>(HitProxy);
+			if (AActor* ClickedActor = ActorHitProxy->Actor)
+			{
+				ScenePtr->SetSelectedObject(ClickedActor);
+			}
+		}
+	}
+}
+
+void FChaosVDPlaybackViewportClient::SetScene(TWeakPtr<FChaosVDScene> InScene)
+{
+	if (TSharedPtr<FChaosVDScene> ScenePtr = InScene.Pin())
+	{
+		CVDWorld = ScenePtr->GetUnderlyingWorld();
+		CVDScene = InScene;
+
+		ObjectFocusedDelegateHandle = ScenePtr->OnObjectFocused().AddRaw(this, &FChaosVDPlaybackViewportClient::HandleObjectFocused);
+	}
+}
+
+void FChaosVDPlaybackViewportClient::HandleObjectFocused(UObject* FocusedObject)
+{
+	if (AActor* FocusedActor = Cast<AActor>(FocusedObject))
+	{
+		FocusViewportOnBox(FocusedActor->GetComponentsBoundingBox(false));
+	}
+}
diff --git a/Engine/Plugins/Experimental/ChaosVD/Source/ChaosVD/Private/ChaosVDPlaybackViewportTab.cpp b/Engine/Plugins/Experimental/ChaosVD/Source/ChaosVD/Private/ChaosVDPlaybackViewportTab.cpp
--- a/Engine/Plugins/Experimental/ChaosVD/Source/ChaosVD/Private/ChaosVDPlaybackViewportTab.cpp
+++ b/Engine/Plugins/Experimental/ChaosVD/Source/ChaosVD/Private/ChaosVDPlaybackViewportTab.cpp
@@ -22,7 +22,7 @@ TSharedRef<SDockTab> FChaosVDPlaybackViewportTab::HandleTabSpawned(const FSpawnT
 	ViewportTab->SetContent
 	(
 		//TODO: Handle Null cases to not crash the Editor
-		SAssignNew(PlaybackViewportWidget, SChaosVDPlaybackViewport, GetChaosVDWorld(), OwningTabWidget->GetChaosVDEngineInstance()->GetPlaybackController())
+		SAssignNew(PlaybackViewportWidget, SChaosVDPlaybackViewport, GetChaosVDScene(), OwningTabWidget->GetChaosVDEngineInstance()->GetPlaybackController())
 	);
 
 	ViewportTab->SetTabIcon(FChaosVDStyle::Get().GetBrush("TabIconPlaybackViewport"));
diff --git a/Engine/Plugins/Experimental/ChaosVD/Source/ChaosVD/Private/ChaosVDScene.cpp b/Engine/Plugins/Experimental/ChaosVD/Source/ChaosVD/Private/ChaosVDScene.cpp
--- a/Engine/Plugins/Experimental/ChaosVD/Source/ChaosVD/Private/ChaosVDScene.cpp
+++ b/Engine/Plugins/Experimental/ChaosVD/Source/ChaosVD/Private/ChaosVDScene.cpp
@@ -7,8 +7,8 @@
 #include "ChaosVDParticleActor.h"
 #include "Chaos/ImplicitObject.h"
 #include "ChaosVDRecording.h"
-#include "EditorActorFolders.h"
-#include "WorldPersistentFolders.h"
+#include "Elements/Framework/EngineElementsLibrary.h"
+#include "Elements/Framework/TypedElementSelectionSet.h"
 #include "Engine/Engine.h"
 #include "Engine/LevelStreamingDynamic.h"
 #include "Engine/World.h"
@@ -26,6 +26,11 @@ void FChaosVDScene::Initialize()
 		return;
 	}
 
+	SelectionSet = NewObject<UTypedElementSelectionSet>();
+
+	SelectionSet->OnPreChange().AddRaw(this, &FChaosVDScene::HandlePreSelectionChange);
+	SelectionSet->OnChanged().AddRaw(this, &FChaosVDScene::HandlePostSelectionChange);
+	
 	PhysicsVDWorld = CreatePhysicsVDWorld();
 
 	GeometryGenerator = MakeShared<FChaosVDGeometryBuilder>();
@@ -40,6 +45,9 @@ void FChaosVDScene::DeInitialize()
 		return;
 	}
 
+	SelectionSet->OnPreChange().RemoveAll(this);
+	SelectionSet->OnChanged().RemoveAll(this);
+
 	GeometryGenerator.Reset();
 
 	CleanUpScene();
@@ -61,6 +69,7 @@ void FChaosVDScene::DeInitialize()
 void FChaosVDScene::AddReferencedObjects(FReferenceCollector& Collector)
 {
 	Collector.AddReferencedObject(PhysicsVDWorld);
+	Collector.AddReferencedObject(SelectionSet);
 }
 
 void FChaosVDScene::UpdateFromRecordedStepData(const int32 SolverID, const FString& SolverName, const FChaosVDStepData& InRecordedStepData, const FChaosVDSolverFrameData& InFrameData)
@@ -214,6 +223,7 @@ AChaosVDParticleActor* FChaosVDScene::SpawnParticleFromRecordedData(const FChaos
 	FActorSpawnParameters Params;
 	Params.Name = *InParticleData.DebugName;
 	Params.NameMode = FActorSpawnParameters::ESpawnActorNameMode::Requested;
+
 	if (AChaosVDParticleActor* NewActor = PhysicsVDWorld->SpawnActor<AChaosVDParticleActor>(Params))
 	{
 		NewActor->UpdateFromRecordedData(InParticleData, InFrameData.SimulationTransform);
@@ -249,7 +259,7 @@ UWorld* FChaosVDScene::CreatePhysicsVDWorld() const
 	const FName UniqueWorldName = FName(FGuid::NewGuid().ToString());
 	UWorld* NewWorld = NewWorld = NewObject<UWorld>( GetTransientPackage(), UniqueWorldName );
 	
-	NewWorld->WorldType = EWorldType::EditorPreview;
+	NewWorld->WorldType = EWorldType::Editor;
 
 	FWorldContext& WorldContext = GEngine->CreateNewWorldContext( NewWorld->WorldType );
 	WorldContext.SetCurrentWorld(NewWorld);
@@ -257,7 +267,7 @@ UWorld* FChaosVDScene::CreatePhysicsVDWorld() const
 	NewWorld->InitializeNewWorld( UWorld::InitializationValues()
 										  .AllowAudioPlayback( false )
 										  .CreatePhysicsScene( false )
-										  .RequiresHitProxies( false )
+										  .RequiresHitProxies( true )
 										  .CreateNavigation( false )
 										  .CreateAISystem( false )
 										  .ShouldSimulatePhysics( false )
@@ -286,3 +296,69 @@ UWorld* FChaosVDScene::CreatePhysicsVDWorld() const
 
 	return NewWorld;
 }
+
+FTypedElementHandle FChaosVDScene::GetSelectionHandleForObject(const UObject* Object) const
+{
+	FTypedElementHandle Handle;
+	if (const AActor* Actor = Cast<AActor>(Object))
+	{
+		Handle = UEngineElementsLibrary::AcquireEditorActorElementHandle(Actor);
+	}
+	else if (const UActorComponent* Component = Cast<UActorComponent>(Object))
+	{
+		Handle = UEngineElementsLibrary::AcquireEditorComponentElementHandle(Component);
+	}
+	else
+	{
+		Handle = UEngineElementsLibrary::AcquireEditorObjectElementHandle(Object);
+	}
+
+	return Handle;
+}
+
+void FChaosVDScene::UpdateSelectionProxiesForActors(const TArray<AActor*>& SelectedActors)
+{
+	for (AActor* SelectedActor : SelectedActors)
+	{
+		if (SelectedActor)
+		{
+			SelectedActor->PushSelectionToProxies();
+		}
+	}
+}
+
+void FChaosVDScene::HandlePreSelectionChange(const UTypedElementSelectionSet* PreChangeSelectionSet)
+{
+	PendingActorsToUpdateSelectionProxy.Append(PreChangeSelectionSet->GetSelectedObjects<AActor>());
+}
+
+void FChaosVDScene::HandlePostSelectionChange(const UTypedElementSelectionSet* PreChangeSelectionSet)
+{
+	TArray<AActor*> SelectedActors = PreChangeSelectionSet->GetSelectedObjects<AActor>();
+
+	SelectedActors.Append(PendingActorsToUpdateSelectionProxy);
+	UpdateSelectionProxiesForActors(SelectedActors);
+
+	PendingActorsToUpdateSelectionProxy.Reset();
+}
+
+void FChaosVDScene::SetSelectedObject(UObject* SelectedObject)
+{
+	if (IsObjectSelected(SelectedObject))
+	{
+		// Already selected, nothing to do here
+		return;
+	}
+
+	SelectionSet->ClearSelection(FTypedElementSelectionOptions());
+
+	TArray<FTypedElementHandle> NewEditorSelection = { GetSelectionHandleForObject(SelectedObject) };
+
+	SelectionSet->SetSelection(NewEditorSelection, FTypedElementSelectionOptions());
+	SelectionSet->NotifyPendingChanges();
+}
+
+bool FChaosVDScene::IsObjectSelected(const UObject* Object)
+{
+	return SelectionSet->IsElementSelected(GetSelectionHandleForObject(Object), FTypedElementIsSelectedOptions());;
+}
diff --git a/Engine/Plugins/Experimental/ChaosVD/Source/ChaosVD/Private/ChaosVDSceneSelectionObserver.cpp b/Engine/Plugins/Experimental/ChaosVD/Source/ChaosVD/Private/ChaosVDSceneSelectionObserver.cpp
new file mode 100644
--- /dev/null
+++ b/Engine/Plugins/Experimental/ChaosVD/Source/ChaosVD/Private/ChaosVDSceneSelectionObserver.cpp
@@ -0,0 +1,34 @@
+﻿// Copyright Epic Games, Inc. All Rights Reserved.
+
+#include "ChaosVDSceneSelectionObserver.h"
+
+#include "Elements/Framework/TypedElementSelectionSet.h"
+
+FChaosVDSceneSelectionObserver::FChaosVDSceneSelectionObserver()
+{
+}
+
+FChaosVDSceneSelectionObserver::~FChaosVDSceneSelectionObserver()
+{
+	if (ObservedSelection.IsValid())
+	{
+		ObservedSelection->OnPreChange().RemoveAll(this);
+		ObservedSelection->OnChanged().RemoveAll(this);
+	}
+}
+
+void FChaosVDSceneSelectionObserver::RegisterSelectionSetObject(UTypedElementSelectionSet* SelectionSetObject)
+{
+	if (ObservedSelection.IsValid())
+	{
+		ObservedSelection->OnPreChange().RemoveAll(this);
+		ObservedSelection->OnChanged().RemoveAll(this);
+	}
+
+	if (SelectionSetObject)
+	{
+		ObservedSelection = TWeakObjectPtr<UTypedElementSelectionSet>(SelectionSetObject);
+		ObservedSelection->OnPreChange().AddRaw(this, &FChaosVDSceneSelectionObserver::HandlePreSelectionChange);
+		ObservedSelection->OnChanged().AddRaw(this, &FChaosVDSceneSelectionObserver::HandlePostSelectionChange);
+	}
+}
diff --git a/Engine/Plugins/Experimental/ChaosVD/Source/ChaosVD/Private/ChaosVDTabSpawnerBase.cpp b/Engine/Plugins/Experimental/ChaosVD/Source/ChaosVD/Private/ChaosVDTabSpawnerBase.cpp
--- a/Engine/Plugins/Experimental/ChaosVD/Source/ChaosVD/Private/ChaosVDTabSpawnerBase.cpp
+++ b/Engine/Plugins/Experimental/ChaosVD/Source/ChaosVD/Private/ChaosVDTabSpawnerBase.cpp
@@ -14,10 +14,23 @@ FChaosVDTabSpawnerBase::FChaosVDTabSpawnerBase(const FName& InTabID, TSharedPtr<
 
 UWorld* FChaosVDTabSpawnerBase::GetChaosVDWorld() const
 {
-	if (ensure(OwningTabWidget) && OwningTabWidget->GetChaosVDEngineInstance()->GetCurrentScene().IsValid())
+	if (TSharedPtr<FChaosVDScene> ScenePtr = GetChaosVDScene().Pin())
 	{
-		return OwningTabWidget->GetChaosVDEngineInstance()->GetCurrentScene()->GetUnderlyingWorld();
+		return ScenePtr->GetUnderlyingWorld();
 	}
 
 	return nullptr;
 }
+
+TWeakPtr<FChaosVDScene> FChaosVDTabSpawnerBase::GetChaosVDScene() const
+{
+	if (ensure(OwningTabWidget))
+	{
+		const TSharedRef<FChaosVDEngine> ChaosVDEngine = OwningTabWidget->GetChaosVDEngineInstance();
+		if (TSharedPtr<FChaosVDScene> ScenePtr = ChaosVDEngine->GetCurrentScene())
+		{
+			return ScenePtr;
+		}
+	}
+	return nullptr;
+}
diff --git a/Engine/Plugins/Experimental/ChaosVD/Source/ChaosVD/Private/ChaosVDWorldOutlinerMode.cpp b/Engine/Plugins/Experimental/ChaosVD/Source/ChaosVD/Private/ChaosVDWorldOutlinerMode.cpp
new file mode 100644
--- /dev/null
+++ b/Engine/Plugins/Experimental/ChaosVD/Source/ChaosVD/Private/ChaosVDWorldOutlinerMode.cpp
@@ -0,0 +1,95 @@
+﻿// Copyright Epic Games, Inc. All Rights Reserved.
+
+#include "ChaosVDWorldOutlinerMode.h"
+
+#include "ActorTreeItem.h"
+#include "ChaosVDScene.h"
+#include "Elements/Framework/TypedElementSelectionSet.h"
+
+
+FChaosVDWorldOutlinerMode::FChaosVDWorldOutlinerMode(const FActorModeParams& InModeParams, TWeakPtr<FChaosVDScene> InScene)
+	: FActorMode(InModeParams),
+	CVDScene(InScene)
+{
+	TSharedPtr<FChaosVDScene> ScenePtr = CVDScene.Pin();
+	if (!ensure(ScenePtr.IsValid()))
+	{
+		return;
+	}
+
+	RegisterSelectionSetObject(ScenePtr->GetElementSelectionSet());
+
+	ActorLabelChangedDelegateHandle = FCoreDelegates::OnActorLabelChanged.AddRaw(this, &FChaosVDWorldOutlinerMode::HandleActorLabelChanged);
+}
+
+FChaosVDWorldOutlinerMode::~FChaosVDWorldOutlinerMode()
+{
+	FCoreDelegates::OnActorLabelChanged.Remove(ActorLabelChangedDelegateHandle);
+}
+
+void FChaosVDWorldOutlinerMode::OnItemSelectionChanged(FSceneOutlinerTreeItemPtr Item, ESelectInfo::Type SelectionType, const FSceneOutlinerItemSelection& Selection)
+{
+	TSharedPtr<FChaosVDScene> ScenePtr = CVDScene.Pin();
+	if (!ScenePtr.IsValid())
+	{
+		return;
+	}
+
+	TArray<AActor*> OutlinerSelectedActors = Selection.GetData<AActor*>(SceneOutliner::FActorSelector());
+
+	if (const UTypedElementSelectionSet* SelectionSet = GetSelectionSetObject())
+	{
+		TArray<AActor*> SceneSelectedActors = SelectionSet->GetSelectedObjects<AActor>();
+		ensureMsgf(SceneSelectedActors.Num() < 2, TEXT("Multi Selection is not supported, but [%s] Actors are selected... Choosing the first one"));
+
+		ScenePtr->SetSelectedObject(OutlinerSelectedActors[0]);
+	}
+}
+
+void FChaosVDWorldOutlinerMode::OnItemDoubleClick(FSceneOutlinerTreeItemPtr Item)
+{
+	TSharedPtr<FChaosVDScene> ScenePtr = CVDScene.Pin();
+	if (!ScenePtr.IsValid())
+	{
+		return;
+	}
+
+	if (const FActorTreeItem* ActorItem = Item->CastTo<FActorTreeItem>())
+	{
+		if (AActor* Actor = ActorItem->Actor.Get())
+		{
+			ScenePtr->OnObjectFocused().Broadcast(Actor);
+		}
+	}
+}
+
+void FChaosVDWorldOutlinerMode::HandleActorLabelChanged(AActor* ChangedActor)
+{
+	if (!ensure(ChangedActor))
+	{
+		return;
+	}
+
+	if (IsActorDisplayable(ChangedActor) && RepresentingWorld.Get() == ChangedActor->GetWorld())
+	{
+		// Force create the item otherwise the outliner may not be notified of a change to the item if it is filtered out
+		if (FSceneOutlinerTreeItemPtr Item = CreateItemFor<FActorTreeItem>(ChangedActor, true))
+		{
+			SceneOutliner->OnItemLabelChanged(Item);
+		}
+	}
+}
+
+void FChaosVDWorldOutlinerMode::HandlePostSelectionChange(const UTypedElementSelectionSet* ChangesSelectionSet)
+{
+	TArray<AActor*> SelectedActors = ChangesSelectionSet->GetSelectedObjects<AActor>();
+
+	if (SelectedActors.Num() > 0)
+	{
+		// We don't support multi selection yet
+		ensure(SelectedActors.Num() == 1);
+
+		FSceneOutlinerTreeItemPtr ObjectItem = SceneOutliner->GetTreeItem(SelectedActors[0]);
+		SceneOutliner->SetItemSelection(ObjectItem, true, ESelectInfo::OnMouseClick);
+	}
+}
diff --git a/Engine/Plugins/Experimental/ChaosVD/Source/ChaosVD/Private/ChaosVDWorldOutlinerMode.h b/Engine/Plugins/Experimental/ChaosVD/Source/ChaosVD/Private/ChaosVDWorldOutlinerMode.h
new file mode 100644
--- /dev/null
+++ b/Engine/Plugins/Experimental/ChaosVD/Source/ChaosVD/Private/ChaosVDWorldOutlinerMode.h
@@ -0,0 +1,42 @@
+﻿// Copyright Epic Games, Inc. All Rights Reserved.
+
+#pragma once
+#include "ActorBrowsingMode.h"
+#include "ChaosVDScene.h"
+#include "ChaosVDSceneSelectionObserver.h"
+
+class FChaosVDScene;
+
+/**
+ * Scene outliner mode used to represent a CVD (Chaos Visual Debugger) world
+ * It has a more limited view compared to the normal outliner, hiding features we don't support,
+ * and it is integrated with the CVD local selection system
+ */
+class FChaosVDWorldOutlinerMode : public FActorMode, public FChaosVDSceneSelectionObserver
+{
+public:
+
+	FChaosVDWorldOutlinerMode(const FActorModeParams& InModeParams, TWeakPtr<FChaosVDScene> InScene);
+
+	virtual ~FChaosVDWorldOutlinerMode() override;
+
+	virtual void OnItemSelectionChanged(FSceneOutlinerTreeItemPtr Item, ESelectInfo::Type SelectionType, const FSceneOutlinerItemSelection& Selection) override;
+	virtual void OnItemDoubleClick(FSceneOutlinerTreeItemPtr Item) override;
+
+	virtual void CreateViewContent(FMenuBuilder& MenuBuilder) override
+	{
+		// Intentionally not implemented, as we don't support the built in menu to switch worlds
+	}
+
+	virtual bool ShouldShowFolders() const override { return true;};
+
+private:
+
+	void HandleActorLabelChanged(AActor* ChangedActor);
+
+	virtual void HandlePostSelectionChange(const UTypedElementSelectionSet* ChangesSelectionSet) override;
+
+	FDelegateHandle ActorLabelChangedDelegateHandle;
+
+	TWeakPtr<FChaosVDScene> CVDScene;
+};
diff --git a/Engine/Plugins/Experimental/ChaosVD/Source/ChaosVD/Private/ChaosVDWorldOutlinerTab.cpp b/Engine/Plugins/Experimental/ChaosVD/Source/ChaosVD/Private/ChaosVDWorldOutlinerTab.cpp
--- a/Engine/Plugins/Experimental/ChaosVD/Source/ChaosVD/Private/ChaosVDWorldOutlinerTab.cpp
+++ b/Engine/Plugins/Experimental/ChaosVD/Source/ChaosVD/Private/ChaosVDWorldOutlinerTab.cpp
@@ -2,8 +2,8 @@
 
 #include "ChaosVDWorldOutlinerTab.h"
 
+#include "ChaosVDWorldOutlinerMode.h"
 #include "ChaosVDStyle.h"
-#include "ChaosVDTabsIDs.h"
 #include "Modules/ModuleManager.h"
 #include "SceneOutlinerModule.h"
 #include "SceneOutlinerPublicTypes.h"
@@ -12,20 +12,38 @@
 
 #define LOCTEXT_NAMESPACE "ChaosVisualDebugger"
 
-TSharedRef<SDockTab> FChaosVDWorldOutlinerTab::HandleTabSpawned(const FSpawnTabArgs& Args)
+void FChaosVDWorldOutlinerTab::CreateWorldOutlinerWidget()
 {
 	FSceneOutlinerInitializationOptions InitOptions;
 	InitOptions.bShowTransient = true;
 
 	InitOptions.OutlinerIdentifier = TEXT("ChaosVDOutliner");
 	InitOptions.FilterBarOptions.bHasFilterBar = true;
-
 	InitOptions.FilterBarOptions.bUseSharedSettings = false;
 
+	InitOptions.UseDefaultColumns();
+	InitOptions.ColumnMap.Add(FSceneOutlinerBuiltInColumnTypes::Gutter(), FSceneOutlinerColumnInfo(ESceneOutlinerColumnVisibility::Visible, 0, FCreateSceneOutlinerColumn(), true, TOptional<float>(), FSceneOutlinerBuiltInColumnTypes::Gutter_Localized()));
+	InitOptions.ColumnMap.Add(FSceneOutlinerBuiltInColumnTypes::ActorInfo(), FSceneOutlinerColumnInfo(ESceneOutlinerColumnVisibility::Visible, 20, FCreateSceneOutlinerColumn(), true, TOptional<float>(), FSceneOutlinerBuiltInColumnTypes::ActorInfo_Localized()));
+		
+	FCreateSceneOutlinerMode ModeFactory = FCreateSceneOutlinerMode::CreateLambda([this](SSceneOutliner* Outliner)
+	{
+		FActorModeParams ModeParams(Outliner);
+		ModeParams.SpecifiedWorldToDisplay = GetChaosVDWorld();
+		ModeParams.bHideEmptyFolders = true;
+
+		// The mode is deleted by the Outliner when it is destroyed 
+		return new FChaosVDWorldOutlinerMode(ModeParams, GetChaosVDScene());
+	});
+	InitOptions.ModeFactory = ModeFactory;
+	
 	const FSceneOutlinerModule& SceneOutlinerModule = FModuleManager::Get().LoadModuleChecked<FSceneOutlinerModule>("SceneOutliner");
 
-	//TODO: See if we can change this to an Actor Picker instead of Actor Browser as it is simpler
-	SceneOutlinerWidget = SceneOutlinerModule.CreateActorBrowser(InitOptions, TWeakObjectPtr<UWorld>(GetChaosVDWorld()));
+	SceneOutlinerWidget = SceneOutlinerModule.CreateSceneOutliner(InitOptions);
+}
+
+TSharedRef<SDockTab> FChaosVDWorldOutlinerTab::HandleTabSpawned(const FSpawnTabArgs& Args)
+{
+	CreateWorldOutlinerWidget();
 
 	TSharedRef<SDockTab> OutlinerTab =
 		SNew(SDockTab)
diff --git a/Engine/Plugins/Experimental/ChaosVD/Source/ChaosVD/Private/Widgets/SChaosVDPlaybackViewport.cpp b/Engine/Plugins/Experimental/ChaosVD/Source/ChaosVD/Private/Widgets/SChaosVDPlaybackViewport.cpp
--- a/Engine/Plugins/Experimental/ChaosVD/Source/ChaosVD/Private/Widgets/SChaosVDPlaybackViewport.cpp
+++ b/Engine/Plugins/Experimental/ChaosVD/Source/ChaosVD/Private/Widgets/SChaosVDPlaybackViewport.cpp
@@ -3,9 +3,11 @@
 #include "Widgets/SChaosVDPlaybackViewport.h"
 
 #include "ChaosVDPlaybackController.h"
+#include "ChaosVDPlaybackViewportClient.h"
 #include "ChaosVDScene.h"
 #include "Framework/Application/SlateApplication.h"
 #include "LevelEditorViewport.h"
+#include "Elements/Framework/TypedElementSelectionSet.h"
 #include "Slate/SceneViewport.h"
 #include "Widgets/SChaosVDTimelineWidget.h"
 #include "Widgets/SViewport.h"
@@ -15,13 +17,13 @@
 
 SChaosVDPlaybackViewport::~SChaosVDPlaybackViewport()
 {
-	LevelViewportClient->Viewport = nullptr;
-	LevelViewportClient.Reset();
+	PlaybackViewportClient->Viewport = nullptr;
+	PlaybackViewportClient.Reset();
 }
 
-TSharedPtr<FLevelEditorViewportClient> SChaosVDPlaybackViewport::CreateViewportClient() const
+TSharedPtr<FChaosVDPlaybackViewportClient> SChaosVDPlaybackViewport::CreateViewportClient() const
 {
-	TSharedPtr<FLevelEditorViewportClient> NewViewport = MakeShareable(new FLevelEditorViewportClient(TSharedPtr<class SLevelViewport>()));
+	TSharedPtr<FChaosVDPlaybackViewportClient> NewViewport = MakeShared<FChaosVDPlaybackViewportClient>();
 
 	NewViewport->SetAllowCinematicControl(false);
 	
@@ -32,16 +34,18 @@ TSharedPtr<FLevelEditorViewportClient> SChaosVDPlaybackViewport::CreateViewportC
 	NewViewport->bDrawAxes = true;
 	NewViewport->bDisableInput = false;
 	NewViewport->VisibilityDelegate.BindLambda([] {return true; });
+	NewViewport->EngineShowFlags.SetSelectionOutline(true);
 
 	return NewViewport;
 }
 
-void SChaosVDPlaybackViewport::Construct(const FArguments& InArgs, const UWorld* DefaultWorld, TWeakPtr<FChaosVDPlaybackController> InPlaybackController)
+void SChaosVDPlaybackViewport::Construct(const FArguments& InArgs, TWeakPtr<FChaosVDScene> InScene, TWeakPtr<FChaosVDPlaybackController> InPlaybackController)
 {
-	ensure(DefaultWorld);
+	TSharedPtr<FChaosVDScene> ScenePtr = InScene.Pin();
+	ensure(ScenePtr.IsValid());
 	ensure(InPlaybackController.IsValid());
 
-	LevelViewportClient = CreateViewportClient();
+	PlaybackViewportClient = CreateViewportClient();
 
 	ViewportWidget = SNew(SViewport)
 		.RenderDirectlyToWindow(false)
@@ -49,14 +53,14 @@ void SChaosVDPlaybackViewport::Construct(const FArguments& InArgs, const UWorld*
 		.EnableGammaCorrection(false)
 		.EnableBlending(false);
 
-	SceneViewport = MakeShareable(new FSceneViewport(LevelViewportClient.Get(), ViewportWidget));
+	SceneViewport = MakeShareable(new FSceneViewport(PlaybackViewportClient.Get(), ViewportWidget));
 
-	LevelViewportClient->Viewport = SceneViewport.Get();
+	PlaybackViewportClient->Viewport = SceneViewport.Get();
 
 	ViewportWidget->SetViewportInterface(SceneViewport.ToSharedRef());
 	
 	// Default to the base map
-	LevelViewportClient->SetReferenceToWorldContext(*GEngine->GetWorldContextFromWorld(DefaultWorld));
+	PlaybackViewportClient->SetScene(InScene);
 
 	ChildSlot
 	[
@@ -118,7 +122,7 @@ void SChaosVDPlaybackViewport::HandlePlaybackControllerDataUpdated(TWeakPtr<FCha
 		GameFramesTimelineWidget->ResetTimeline();
 	}
 
-	LevelViewportClient->bNeedsRedraw = true;
+	PlaybackViewportClient->bNeedsRedraw = true;
 }
 
 void SChaosVDPlaybackViewport::HandleControllerTrackFrameUpdated(TWeakPtr<FChaosVDPlaybackController> InController, const FChaosVDTrackInfo* UpdatedTrackInfo, FGuid InstigatorGuid)
@@ -137,9 +141,14 @@ void SChaosVDPlaybackViewport::HandleControllerTrackFrameUpdated(TWeakPtr<FChaos
 	}
 }
 
+void SChaosVDPlaybackViewport::HandlePostSelectionChange(const UTypedElementSelectionSet* ChangesSelectionSet)
+{
+	PlaybackViewportClient->bNeedsRedraw = true;
+}
+
 void SChaosVDPlaybackViewport::OnPlaybackSceneUpdated()
 {
-	LevelViewportClient->bNeedsRedraw = true;
+	PlaybackViewportClient->bNeedsRedraw = true;
 }
 
 void SChaosVDPlaybackViewport::RegisterNewController(TWeakPtr<FChaosVDPlaybackController> NewController)
@@ -173,7 +182,7 @@ void SChaosVDPlaybackViewport::OnFrameSelectionUpdated(int32 NewFrameIndex) cons
 		constexpr int32 StepNumber = 0;
 		PlaybackControllerPtr->GoToTrackFrame(GetInstigatorID(), EChaosVDTrackType::Game, FChaosVDPlaybackController::GameTrackID, NewFrameIndex, StepNumber);
 
-		LevelViewportClient->bNeedsRedraw = true;
+		PlaybackViewportClient->bNeedsRedraw = true;
 	}
 }
 
diff --git a/Engine/Plugins/Experimental/ChaosVD/Source/ChaosVD/Public/ChaosVDObjectDetailsTab.h b/Engine/Plugins/Experimental/ChaosVD/Source/ChaosVD/Public/ChaosVDObjectDetailsTab.h
--- a/Engine/Plugins/Experimental/ChaosVD/Source/ChaosVD/Public/ChaosVDObjectDetailsTab.h
+++ b/Engine/Plugins/Experimental/ChaosVD/Source/ChaosVD/Public/ChaosVDObjectDetailsTab.h
@@ -2,6 +2,7 @@
 
 #pragma once
 
+#include "ChaosVDSceneSelectionObserver.h"
 #include "Templates/SharedPointer.h"
 #include "ChaosVDTabSpawnerBase.h"
 #include "Delegates/IDelegateInstance.h"
@@ -17,7 +18,7 @@ class SDockTab;
 
 
 /** Spawns and handles and instance for the visual debugger details panel */
-class FChaosVDObjectDetailsTab : public FChaosVDTabSpawnerBase, public TSharedFromThis<FChaosVDObjectDetailsTab>
+class FChaosVDObjectDetailsTab : public FChaosVDTabSpawnerBase, public FChaosVDSceneSelectionObserver, public TSharedFromThis<FChaosVDObjectDetailsTab>
 {
 public:
 
@@ -25,13 +26,11 @@ public:
 	{
 	}
 
-	virtual ~FChaosVDObjectDetailsTab() override;
-
 protected:
-	
+
 	virtual TSharedRef<SDockTab> HandleTabSpawned(const FSpawnTabArgs& Args) override;
-	
-	void UpdateSelectedObject(AActor* NewObject) const;
+
+	virtual void HandlePostSelectionChange(const UTypedElementSelectionSet* ChangedSelectionSet) override;
 
 	FDelegateHandle SelectionDelegateHandle;
 	TSharedPtr<IDetailsView> DetailsPanel;
diff --git a/Engine/Plugins/Experimental/ChaosVD/Source/ChaosVD/Public/ChaosVDParticleActor.h b/Engine/Plugins/Experimental/ChaosVD/Source/ChaosVD/Public/ChaosVDParticleActor.h
--- a/Engine/Plugins/Experimental/ChaosVD/Source/ChaosVD/Public/ChaosVDParticleActor.h
+++ b/Engine/Plugins/Experimental/ChaosVD/Source/ChaosVD/Public/ChaosVDParticleActor.h
@@ -39,6 +39,10 @@ public:
 	void SetScene(const TSharedPtr<FChaosVDScene>& InScene);
 
 	virtual void BeginDestroy() override;
+	
+#if WITH_EDITOR
+	virtual bool IsSelectedInEditor() const override;
+#endif
 
 protected:
 
diff --git a/Engine/Plugins/Experimental/ChaosVD/Source/ChaosVD/Public/ChaosVDPlaybackControllerObserver.h b/Engine/Plugins/Experimental/ChaosVD/Source/ChaosVD/Public/ChaosVDPlaybackControllerObserver.h
--- a/Engine/Plugins/Experimental/ChaosVD/Source/ChaosVD/Public/ChaosVDPlaybackControllerObserver.h
+++ b/Engine/Plugins/Experimental/ChaosVD/Source/ChaosVD/Public/ChaosVDPlaybackControllerObserver.h
@@ -2,6 +2,7 @@
 
 #pragma once
 
+#include "ChaosVDSceneSelectionObserver.h"
 #include "Misc/Guid.h"
 #include "Templates/SharedPointer.h"
 
@@ -10,17 +11,21 @@ class FChaosVDPlaybackController;
 
 
 /**
- * Class to be used as base to any object that needs to Process Player Controller changes
+ * Class to be used as base to any object that needs to Process Player Controller changes,
+ * and scene Selection Changes
  */
-class FChaosVDPlaybackControllerObserver
+class FChaosVDPlaybackControllerObserver : public FChaosVDSceneSelectionObserver
 {
 public:
-	virtual ~FChaosVDPlaybackControllerObserver();
+
+	virtual ~FChaosVDPlaybackControllerObserver() override;
 
 protected:
 	virtual void RegisterNewController(TWeakPtr<FChaosVDPlaybackController> NewController);
 	virtual void HandlePlaybackControllerDataUpdated(TWeakPtr<FChaosVDPlaybackController> InController) {};
 	virtual void HandleControllerTrackFrameUpdated(TWeakPtr<FChaosVDPlaybackController> InController, const FChaosVDTrackInfo* UpdatedTrackInfo, FGuid InstigatorGuid){};
 	
+	virtual void HandlePostSelectionChange(const UTypedElementSelectionSet* ChangesSelectionSet) override{};
+	
 	TWeakPtr<FChaosVDPlaybackController> PlaybackController;
 };
diff --git a/Engine/Plugins/Experimental/ChaosVD/Source/ChaosVD/Public/ChaosVDPlaybackViewportClient.h b/Engine/Plugins/Experimental/ChaosVD/Source/ChaosVD/Public/ChaosVDPlaybackViewportClient.h
new file mode 100644
--- /dev/null
+++ b/Engine/Plugins/Experimental/ChaosVD/Source/ChaosVD/Public/ChaosVDPlaybackViewportClient.h
@@ -0,0 +1,30 @@
+﻿// Copyright Epic Games, Inc. All Rights Reserved.
+
+#pragma once
+#include "LevelEditorViewport.h"
+
+class FChaosVDScene;
+
+/** Client viewport class used for to handle a Chaos Visual Debugger world Interaction/Rendering.
+ * It re-routes interaction events to our Chaos VD scene
+ */
+class FChaosVDPlaybackViewportClient : public FEditorViewportClient
+{
+public:
+	FChaosVDPlaybackViewportClient();
+	virtual ~FChaosVDPlaybackViewportClient() override;
+
+	virtual void ProcessClick(FSceneView& View, HHitProxy* HitProxy, FKey Key, EInputEvent Event, uint32 HitX, uint32 HitY) override;
+
+	void SetScene(TWeakPtr<FChaosVDScene> InScene);
+
+	virtual UWorld* GetWorld() const override { return CVDWorld; };
+
+private:
+
+	void HandleObjectFocused(UObject* FocusedObject);
+
+	FDelegateHandle ObjectFocusedDelegateHandle;
+	UWorld* CVDWorld;
+	TWeakPtr<FChaosVDScene> CVDScene;
+};
diff --git a/Engine/Plugins/Experimental/ChaosVD/Source/ChaosVD/Public/ChaosVDScene.h b/Engine/Plugins/Experimental/ChaosVD/Source/ChaosVD/Public/ChaosVDScene.h
--- a/Engine/Plugins/Experimental/ChaosVD/Source/ChaosVD/Public/ChaosVDScene.h
+++ b/Engine/Plugins/Experimental/ChaosVD/Source/ChaosVD/Public/ChaosVDScene.h
@@ -11,11 +11,16 @@
 class FChaosVDGeometryBuilder;
 class AChaosVDParticleActor;
 class FReferenceCollector;
+class UObject;
+class UTypedElementSelectionSet;
 class UWorld;
 
+struct FTypedElementHandle;
+
 typedef TMap<int32, AChaosVDParticleActor*> FChaosVDParticlesByIDMap;
 
 DECLARE_MULTICAST_DELEGATE(FChaosVDSceneUpdatedDelegate)
+DECLARE_MULTICAST_DELEGATE_OneParam(FChaosVDOnObjectSelectedDelegate, UObject*)
 
 /** Recreates a UWorld from a recorded Chaos VD Frame */
 class FChaosVDScene : public FGCObject , public TSharedFromThis<FChaosVDScene>
@@ -57,6 +62,18 @@ public:
 
 	const TSharedPtr<const Chaos::FImplicitObject>* GetUpdatedGeometry(int32 GeometryID) const;
 
+	/** Adds an object to the selection set if it was not selected already, making it selected in practice */
+	void SetSelectedObject(UObject* SelectedObject);
+
+	/** Evaluates an object and returns true if it is selected */
+	bool IsObjectSelected(const UObject* Object);
+
+	/** Returns a ptr to the current selection set object */
+	UTypedElementSelectionSet* GetElementSelectionSet() const { return SelectionSet; }
+	
+	/** Event triggered when an object is focused in the scene (double click in the scene outliner)*/
+	FChaosVDOnObjectSelectedDelegate& OnObjectFocused() { return ObjectFocusedDelegate; }
+
 	TSharedPtr<FChaosVDRecording> LoadedRecording;
 
 private:
@@ -73,6 +90,15 @@ private:
 	/** Map of ID-ChaosVDParticle Actor. Used to keep track of actor instances and be able to modify them as needed*/
 	TMap<int32, FChaosVDParticlesByIDMap> ParticlesBySolverID;
 
+	/** Returns the correct TypedElementHandle based on an object type so it can be used with the selection set object */
+	FTypedElementHandle GetSelectionHandleForObject(const UObject* Object) const;
+
+	/** Updates the render state of the hit proxies of an array of actors. This used to update the selection outline state */
+	void UpdateSelectionProxiesForActors(const TArray<AActor*>& SelectedActors);
+
+	void HandlePreSelectionChange(const UTypedElementSelectionSet* PreChangeSelectionSet);
+	void HandlePostSelectionChange(const UTypedElementSelectionSet* PreChangeSelectionSet);
+
 	/** UWorld instance used to represent the recorded debug data */
 	UWorld* PhysicsVDWorld = nullptr;
 
@@ -82,5 +108,13 @@ private:
 
 	FChaosVDGeometryDataLoaded NewGeometryAvailableDelegate;
 
+	FChaosVDOnObjectSelectedDelegate ObjectFocusedDelegate;
+
+	/** Selection set object holding the current selection state */
+	UTypedElementSelectionSet* SelectionSet;
+
+	/** Array of actors with hit proxies that need to be updated */
+	TArray<AActor*> PendingActorsToUpdateSelectionProxy;
+
 	bool bIsInitialized = false;
 };
diff --git a/Engine/Plugins/Experimental/ChaosVD/Source/ChaosVD/Public/ChaosVDSceneSelectionObserver.h b/Engine/Plugins/Experimental/ChaosVD/Source/ChaosVD/Public/ChaosVDSceneSelectionObserver.h
new file mode 100644
--- /dev/null
+++ b/Engine/Plugins/Experimental/ChaosVD/Source/ChaosVD/Public/ChaosVDSceneSelectionObserver.h
@@ -0,0 +1,39 @@
+﻿// Copyright Epic Games, Inc. All Rights Reserved.
+
+#pragma once
+#include "UObject/WeakObjectPtrTemplates.h"
+#include "UObject/WeakObjectPtr.h"
+
+class FChaosVDScene;
+class UTypedElementSelectionSet;
+
+/** Base class that provided access to selection events from the Chaos visual debugger
+ * local selection system.
+ */
+class FChaosVDSceneSelectionObserver
+{
+public:
+
+	FChaosVDSceneSelectionObserver();
+	virtual ~FChaosVDSceneSelectionObserver();
+
+protected:
+
+	/** Stores a weak ptr to the selection set object containing the current selection state */
+	void RegisterSelectionSetObject(UTypedElementSelectionSet* SelectionSetObject);
+	
+	const UTypedElementSelectionSet* GetSelectionSetObject() const { return ObservedSelection.IsValid() ? ObservedSelection.Get() : nullptr; };
+
+	/** Called when the current selection might change
+	 * @param SelectionSetPreChange the pointer to the selection set object that might be about to change
+	 */
+	virtual void HandlePreSelectionChange(const UTypedElementSelectionSet* SelectionSetPreChange){}
+	
+	/** Called when the current selection has changed
+	 * @param ChangesSelectionSet the pointer to the selection set object that has changed
+	 */
+	virtual void HandlePostSelectionChange(const UTypedElementSelectionSet* ChangesSelectionSet){}
+
+	/** Weak Ptr to the selection set object we are observing */
+	TWeakObjectPtr<UTypedElementSelectionSet> ObservedSelection;
+};
diff --git a/Engine/Plugins/Experimental/ChaosVD/Source/ChaosVD/Public/ChaosVDTabSpawnerBase.h b/Engine/Plugins/Experimental/ChaosVD/Source/ChaosVD/Public/ChaosVDTabSpawnerBase.h
--- a/Engine/Plugins/Experimental/ChaosVD/Source/ChaosVD/Public/ChaosVDTabSpawnerBase.h
+++ b/Engine/Plugins/Experimental/ChaosVD/Source/ChaosVD/Public/ChaosVDTabSpawnerBase.h
@@ -5,6 +5,7 @@
 #include "Templates/SharedPointer.h"
 #include "UObject/NameTypes.h"
 
+class FChaosVDScene;
 class SChaosVDMainTab;
 class FName;
 class FTabManager;
@@ -27,4 +28,6 @@ protected:
 	SChaosVDMainTab* OwningTabWidget;
 
 	UWorld* GetChaosVDWorld() const;
+
+	TWeakPtr<FChaosVDScene> GetChaosVDScene() const;
 };
diff --git a/Engine/Plugins/Experimental/ChaosVD/Source/ChaosVD/Public/ChaosVDWorldOutlinerTab.h b/Engine/Plugins/Experimental/ChaosVD/Source/ChaosVD/Public/ChaosVDWorldOutlinerTab.h
--- a/Engine/Plugins/Experimental/ChaosVD/Source/ChaosVD/Public/ChaosVDWorldOutlinerTab.h
+++ b/Engine/Plugins/Experimental/ChaosVD/Source/ChaosVD/Public/ChaosVDWorldOutlinerTab.h
@@ -2,6 +2,7 @@
 
 #pragma once
 
+#include "ChaosVDSceneSelectionObserver.h"
 #include "ChaosVDTabSpawnerBase.h"
 #include "Delegates/DelegateCombinations.h"
 #include "Delegates/Delegate.h"
@@ -26,16 +27,13 @@ class FChaosVDWorldOutlinerTab : public FChaosVDTabSpawnerBase, public TSharedFr
 {
 public:
 	FChaosVDWorldOutlinerTab(const FName& InTabID, TSharedPtr<FTabManager> InTabManager, SChaosVDMainTab* InOwningTabWidget) : FChaosVDTabSpawnerBase(InTabID, InTabManager, InOwningTabWidget)
-	{	
+	{
 	}
 
 	virtual TSharedRef<SDockTab> HandleTabSpawned(const FSpawnTabArgs& Args) override;
 
-	FChaosVDSelectionChanged& OnSelectionChanged() { return ChaosVDSelectionChangedDelegate; }
-
 private:
+	void CreateWorldOutlinerWidget();
 
 	TSharedPtr<ISceneOutliner> SceneOutlinerWidget;
-
-	FChaosVDSelectionChanged ChaosVDSelectionChangedDelegate;
 };
diff --git a/Engine/Plugins/Experimental/ChaosVD/Source/ChaosVD/Public/Widgets/SChaosVDPlaybackViewport.h b/Engine/Plugins/Experimental/ChaosVD/Source/ChaosVD/Public/Widgets/SChaosVDPlaybackViewport.h
--- a/Engine/Plugins/Experimental/ChaosVD/Source/ChaosVD/Public/Widgets/SChaosVDPlaybackViewport.h
+++ b/Engine/Plugins/Experimental/ChaosVD/Source/ChaosVD/Public/Widgets/SChaosVDPlaybackViewport.h
@@ -7,6 +7,7 @@
 #include "Templates/SharedPointer.h"
 #include "Widgets/SCompoundWidget.h"
 
+class FChaosVDPlaybackViewportClient;
 class SChaosVDSolverPlaybackControls;
 class FChaosVDPlaybackController;
 class SChaosVDTimelineWidget;
@@ -26,22 +27,24 @@ public:
 
 	virtual ~SChaosVDPlaybackViewport() override;
 
-	void Construct(const FArguments& InArgs, const UWorld* DefaultWorld, TWeakPtr<FChaosVDPlaybackController> InPlaybackController);
+	void Construct(const FArguments& InArgs, TWeakPtr<FChaosVDScene> InScene, TWeakPtr<FChaosVDPlaybackController> InPlaybackController);
 
 protected:
 
-	TSharedPtr<FLevelEditorViewportClient> CreateViewportClient() const;
+	TSharedPtr<FChaosVDPlaybackViewportClient> CreateViewportClient() const;
 
 	virtual void RegisterNewController(TWeakPtr<FChaosVDPlaybackController> NewController )override;
 	virtual void HandlePlaybackControllerDataUpdated(TWeakPtr<FChaosVDPlaybackController> InController) override;
 	virtual void HandleControllerTrackFrameUpdated(TWeakPtr<FChaosVDPlaybackController> InController, const FChaosVDTrackInfo* UpdatedTrackInfo, FGuid InstigatorGuid) override;
+	virtual void HandlePostSelectionChange(const UTypedElementSelectionSet* ChangesSelectionSet) override;
+
 	void OnPlaybackSceneUpdated();
 
 	void OnFrameSelectionUpdated(int32 NewFrameIndex) const;
-	
+
 	TSharedPtr<SChaosVDTimelineWidget> GameFramesTimelineWidget;
 
-	TSharedPtr<FLevelEditorViewportClient> LevelViewportClient;
+	TSharedPtr<FChaosVDPlaybackViewportClient> PlaybackViewportClient;
 	TSharedPtr<SViewport> ViewportWidget;
 	TSharedPtr<FSceneViewport> SceneViewport;
 };
diff --git a/Engine/Source/Editor/SceneOutliner/Public/ActorMode.h b/Engine/Source/Editor/SceneOutliner/Public/ActorMode.h
--- a/Engine/Source/Editor/SceneOutliner/Public/ActorMode.h
+++ b/Engine/Source/Editor/SceneOutliner/Public/ActorMode.h
@@ -13,7 +13,7 @@ namespace SceneOutliner
 	};
 
 	/** Functor which can be used to get actors from a selection including component parents */
-	struct FActorSelector
+	struct SCENEOUTLINER_API FActorSelector
 	{
 		bool operator()(const TWeakPtr<ISceneOutlinerTreeItem>& Item, AActor*& ActorPtrOut) const;
 	};
