Non-unity + Linux build fixes

diff --git a/Engine/Plugins/Experimental/ChaosVD/Source/ChaosVD/Public/ChaosVDGeometryBuilder.h b/Engine/Plugins/Experimental/ChaosVD/Source/ChaosVD/Public/ChaosVDGeometryBuilder.h
--- a/Engine/Plugins/Experimental/ChaosVD/Source/ChaosVD/Public/ChaosVDGeometryBuilder.h
+++ b/Engine/Plugins/Experimental/ChaosVD/Source/ChaosVD/Public/ChaosVDGeometryBuilder.h
@@ -49,7 +49,7 @@ public:
 	{
 		FTSTicker::GetCoreTicker().RemoveTicker(GameThreadTickDelegate);
 
-		for (const TPair<uint32, UStaticMesh*>& StaticMeshByKey : StaticMeshCacheMap)
+		for (const TPair<uint32, TObjectPtr<UStaticMesh>>& StaticMeshByKey : StaticMeshCacheMap)
 		{
 			if (StaticMeshByKey.Value)
 			{
@@ -58,7 +58,7 @@ public:
 			}
 		}
 
-		for (const TPair<uint32, UDynamicMesh*>& DynamicMeshByKey : DynamicMeshCacheMap)
+		for (const TPair<uint32, TObjectPtr<UDynamicMesh>>& DynamicMeshByKey : DynamicMeshCacheMap)
 		{
 			if (DynamicMeshByKey.Value)
 			{
diff --git a/Engine/Source/Runtime/Experimental/Iris/Core/Private/Iris/ReplicationSystem/NetRefHandleManager.h b/Engine/Source/Runtime/Experimental/Iris/Core/Private/Iris/ReplicationSystem/NetRefHandleManager.h
--- a/Engine/Source/Runtime/Experimental/Iris/Core/Private/Iris/ReplicationSystem/NetRefHandleManager.h
+++ b/Engine/Source/Runtime/Experimental/Iris/Core/Private/Iris/ReplicationSystem/NetRefHandleManager.h
@@ -11,6 +11,7 @@
 #include "Iris/ReplicationState/ReplicationStateDescriptor.h"
 #include "Iris/ReplicationSystem/NetRefHandle.h"
 #include "Iris/ReplicationSystem/NetDependencyData.h"
+#include "UObject/ObjectPtr.h"
 
 class FReferenceCollector;
 namespace UE::Net
