[Chaos VD] Fixed for crashes introduced when the new local selection system was added to the tool

- Fix for a crash the changing worlds in the editor while the CVD tool is open
- Fix for a crash when a selected object is destroyed
- Fix for a crash where we are trying to select a garbage object

diff --git a/Engine/Plugins/Experimental/ChaosVD/Source/ChaosVD/Private/ChaosVDScene.cpp b/Engine/Plugins/Experimental/ChaosVD/Source/ChaosVD/Private/ChaosVDScene.cpp
--- a/Engine/Plugins/Experimental/ChaosVD/Source/ChaosVD/Private/ChaosVDScene.cpp
+++ b/Engine/Plugins/Experimental/ChaosVD/Source/ChaosVD/Private/ChaosVDScene.cpp
@@ -133,6 +133,11 @@ void FChaosVDScene::UpdateFromRecordedStepData(const int32 SolverID, const FStri
 		{
 			if (AChaosVDParticleActor* ActorToRemove = RemoveIterator.Value())
 			{
+				if (IsObjectSelected(ActorToRemove))
+				{
+					ClearSelectionAndNotify();
+				}
+
 				PhysicsVDWorld->DestroyActor(ActorToRemove);
 			}
 
@@ -177,6 +182,11 @@ void FChaosVDScene::HandleEnterNewGameFrame(int32 FrameNumber, const TArray<int3
 		{
 			for (const TPair<int32, AChaosVDParticleActor*>& ParticleVDInstanceWithID : RemoveIterator.Value())
 			{
+				if (IsObjectSelected(ParticleVDInstanceWithID.Value))
+				{
+					ClearSelectionAndNotify();
+				}
+
 				PhysicsVDWorld->DestroyActor(ParticleVDInstanceWithID.Value);
 			}
 
@@ -194,6 +204,8 @@ void FChaosVDScene::HandleEnterNewGameFrame(int32 FrameNumber, const TArray<int3
 
 void FChaosVDScene::CleanUpScene()
 {
+	ClearSelectionAndNotify();
+
 	if (PhysicsVDWorld)
 	{
 		for (const TPair<int32, FChaosVDParticlesByIDMap>& SolverParticleVDInstanceWithID : ParticlesBySolverID)
@@ -264,7 +276,7 @@ UWorld* FChaosVDScene::CreatePhysicsVDWorld() const
 	const FName UniqueWorldName = FName(FGuid::NewGuid().ToString());
 	UWorld* NewWorld = NewWorld = NewObject<UWorld>( GetTransientPackage(), UniqueWorldName );
 	
-	NewWorld->WorldType = EWorldType::Editor;
+	NewWorld->WorldType = EWorldType::EditorPreview;
 
 	FWorldContext& WorldContext = GEngine->CreateNewWorldContext( NewWorld->WorldType );
 	WorldContext.SetCurrentWorld(NewWorld);
@@ -347,8 +359,29 @@ void FChaosVDScene::HandlePostSelectionChange(const UTypedElementSelectionSet* P
 	PendingActorsToUpdateSelectionProxy.Reset();
 }
 
+void FChaosVDScene::ClearSelectionAndNotify()
+{
+	if (!SelectionSet)
+	{
+		return;
+	}
+
+	SelectionSet->ClearSelection(FTypedElementSelectionOptions());
+	SelectionSet->NotifyPendingChanges();
+}
+
 void FChaosVDScene::SetSelectedObject(UObject* SelectedObject)
 {
+	if (!SelectionSet)
+	{
+		return;
+	}
+
+	if (!::IsValid(SelectedObject))
+	{
+		ClearSelectionAndNotify();
+		return;
+	}
 
 	if (IsObjectSelected(SelectedObject))
 	{
@@ -366,5 +399,15 @@ void FChaosVDScene::SetSelectedObject(UObject* SelectedObject)
 
 bool FChaosVDScene::IsObjectSelected(const UObject* Object)
 {
+	if (!SelectionSet)
+	{
+		return false;
+	}
+
+	if (!::IsValid(Object))
+	{
+		return false;
+	}
+
 	return SelectionSet->IsElementSelected(GetSelectionHandleForObject(Object), FTypedElementIsSelectedOptions());;
 }
diff --git a/Engine/Plugins/Experimental/ChaosVD/Source/ChaosVD/Private/ChaosVDWorldOutlinerMode.cpp b/Engine/Plugins/Experimental/ChaosVD/Source/ChaosVD/Private/ChaosVDWorldOutlinerMode.cpp
--- a/Engine/Plugins/Experimental/ChaosVD/Source/ChaosVD/Private/ChaosVDWorldOutlinerMode.cpp
+++ b/Engine/Plugins/Experimental/ChaosVD/Source/ChaosVD/Private/ChaosVDWorldOutlinerMode.cpp
@@ -42,7 +42,10 @@ void FChaosVDWorldOutlinerMode::OnItemSelectionChanged(FSceneOutlinerTreeItemPtr
 		TArray<AActor*> SceneSelectedActors = SelectionSet->GetSelectedObjects<AActor>();
 		ensureMsgf(SceneSelectedActors.Num() < 2, TEXT("Multi Selection is not supported, but [%s] Actors are selected... Choosing the first one"));
 
-		ScenePtr->SetSelectedObject(OutlinerSelectedActors[0]);
+		AActor* SelectedActor = OutlinerSelectedActors.IsEmpty() ? nullptr : OutlinerSelectedActors[0];
+
+		// If the actor ptr is null, the selection will be cleared
+		ScenePtr->SetSelectedObject(SelectedActor);
 	}
 }
 
diff --git a/Engine/Plugins/Experimental/ChaosVD/Source/ChaosVD/Public/ChaosVDScene.h b/Engine/Plugins/Experimental/ChaosVD/Source/ChaosVD/Public/ChaosVDScene.h
--- a/Engine/Plugins/Experimental/ChaosVD/Source/ChaosVD/Public/ChaosVDScene.h
+++ b/Engine/Plugins/Experimental/ChaosVD/Source/ChaosVD/Public/ChaosVDScene.h
@@ -99,6 +99,8 @@ private:
 	void HandlePreSelectionChange(const UTypedElementSelectionSet* PreChangeSelectionSet);
 	void HandlePostSelectionChange(const UTypedElementSelectionSet* PreChangeSelectionSet);
 
+	void ClearSelectionAndNotify();
+
 	/** UWorld instance used to represent the recorded debug data */
 	TObjectPtr<UWorld> PhysicsVDWorld = nullptr;
 
