[ChaosVD] Fix CIS issue due to a case error on an include file in IChaosVDDataProcessor.h

diff --git a/Engine/Plugins/Experimental/ChaosVD/Source/ChaosVD/Public/Trace/DataProcessors/IChaosVDDataProcessor.h b/Engine/Plugins/Experimental/ChaosVD/Source/ChaosVD/Public/Trace/DataProcessors/IChaosVDDataProcessor.h
--- a/Engine/Plugins/Experimental/ChaosVD/Source/ChaosVD/Public/Trace/DataProcessors/IChaosVDDataProcessor.h
+++ b/Engine/Plugins/Experimental/ChaosVD/Source/ChaosVD/Public/Trace/DataProcessors/IChaosVDDataProcessor.h
@@ -3,7 +3,7 @@
 #pragma once
 #include "Containers/Array.h"
 #include "Containers/StringFwd.h"
-#include "Hal/Platform.h"
+#include "HAL/Platform.h"
 #include "Templates/SharedPointer.h"
 
 class FChaosVDTraceProvider;
