[Chaos VD] Implemented Particle Data and Collision Data visualization
- Now the tool "debug draws" points, vectors and text representing collisions (contacts, push out forces, normals, etc) and particle data (velocity, acceleration, etc) when a particle is selected.
- What data to visualize is controlled with enum flags you can set in the tool itself. These are available as global configuration flags (it will affect all the particles that are selected), and local flags. If a particle has local flags set, these will override the global ones.
- Added a way to lock the current step, so when you scrub to a different recorded frame, the tool will automatically go to the same step number (playing back all the previous steps). This is useful to compare differences between frames on a specific sub-step.
- Made the last recorded substep of any given frame, the default sub step that will be played back when scrubbing to a new frame. This better represents the end state of the simulation on any given frame. You can use the existing next/prev steps to go backwards if needed.

Something worth noticing is that collision data visualization is still not 100% correct, I will address that in a new CL

diff --git a/Engine/Plugins/Experimental/ChaosVD/Source/ChaosVD/Private/ChaosVDEditorSettingsTab.cpp b/Engine/Plugins/Experimental/ChaosVD/Source/ChaosVD/Private/ChaosVDEditorSettingsTab.cpp
new file mode 100644
--- /dev/null
+++ b/Engine/Plugins/Experimental/ChaosVD/Source/ChaosVD/Private/ChaosVDEditorSettingsTab.cpp
@@ -0,0 +1,44 @@
+﻿// Copyright Epic Games, Inc. All Rights Reserved.
+
+#include "ChaosVDEditorSettingsTab.h"
+
+#include "ChaosVDEditorSettings.h"
+#include "ChaosVDStyle.h"
+#include "DetailsViewArgs.h"
+#include "Framework/Docking/TabManager.h"
+#include "IDetailsView.h"
+#include "Modules/ModuleManager.h"
+#include "PropertyEditorModule.h"
+#include "Templates/SharedPointer.h"
+#include "Widgets/Docking/SDockTab.h"
+
+#define LOCTEXT_NAMESPACE "ChaosVisualDebugger"
+
+TSharedRef<SDockTab> FChaosVDEditorSettingsTab::HandleTabSpawned(const FSpawnTabArgs& Args)
+{
+	FPropertyEditorModule& PropertyEditorModule = FModuleManager::GetModuleChecked<FPropertyEditorModule>("PropertyEditor");
+	FDetailsViewArgs DetailsViewArgs;
+	DetailsViewArgs.bAllowSearch = false;
+	DetailsViewArgs.NameAreaSettings = FDetailsViewArgs::HideNameArea;
+
+	TSharedRef<IDetailsView> DetailsPanel = PropertyEditorModule.CreateDetailView(DetailsViewArgs);
+
+	DetailsPanel->SetObject(GetMutableDefault<UChaosVDEditorSettings>());
+
+	TSharedRef<SDockTab> DetailsPanelTab =
+		SNew(SDockTab)
+		.TabRole(ETabRole::MajorTab)
+		.Label(LOCTEXT("ChaosVDEditorSettings", "Chaos VD Settings"))
+		.ToolTipText(LOCTEXT("ChaosVDEditorSettingsTip", "See the available settings for the editor"));
+
+	DetailsPanelTab->SetContent
+	(
+		DetailsPanel
+	);
+
+	DetailsPanelTab->SetTabIcon(FChaosVDStyle::Get().GetBrush("TabIconDetailsPanel"));
+
+	return DetailsPanelTab;
+}
+
+#undef LOCTEXT_NAMESPACE 
diff --git a/Engine/Plugins/Experimental/ChaosVD/Source/ChaosVD/Private/ChaosVDParticleActor.cpp b/Engine/Plugins/Experimental/ChaosVD/Source/ChaosVD/Private/ChaosVDParticleActor.cpp
--- a/Engine/Plugins/Experimental/ChaosVD/Source/ChaosVD/Private/ChaosVDParticleActor.cpp
+++ b/Engine/Plugins/Experimental/ChaosVD/Source/ChaosVD/Private/ChaosVDParticleActor.cpp
@@ -4,12 +4,13 @@
 
 #include "ChaosVDGeometryBuilder.h"
 #include "ChaosVDScene.h"
-#include "DataWrappers/ChaosVDParticleDataWrapper.h"
 #include "Components/MeshComponent.h"
 #include "Components/SceneComponent.h"
 #include "Components/StaticMeshComponent.h"
+#include "DataWrappers/ChaosVDParticleDataWrapper.h"
 #include "Engine/CollisionProfile.h"
 #include "Engine/StaticMesh.h"
+#include "Visualizers/ChaosVDParticleDataVisualizer.h"
 
 static FAutoConsoleVariable CVarChaosVDHideVolumeAndBrushesHack(
 	TEXT("p.Chaos.VD.Tool.HideVolumeAndBrushesHack"),
@@ -19,10 +20,15 @@ static FAutoConsoleVariable CVarChaosVDHideVolumeAndBrushesHack(
 AChaosVDParticleActor::AChaosVDParticleActor(const FObjectInitializer& ObjectInitializer) : Super(ObjectInitializer)
 {
 	RootComponent = CreateDefaultSubobject<USceneComponent>(TEXT("SceneComponent0"));
+	CreateVisualizers();
 }
 
 void AChaosVDParticleActor::UpdateFromRecordedParticleData(const FChaosVDParticleDataWrapper& InRecordedData, const Chaos::FRigidTransform3& SimulationTransform)
 {
+	//TODO: Make the simulation transform be cached on the CVD Scene, so we can query from it when needed
+	// Copying it to each particle actor is not efficient
+	CachedSimulationTransform = SimulationTransform;
+
 	if (InRecordedData.ParticlePositionRotation.HasValidData())
 	{
 		// TODO: Only update if the transform (either the simulation or the particle) has changed;
@@ -166,7 +172,21 @@ void AChaosVDParticleActor::BeginDestroy()
 	Super::BeginDestroy();
 }
 
+void AChaosVDParticleActor::GetVisualizationContext(FChaosVDVisualizationContext& OutVisualizationContext)
+{
+	OutVisualizationContext.SpaceTransform = CachedSimulationTransform;
+	OutVisualizationContext.CVDScene = OwningScene;
+	OutVisualizationContext.SolverID = ParticleDataViewer.SolverID;
+}
+
+void AChaosVDParticleActor::CreateVisualizers()
+{
+	CVDVisualizers.Add(FChaosVDParticleDataVisualizer::VisualizerID,MakeUnique<FChaosVDParticleDataVisualizer>(*this));
+	CVDVisualizers.Add(FChaosVDCollisionDataVisualizer::VisualizerID, MakeUnique<FChaosVDCollisionDataVisualizer>(*this));
+}
+
 #if WITH_EDITOR
+
 bool AChaosVDParticleActor::IsSelectedInEditor() const
 {
 	// The implementation of this method in UObject, used a global edit callback,
@@ -192,5 +212,33 @@ void AChaosVDParticleActor::PostEditChangeProperty(FPropertyChangedEvent& Proper
 			UpdateGeometry(ParticleDataViewer.GeometryHash, EChaosVDActorGeometryUpdateFlags::ForceUpdate);
 		}
 	}
+
+	if (PropertyChangedEvent.GetMemberPropertyName() == GET_MEMBER_NAME_STRING_CHECKED(AChaosVDParticleActor, LocalCollisionDataVisualizationFlags))
+	{
+		if (TUniquePtr<FChaosVDDataVisualizerBase>* CollisionVisualizer = CVDVisualizers.Find(FChaosVDCollisionDataVisualizer::VisualizerID))
+		{
+			CollisionVisualizer->Get()->UpdateVisualizationFlags(LocalCollisionDataVisualizationFlags);
+		}
+	}
+
+	if (PropertyChangedEvent.GetMemberPropertyName() == GET_MEMBER_NAME_STRING_CHECKED(AChaosVDParticleActor, LocalParticleDataVisualizationFlags))
+	{
+		if (TUniquePtr<FChaosVDDataVisualizerBase>* ParticleDataVisualizer = CVDVisualizers.Find(FChaosVDParticleDataVisualizer::VisualizerID))
+		{
+			ParticleDataVisualizer->Get()->UpdateVisualizationFlags(LocalParticleDataVisualizationFlags);
+		}
+	}
+}
+
+void AChaosVDParticleActor::DrawVisualization(const FSceneView* View, FPrimitiveDrawInterface* PDI)
+{
+	for (const TPair<FStringView, TUniquePtr<FChaosVDDataVisualizerBase>>& VisualizerWithID : CVDVisualizers)
+	{
+		if (VisualizerWithID.Value)
+		{
+			VisualizerWithID.Value->DrawVisualization(View, PDI);
+		}
+	}
 }
+
 #endif //WITH_EDITOR
diff --git a/Engine/Plugins/Experimental/ChaosVD/Source/ChaosVD/Private/ChaosVDParticleActorCustomization.cpp b/Engine/Plugins/Experimental/ChaosVD/Source/ChaosVD/Private/ChaosVDParticleActorCustomization.cpp
--- a/Engine/Plugins/Experimental/ChaosVD/Source/ChaosVD/Private/ChaosVDParticleActorCustomization.cpp
+++ b/Engine/Plugins/Experimental/ChaosVD/Source/ChaosVD/Private/ChaosVDParticleActorCustomization.cpp
@@ -5,6 +5,11 @@
 #include "DetailCategoryBuilder.h"
 #include "DetailLayoutBuilder.h"
 
+FChaosVDParticleActorCustomization::FChaosVDParticleActorCustomization()
+{
+	AllowedCategories.Add(FChaosVDParticleActorCustomization::ChaosVDCategoryName);
+	AllowedCategories.Add(FChaosVDParticleActorCustomization::ChaosVDVisualizationCategoryName);
+}
 
 TSharedRef<IDetailCustomization> FChaosVDParticleActorCustomization::MakeInstance()
 {
@@ -13,12 +18,14 @@ TSharedRef<IDetailCustomization> FChaosVDParticleActorCustomization::MakeInstanc
 
 void FChaosVDParticleActorCustomization::CustomizeDetails(IDetailLayoutBuilder& DetailBuilder)
 {
+
+	IDetailCategoryBuilder& YourCategory = DetailBuilder.EditCategory(ChaosVDVisualizationCategoryName, FText::GetEmpty(), ECategoryPriority::Important);
 	// Hide everything as the only thing we want to show in these actors is the Recorded debug data
 	TArray<FName> CurrentCategoryNames;
 	DetailBuilder.GetCategoryNames(CurrentCategoryNames);
 	for (const FName& CategoryToHide : CurrentCategoryNames)
 	{
-		if (CategoryToHide != ChaosVDCategoryName)
+		if (!AllowedCategories.Contains(CategoryToHide))
 		{
 			DetailBuilder.HideCategory(CategoryToHide);
 		}
diff --git a/Engine/Plugins/Experimental/ChaosVD/Source/ChaosVD/Private/ChaosVDPlaybackController.cpp b/Engine/Plugins/Experimental/ChaosVD/Source/ChaosVD/Private/ChaosVDPlaybackController.cpp
--- a/Engine/Plugins/Experimental/ChaosVD/Source/ChaosVD/Private/ChaosVDPlaybackController.cpp
+++ b/Engine/Plugins/Experimental/ChaosVD/Source/ChaosVD/Private/ChaosVDPlaybackController.cpp
@@ -117,21 +117,55 @@ void FChaosVDPlaybackController::GoToRecordedSolverStep(const int32 InTrackID, c
 			TraceServices::FAnalysisSessionReadScope SessionReadScope(*TraceSession);
 
 			const FChaosVDSolverFrameData* SolverFrameData = LoadedRecording->GetSolverFrameData(InTrackID, FrameNumber);
-			if (SolverFrameData && SolverFrameData->SolverSteps.IsValidIndex(Step))
-			{
-				SceneToControlSharedPtr->UpdateFromRecordedStepData(InTrackID, SolverFrameData->DebugName, SolverFrameData->SolverSteps[Step], *SolverFrameData);
-			}
 
+			TSharedPtr<FChaosVDTrackInfo> CurrentTrackInfo;
 			if (TrackInfoByIDMap* TrackInfoByID = TrackInfoPerType.Find(EChaosVDTrackType::Solver))
 			{
 				if (TSharedPtr<FChaosVDTrackInfo>* TrackInfo = TrackInfoByID->Find(InTrackID))
 				{
-					TrackInfo->Get()->CurrentFrame = FrameNumber;
-					TrackInfo->Get()->CurrentStep = Step;
+					CurrentTrackInfo = *TrackInfo;
+					if (!ensure(CurrentTrackInfo.IsValid()))
+					{
+						UE_LOG(LogChaosVDEditor, Warning, TEXT("[%s] Track info for track ID [%d]. We can't continue..."), ANSI_TO_TCHAR(__FUNCTION__), InTrackID);
+						return;
+					}
+				}
+			}
 
-					OnTrackFrameUpdated().Broadcast(AsWeak(), TrackInfo->Get(), InstigatorID);
+			if (CurrentTrackInfo->LockedOnStep != INDEX_NONE)
+			{
+				// If this track is locked to a specific step, we need to play back the previous steps on the current frame, because not all steps capture the same data.
+				// For example, Particles positions are fully captured in the first sub-step and the last one
+				for (int32 StepIndex = 0; StepIndex <= CurrentTrackInfo->LockedOnStep; StepIndex++)
+				{
+					if (SolverFrameData && ensure(SolverFrameData->SolverSteps.IsValidIndex(CurrentTrackInfo->LockedOnStep)))
+					{
+						SceneToControlSharedPtr->UpdateFromRecordedStepData(InTrackID, SolverFrameData->DebugName, SolverFrameData->SolverSteps[StepIndex], *SolverFrameData);
+					}
+					else
+					{
+						UE_LOG(LogChaosVDEditor, Error, TEXT("[%s] Tried to scrub to an invalid step | Step Number [%d] | Available Steps [%d]..."), ANSI_TO_TCHAR(__FUNCTION__), CurrentTrackInfo->LockedOnStep, SolverFrameData->SolverSteps.Num());
+						return;
+					}
+				}
+			}
+			else
+			{
+				if (SolverFrameData && ensure(SolverFrameData->SolverSteps.IsValidIndex(Step)))
+				{
+					SceneToControlSharedPtr->UpdateFromRecordedStepData(InTrackID, SolverFrameData->DebugName, SolverFrameData->SolverSteps[Step], *SolverFrameData);
+				}
+				else
+				{
+					UE_LOG(LogChaosVDEditor, Error, TEXT("[%s] Tried to scrub to an invalid step | Step Number [%d] | Available Steps [%d]..."), ANSI_TO_TCHAR(__FUNCTION__), CurrentTrackInfo->LockedOnStep, SolverFrameData->SolverSteps.Num());
 				}
 			}
+			
+
+			CurrentTrackInfo->CurrentFrame = FrameNumber;
+			CurrentTrackInfo->CurrentStep = Step;
+
+			OnTrackFrameUpdated().Broadcast(AsWeak(), CurrentTrackInfo.Get(), InstigatorID);
 		}
 	}
 	else
@@ -157,11 +191,12 @@ void FChaosVDPlaybackController::GoToRecordedGameFrame(const int32 FrameNumber,
 
 					for (int32 SolverID : AvailableSolversID)
 					{
-						constexpr int32 StepNumber = 0;
-
 						// When Scrubbing the timeline by game frames instead of solvers, try to go to the first solver frame on the first platform cycle of the game frame.
 						// Game Frames are not in sync with Solver Frames and Solver steps.
-						int32 SolverFrameNumber = LoadedRecording->GetLowestSolverFrameNumberAtCycle(SolverID, FrameData->FirstCycle);
+						const int32 SolverFrameNumber = LoadedRecording->GetLowestSolverFrameNumberAtCycle(SolverID, FrameData->FirstCycle);
+	
+						const int32 StepNumber = GetTrackLastStepAtFrame(EChaosVDTrackType::Solver, SolverID,SolverFrameNumber);
+
 						GoToTrackFrame(InstigatorID, EChaosVDTrackType::Solver, SolverID, SolverFrameNumber, StepNumber);
 					}
 
@@ -358,7 +393,31 @@ int32 FChaosVDPlaybackController::GetTrackCurrentStep(EChaosVDTrackType TrackTyp
 	return INDEX_NONE;
 }
 
+int32 FChaosVDPlaybackController::GetTrackLastStepAtFrame(EChaosVDTrackType TrackType, int32 InTrackID, int32 InFrameNumber) const
+{
+	switch (TrackType)
+	{
+		case EChaosVDTrackType::Solver:
+		{
+			const int32 AvailableSteps = GetTrackStepsNumberAtFrame(EChaosVDTrackType::Solver, InTrackID, InFrameNumber);
+			return AvailableSteps == INDEX_NONE ? INDEX_NONE: AvailableSteps -1;
+			break;
+		}
+		case EChaosVDTrackType::Game:
+		default:
+			ensureMsgf(false, TEXT("Unsuported Track Type"));
+			return INDEX_NONE;
+			break;
+	}
+	
+}
+
 const FChaosVDTrackInfo* FChaosVDPlaybackController::GetTrackInfo(EChaosVDTrackType TrackType, int32 TrackID)
+{
+	return GetMutableTrackInfo(TrackType, TrackID);
+}
+
+FChaosVDTrackInfo* FChaosVDPlaybackController::GetMutableTrackInfo(EChaosVDTrackType TrackType, int32 TrackID)
 {
 	if (const TrackInfoByIDMap* TrackInfoByID = TrackInfoPerType.Find(TrackType))
 	{
@@ -372,6 +431,22 @@ const FChaosVDTrackInfo* FChaosVDPlaybackController::GetTrackInfo(EChaosVDTrackT
 	return nullptr;
 }
 
+void FChaosVDPlaybackController::LockTrackInCurrentStep(EChaosVDTrackType TrackType, int32 TrackID)
+{
+	if (FChaosVDTrackInfo* TrackInfo = GetMutableTrackInfo(TrackType, TrackID))
+	{
+		TrackInfo->LockedOnStep = TrackInfo->CurrentStep;
+	}
+}
+
+void FChaosVDPlaybackController::UnlockTrackStep(EChaosVDTrackType TrackType, int32 TrackID)
+{
+	if (FChaosVDTrackInfo* TrackInfo = GetMutableTrackInfo(TrackType, TrackID))
+	{
+		TrackInfo->LockedOnStep = INDEX_NONE;
+	}
+}
+
 void FChaosVDPlaybackController::GetAvailableTracks(EChaosVDTrackType TrackType, TArray<TSharedPtr<FChaosVDTrackInfo>>& OutTrackInfo)
 {
 	OutTrackInfo.Reset();
diff --git a/Engine/Plugins/Experimental/ChaosVD/Source/ChaosVD/Private/ChaosVDPlaybackViewportClient.cpp b/Engine/Plugins/Experimental/ChaosVD/Source/ChaosVD/Private/ChaosVDPlaybackViewportClient.cpp
--- a/Engine/Plugins/Experimental/ChaosVD/Source/ChaosVD/Private/ChaosVDPlaybackViewportClient.cpp
+++ b/Engine/Plugins/Experimental/ChaosVD/Source/ChaosVD/Private/ChaosVDPlaybackViewportClient.cpp
@@ -2,8 +2,12 @@
 
 #include "ChaosVDPlaybackViewportClient.h"
 
+#include "ChaosVDParticleActor.h"
 #include "ChaosVDScene.h"
 #include "EngineUtils.h"
+#include "Elements/Framework/TypedElementSelectionSet.h"
+#include "Visualizers/ChaosVDDebugDrawUtils.h"
+
 
 FChaosVDPlaybackViewportClient::FChaosVDPlaybackViewportClient() : FEditorViewportClient(nullptr), CVDWorld(nullptr)
 {
@@ -58,3 +62,28 @@ void FChaosVDPlaybackViewportClient::HandleObjectFocused(UObject* FocusedObject)
 		FocusViewportOnBox(FocusedActor->GetComponentsBoundingBox(false));
 	}
 }
+
+void FChaosVDPlaybackViewportClient::Draw(const FSceneView* View, FPrimitiveDrawInterface* PDI)
+{
+	FEditorViewportClient::Draw(View, PDI);
+
+	if (TSharedPtr<FChaosVDScene> ScenePtr = CVDScene.Pin())
+	{	
+		TArray<AActor*> SelectedActors = ScenePtr->GetElementSelectionSet()->GetSelectedObjects<AActor>();
+
+		for (AActor* SelectedActor : SelectedActors)
+		{
+			if (IChaosVDVisualizerContainerInterface* VisualizerContainer = Cast<IChaosVDVisualizerContainerInterface>(SelectedActor))
+			{
+				VisualizerContainer->DrawVisualization(View, PDI);
+			}
+		}
+	}
+}
+
+void FChaosVDPlaybackViewportClient::DrawCanvas(FViewport& InViewport, FSceneView& View, FCanvas& Canvas)
+{
+	FEditorViewportClient::DrawCanvas(InViewport, View, Canvas);
+
+	FChaosVDDebugDrawUtils::DrawCanvas(InViewport, View, Canvas);
+}
diff --git a/Engine/Plugins/Experimental/ChaosVD/Source/ChaosVD/Private/ChaosVDScene.cpp b/Engine/Plugins/Experimental/ChaosVD/Source/ChaosVD/Private/ChaosVDScene.cpp
--- a/Engine/Plugins/Experimental/ChaosVD/Source/ChaosVD/Private/ChaosVDScene.cpp
+++ b/Engine/Plugins/Experimental/ChaosVD/Source/ChaosVD/Private/ChaosVDScene.cpp
@@ -109,7 +109,7 @@ void FChaosVDScene::UpdateFromRecordedStepData(const int32 SolverID, const FStri
 			if (AChaosVDParticleActor* NewParticleVDInstance = SpawnParticleFromRecordedData(Particle, InFrameData))
 			{
 				FStringFormatOrderedArguments Args {SolverName, FString::FromInt(SolverID)};
-				const FName FolderPath = *FPaths::Combine(FString::Format(TEXT("Solver {0} | ID {1}"), Args), UEnum::GetDisplayValueAsText(NewParticleVDInstance->GetParticleData().Type).ToString());
+				const FName FolderPath = *FPaths::Combine(FString::Format(TEXT("Solver {0} | ID {1}"), Args), UEnum::GetDisplayValueAsText(NewParticleVDInstance->GetParticleData()->Type).ToString());
 
 				NewParticleVDInstance->SetFolderPath(FolderPath);
 
@@ -264,6 +264,12 @@ const TSharedPtr<const Chaos::FImplicitObject>* FChaosVDScene::GetUpdatedGeometr
 	return nullptr;
 }
 
+AChaosVDParticleActor* FChaosVDScene::GetParticleActor(int32 SolverID, int32 ParticleID)
+{
+	AChaosVDParticleActor** ParticleActorPtrPtr = ParticlesBySolverID[SolverID].Find(ParticleID);
+	return ParticleActorPtrPtr ? *ParticleActorPtrPtr : nullptr;
+}
+
 AChaosVDParticleActor* FChaosVDScene::SpawnParticleFromRecordedData(const FChaosVDParticleDataWrapper& InParticleData, const FChaosVDSolverFrameData& InFrameData)
 {
 	using namespace Chaos;
diff --git a/Engine/Plugins/Experimental/ChaosVD/Source/ChaosVD/Private/ChaosVDStyle.cpp b/Engine/Plugins/Experimental/ChaosVD/Source/ChaosVD/Private/ChaosVDStyle.cpp
--- a/Engine/Plugins/Experimental/ChaosVD/Source/ChaosVD/Private/ChaosVDStyle.cpp
+++ b/Engine/Plugins/Experimental/ChaosVD/Source/ChaosVD/Private/ChaosVDStyle.cpp
@@ -54,6 +54,9 @@ TSharedRef< FSlateStyleSet > FChaosVDStyle::Create()
 	Style->Set("TabIconDetailsPanel",  new IMAGE_BRUSH_SVG(TEXT("Details"), Icon16x16));
 	Style->Set("TabIconPlaybackViewport",  new IMAGE_BRUSH(TEXT("icon_tab_PlaybackViewport_16x"), Icon16x16));
 	Style->Set("TabIconWorldOutliner", new IMAGE_BRUSH_SVG(TEXT("WorldOutliner"), Icon16x16));
+	
+	Style->Set("LockIcon", new IMAGE_BRUSH_SVG(TEXT("lock"), Icon16x16));
+	Style->Set("UnlockedIcon", new IMAGE_BRUSH_SVG(TEXT("lock-unlocked"), Icon16x16));
 
 	return Style;
 }
diff --git a/Engine/Plugins/Experimental/ChaosVD/Source/ChaosVD/Private/ChaosVDTabsIDs.cpp b/Engine/Plugins/Experimental/ChaosVD/Source/ChaosVD/Private/ChaosVDTabsIDs.cpp
--- a/Engine/Plugins/Experimental/ChaosVD/Source/ChaosVD/Private/ChaosVDTabsIDs.cpp
+++ b/Engine/Plugins/Experimental/ChaosVD/Source/ChaosVD/Private/ChaosVDTabsIDs.cpp
@@ -8,4 +8,5 @@ const FName FChaosVDTabID::PlaybackViewport = TEXT("PlaybackViewport");
 const FName FChaosVDTabID::WorldOutliner = TEXT("WorldOutliner");
 const FName FChaosVDTabID::DetailsPanel = TEXT("DetailsPanel");
 const FName FChaosVDTabID::OutputLog = TEXT("ChaosVDOutputLog");
-const FName FChaosVDTabID::SolversTrack = TEXT("SolversTrack");
\ No newline at end of file
+const FName FChaosVDTabID::SolversTrack = TEXT("SolversTrack");
+const FName FChaosVDTabID::CVDEditorSettings = TEXT("CVDEditorSettings");
\ No newline at end of file
diff --git a/Engine/Plugins/Experimental/ChaosVD/Source/ChaosVD/Private/Visualizers/ChaosVDCollisionDataVisualizer.cpp b/Engine/Plugins/Experimental/ChaosVD/Source/ChaosVD/Private/Visualizers/ChaosVDCollisionDataVisualizer.cpp
new file mode 100644
--- /dev/null
+++ b/Engine/Plugins/Experimental/ChaosVD/Source/ChaosVD/Private/Visualizers/ChaosVDCollisionDataVisualizer.cpp
@@ -0,0 +1,180 @@
+﻿// Copyright Epic Games, Inc. All Rights Reserved.
+
+#include "Visualizers/ChaosVDCollisionDataVisualizer.h"
+
+#include "ChaosVDEditorSettings.h"
+#include "ChaosVDModule.h"
+#include "ChaosVDParticleActor.h"
+#include "ChaosVDScene.h"
+#include "DataWrappers/ChaosVDCollisionDataWrappers.h"
+#include "Visualizers/ChaosVDDebugDrawUtils.h"
+
+void FChaosVDCollisionDataVisualizer::DrawVisualization(const FSceneView* View, FPrimitiveDrawInterface* PDI)
+{
+	uint32 VisualizationFlagsToUse = 0;
+	if (static_cast<EChaosVDCollisionVisualizationFlags>(LocalVisualizationFlags) == EChaosVDCollisionVisualizationFlags::None)
+	{
+		if (const UChaosVDEditorSettings* EditorSettings = GetDefault<UChaosVDEditorSettings>())
+		{
+			VisualizationFlagsToUse = EditorSettings->GlobalCollisionDataVisualizationFlags;
+		}
+		else
+		{
+			UE_LOG(LogChaosVDEditor, Warning, TEXT("[%s] Failed to retrive global visualization setting. Falling back to local settings"), ANSI_TO_TCHAR(__FUNCTION__));
+
+			VisualizationFlagsToUse = LocalVisualizationFlags;
+		}	
+	}
+	else
+	{
+		VisualizationFlagsToUse = LocalVisualizationFlags;
+	}
+
+	FChaosVDVisualizationContext VisualizationContext;
+	DataProvider.GetVisualizationContext(VisualizationContext);
+
+	const FChaosVDParticleDataWrapper* ParticleDataViewer = DataProvider.GetParticleData();
+	if (!ensure(ParticleDataViewer))
+	{
+		return;
+	}
+
+	TSharedPtr<FChaosVDScene> ScenePtr = VisualizationContext.CVDScene.Pin();
+	if (!ScenePtr.IsValid())
+	{
+		return;
+	}
+ 
+	const FVector& OwnerActorLocation = ParticleDataViewer->ParticlePositionRotation.MX;
+
+	for (const FChaosVDParticlePairMidPhase& MidPhase : ParticleDataViewer->ParticleMidPhases)
+	{
+		for (const FChaosVDConstraint& Constraint : MidPhase.Constraints)
+		{
+			for (const FChaosVDManifoldPoint& ManifoldPoint : Constraint.ManifoldPoints)
+			{
+				constexpr int32 MinShapeContactAmount = 2;
+				if (Constraint.ShapeWorldTransforms.Num() < MinShapeContactAmount || ManifoldPoint.ContactPoint.ShapeContactPoints.Num() < MinShapeContactAmount)
+				{
+					continue;
+				}
+
+				const bool bIsProbe = Constraint.bIsProbe;
+				const bool bIsActive = ManifoldPoint.bIsValid && (!ManifoldPoint.NetPushOut.IsNearlyZero() || !ManifoldPoint.NetImpulse.IsNearlyZero() || (!Constraint.bUseManifold && !Constraint.AccumulatedImpulse.IsNearlyZero()));
+				if (!bIsActive && EnumHasAnyFlags(static_cast<EChaosVDCollisionVisualizationFlags>(LocalVisualizationFlags), EChaosVDCollisionVisualizationFlags::DrawOnlyActiveContacts))
+				{
+					continue;
+				}
+
+				const bool bPruned = ManifoldPoint.bDisabled;
+
+				AChaosVDParticleActor* CVDParticleActor1 = ScenePtr->GetParticleActor(VisualizationContext.SolverID, Constraint.Particle1Index);
+				AChaosVDParticleActor* CVDParticleActor0 = ScenePtr->GetParticleActor(VisualizationContext.SolverID, Constraint.Particle0Index);
+				if (CVDParticleActor1 == nullptr || CVDParticleActor0 == nullptr)
+				{
+					continue;
+				}
+
+				const FTransform& WorldActorTransform1 = CVDParticleActor1->GetTransform();
+				const FTransform& WorldActorTransform0 = CVDParticleActor0->GetTransform();
+
+				constexpr int32 ContactPlaneOwner = 1;
+				constexpr int32 ContactPointOwner = 1 - ContactPlaneOwner;
+				const FTransform& PlaneTransform = Constraint.ImplicitTransforms[1] * WorldActorTransform1;
+				const FTransform& PointTransform = Constraint.ImplicitTransforms[0] * WorldActorTransform0;
+
+				const FVector PlaneNormal = PlaneTransform.TransformVectorNoScale(FVector(ManifoldPoint.ContactPoint.ShapeContactNormal));
+				const FVector PointLocation = PointTransform.TransformPosition(FVector(ManifoldPoint.ContactPoint.ShapeContactPoints[ContactPointOwner]));
+				const FVector PlaneLocation = PlaneTransform.TransformPosition(FVector(ManifoldPoint.ContactPoint.ShapeContactPoints[ContactPlaneOwner]));
+				const FVector PointPlaneLocation = PointLocation - FVector::DotProduct(PointLocation - PlaneLocation, PlaneNormal) * PlaneNormal;
+
+				// TODO: Extract these colors to a class or another place where ChaosDebugDraw and Chaos VD Visualizers can access
+				// The bulk of this implementation including the following colors come from ChaosDebugDraw Implementation.
+				// The implementation itself is likely to drift, but we should try to keep the colors the same between both debugging tools
+
+				// Dynamic friction, restitution = red
+				// Static friction, no restitution = green
+				// Inactive = gray
+				FColor DiscColor = FColor(200, 0, 0);
+				FColor PlaneNormalColor = FColor(200, 0, 0);
+				FColor EdgeNormalColor = FColor(200, 150, 0);
+				FColor ImpulseColor = FColor(0, 0, 200);
+				FColor PushOutImpusleColor = FColor(0, 200, 200);
+				if (ManifoldPoint.bInsideStaticFrictionCone)
+				{
+					DiscColor = FColor(150, 200, 0);
+				}
+				if (bIsProbe)
+				{
+					DiscColor = FColor(50, 180, 180);
+					PlaneNormalColor = FColor(50, 180, 180);
+					EdgeNormalColor = FColor(50, 180, 130);
+				}
+				else if (!bIsActive)
+				{
+					DiscColor = FColor(100, 100, 100);
+					PlaneNormalColor = FColor(100, 0, 0);
+					EdgeNormalColor = FColor(100, 80, 0);
+				}
+				if (bPruned)
+				{
+					PlaneNormalColor = FColor(200, 0, 200);
+					EdgeNormalColor = FColor(200, 0, 200);
+				}
+			
+				const FVector WorldPointLocation = VisualizationContext.SpaceTransform.TransformPosition(PointLocation);
+				const FVector WorldPlaneLocation = VisualizationContext.SpaceTransform.TransformPosition(PlaneLocation);
+				const FVector WorldPointPlaneLocation = VisualizationContext.SpaceTransform.TransformPosition(PointPlaneLocation);
+				const FVector WorldPlaneNormal = VisualizationContext.SpaceTransform.TransformVectorNoScale(PlaneNormal);
+
+				if (EnumHasAnyFlags(static_cast<EChaosVDCollisionVisualizationFlags>(VisualizationFlagsToUse), EChaosVDCollisionVisualizationFlags::NetPushOut))
+				{
+					if (ManifoldPoint.bIsValid && !ManifoldPoint.NetPushOut.IsNearlyZero())
+					{
+						FChaosVDDebugDrawUtils::DrawArrowVector(PDI, WorldPointPlaneLocation, OwnerActorLocation + WorldPointPlaneLocation + VisualizationContext.SpaceTransform.TransformPosition(ManifoldPoint.NetPushOut), TEXT("Net Push out"), PushOutImpusleColor);
+					}
+				}
+
+				if (EnumHasAnyFlags(static_cast<EChaosVDCollisionVisualizationFlags>(VisualizationFlagsToUse), EChaosVDCollisionVisualizationFlags::NetImpulse))
+				{
+					if (ManifoldPoint.bIsValid && !ManifoldPoint.NetImpulse.IsNearlyZero())
+					{
+						FChaosVDDebugDrawUtils::DrawArrowVector(PDI, WorldPointPlaneLocation, OwnerActorLocation + WorldPointPlaneLocation + VisualizationContext.SpaceTransform.TransformPosition(ManifoldPoint.NetImpulse), TEXT("Net Impulse"), ImpulseColor);
+					}
+				}
+
+				if (EnumHasAnyFlags(static_cast<EChaosVDCollisionVisualizationFlags>(VisualizationFlagsToUse), EChaosVDCollisionVisualizationFlags::ContactPoints))
+				{
+					FChaosVDDebugDrawUtils::DrawPoint(PDI, WorldPlaneLocation, TEXT("Contact : WorldPlaneLocation"), DiscColor);
+				}
+
+				if (EnumHasAnyFlags(static_cast<EChaosVDCollisionVisualizationFlags>(VisualizationFlagsToUse), EChaosVDCollisionVisualizationFlags::ContactNormal))
+				{
+					FColor NormalColor = ((ManifoldPoint.ContactPoint.ContactType != EChaosVDContactPointType::EdgeEdge) ? PlaneNormalColor : EdgeNormalColor);
+					constexpr int32 Scale = 10.0f;
+					FChaosVDDebugDrawUtils::DrawArrowVector(PDI, WorldPlaneLocation, WorldPlaneLocation + WorldPlaneNormal * Scale, TEXT("Contact Normal"), NormalColor, 2.0f);
+				}
+
+				if (EnumHasAnyFlags(static_cast<EChaosVDCollisionVisualizationFlags>(VisualizationFlagsToUse), EChaosVDCollisionVisualizationFlags::ContactPoints))
+				{
+					if (ManifoldPoint.ContactPoint.Phi < FLT_MAX)
+					{
+						FColor Color = FColor(128, 128, 0);
+						FChaosVDDebugDrawUtils::DrawPoint(PDI, WorldPlaneLocation - ManifoldPoint.ContactPoint.Phi * WorldPlaneNormal, TEXT("Contact : WorldPlaneLocation - Phi"), Color);
+					}
+					
+					// Manifold point
+					FChaosVDDebugDrawUtils::DrawPoint(PDI, WorldPointLocation , TEXT("Contact : Manifold Point"), DiscColor);
+				}
+
+				if (EnumHasAnyFlags(static_cast<EChaosVDCollisionVisualizationFlags>(VisualizationFlagsToUse), EChaosVDCollisionVisualizationFlags::AccumulatedImpulse))
+				{
+					if (!Constraint.AccumulatedImpulse.IsNearlyZero())
+					{
+						FChaosVDDebugDrawUtils::DrawArrowVector(PDI, WorldActorTransform0.GetLocation(), WorldActorTransform0.GetLocation() + Constraint.AccumulatedImpulse, TEXT("AccumulatedImpulse"), FColor::White, 5.0f);
+					}
+				}
+			}
+		}
+	}
+}
diff --git a/Engine/Plugins/Experimental/ChaosVD/Source/ChaosVD/Private/Visualizers/ChaosVDDataVisualizerBase.cpp b/Engine/Plugins/Experimental/ChaosVD/Source/ChaosVD/Private/Visualizers/ChaosVDDataVisualizerBase.cpp
new file mode 100644
--- /dev/null
+++ b/Engine/Plugins/Experimental/ChaosVD/Source/ChaosVD/Private/Visualizers/ChaosVDDataVisualizerBase.cpp
@@ -0,0 +1,4 @@
+﻿// Copyright Epic Games, Inc. All Rights Reserved.
+
+#include "Visualizers/ChaosVDDataVisualizerBase.h"
+
diff --git a/Engine/Plugins/Experimental/ChaosVD/Source/ChaosVD/Private/Visualizers/ChaosVDDebugDrawUtils.cpp b/Engine/Plugins/Experimental/ChaosVD/Source/ChaosVD/Private/Visualizers/ChaosVDDebugDrawUtils.cpp
new file mode 100644
--- /dev/null
+++ b/Engine/Plugins/Experimental/ChaosVD/Source/ChaosVD/Private/Visualizers/ChaosVDDebugDrawUtils.cpp
@@ -0,0 +1,69 @@
+﻿// Copyright Epic Games, Inc. All Rights Reserved.
+
+#include "Visualizers/ChaosVDDebugDrawUtils.h"
+
+#include "CanvasItem.h"
+#include "ChaosVDEditorSettings.h"
+#include "DebugRenderSceneProxy.h"
+#include "Engine/Engine.h"
+
+TQueue<FChaosVDDebugDrawUtils::FChaosVDQueuedTextToDraw> FChaosVDDebugDrawUtils::TexToDrawQueue = TQueue<FChaosVDQueuedTextToDraw>();
+
+void FChaosVDDebugDrawUtils::DrawArrowVector(FPrimitiveDrawInterface* PDI, const FVector& StartLocation, const FVector& EndLocation, FStringView DebugText, const FColor& Color, float ArrowSize)
+{
+	const FDebugRenderSceneProxy::FArrowLine VelocityArrowLine = { StartLocation, EndLocation, Color};
+	VelocityArrowLine.Draw(PDI, ArrowSize);
+
+	if (!DebugText.IsEmpty())
+	{
+		// Draw the text in the middle of the vector line
+		const FVector VectorToDraw = EndLocation - StartLocation;
+		const FVector TextWorldPosition = StartLocation + VectorToDraw  * 0.5f;
+		DrawText(DebugText.GetData(), TextWorldPosition , Color);
+	}
+}
+
+void FChaosVDDebugDrawUtils::DrawPoint(FPrimitiveDrawInterface* PDI, const FVector& Location, FStringView DebugText, const FColor& Color, float Size)
+{
+	if (DebugText.IsEmpty())
+	{
+		return;
+	}
+
+	PDI->DrawPoint(Location, Color, Size, ESceneDepthPriorityGroup::SDPG_World);
+
+	if (!DebugText.IsEmpty())
+	{
+		DrawText(DebugText, Location, Color);
+	}
+}
+
+void FChaosVDDebugDrawUtils::DrawText(FStringView StringToDraw, const FVector& Location, const FColor& Color)
+{
+	if (const UChaosVDEditorSettings* CVDEditorSettings = GetDefault<UChaosVDEditorSettings>())
+	{
+		if (CVDEditorSettings->bShowDebugText)
+		{
+			TexToDrawQueue.Enqueue({StringToDraw.GetData(), Location, Color });
+		}
+	}
+}
+
+void FChaosVDDebugDrawUtils::DrawCanvas(FViewport& InViewport, FSceneView& View, FCanvas& Canvas)
+{
+	while (!TexToDrawQueue.IsEmpty())
+	{
+		FChaosVDQueuedTextToDraw TextToDraw;
+		if (TexToDrawQueue.Dequeue(TextToDraw))
+		{
+			FVector2D PixelLocation;
+			if (View.WorldToPixel(TextToDraw.WorldPosition, PixelLocation))
+			{
+				FCanvasTextItem TextItem(PixelLocation, FText::AsCultureInvariant(TextToDraw.Text), GEngine->GetSmallFont(), TextToDraw.Color);
+				TextItem.Scale = FVector2D::UnitVector;
+				TextItem.EnableShadow(FLinearColor::Black);
+				TextItem.Draw(&Canvas);
+			}
+		}
+	}
+}
diff --git a/Engine/Plugins/Experimental/ChaosVD/Source/ChaosVD/Private/Visualizers/ChaosVDParticleDataVisualizer.cpp b/Engine/Plugins/Experimental/ChaosVD/Source/ChaosVD/Private/Visualizers/ChaosVDParticleDataVisualizer.cpp
new file mode 100644
--- /dev/null
+++ b/Engine/Plugins/Experimental/ChaosVD/Source/ChaosVD/Private/Visualizers/ChaosVDParticleDataVisualizer.cpp
@@ -0,0 +1,79 @@
+﻿// Copyright Epic Games, Inc. All Rights Reserved.
+
+#include "Visualizers/ChaosVDParticleDataVisualizer.h"
+
+#include "ChaosVDEditorSettings.h"
+#include "ChaosVDModule.h"
+#include "DataWrappers/ChaosVDParticleDataWrapper.h"
+#include "Visualizers/ChaosVDDebugDrawUtils.h"
+
+void FChaosVDParticleDataVisualizer::DrawVisualization(const FSceneView* View, FPrimitiveDrawInterface* PDI)
+{
+	uint32 VisualizationFlagsToUse = 0;
+	if (static_cast<EChaosVDParticleDataVisualizationFlags>(LocalVisualizationFlags) == EChaosVDParticleDataVisualizationFlags::None)
+	{
+		if (const UChaosVDEditorSettings* EditorSettings = GetDefault<UChaosVDEditorSettings>())
+		{
+			VisualizationFlagsToUse = EditorSettings->GlobalParticleDataVisualizationFlags;
+		}
+		else
+		{
+			UE_LOG(LogChaosVDEditor, Warning, TEXT("[%s] Failed to retrive global visualization setting. Falling back to local settings"), ANSI_TO_TCHAR(__FUNCTION__));
+
+			VisualizationFlagsToUse = LocalVisualizationFlags;
+		}	
+	}
+	else
+	{
+		VisualizationFlagsToUse = LocalVisualizationFlags;
+	}
+	
+	FChaosVDVisualizationContext VisualizationContext;
+	DataProvider.GetVisualizationContext(VisualizationContext);
+
+	const FChaosVDParticleDataWrapper* ParticleDataViewer = DataProvider.GetParticleData();
+	if (!ensure(ParticleDataViewer))
+	{
+		return;
+	}
+
+	// TODO: Implement scale settings for the vectors, and take into account the simulation space transform
+
+	const FVector& OwnerLocation = ParticleDataViewer->ParticlePositionRotation.MX;
+
+	if (ParticleDataViewer->ParticleVelocities.HasValidData())
+	{
+		if (EnumHasAnyFlags(static_cast<EChaosVDParticleDataVisualizationFlags>(VisualizationFlagsToUse), EChaosVDParticleDataVisualizationFlags::Velocity))
+		{
+			FChaosVDDebugDrawUtils::DrawArrowVector(PDI, OwnerLocation, OwnerLocation + ParticleDataViewer->ParticleVelocities.MV, TEXT("Velocity"), FColor::Green);
+		}
+
+		if (EnumHasAnyFlags(static_cast<EChaosVDParticleDataVisualizationFlags>(VisualizationFlagsToUse), EChaosVDParticleDataVisualizationFlags::AngularVelocity))
+		{
+			FChaosVDDebugDrawUtils::DrawArrowVector(PDI, OwnerLocation, OwnerLocation + ParticleDataViewer->ParticleVelocities.MW, TEXT("Angular Velocity"), FColor::Blue);
+		}
+	}
+
+	if (ParticleDataViewer->ParticleDynamics.HasValidData())
+	{
+		if (EnumHasAnyFlags(static_cast<EChaosVDParticleDataVisualizationFlags>(VisualizationFlagsToUse), EChaosVDParticleDataVisualizationFlags::Acceleration))
+		{
+			FChaosVDDebugDrawUtils::DrawArrowVector(PDI, OwnerLocation, OwnerLocation + ParticleDataViewer->ParticleDynamics.MAcceleration, TEXT("Acceleration"), FColor::Orange);
+		}
+		
+		if (EnumHasAnyFlags(static_cast<EChaosVDParticleDataVisualizationFlags>(VisualizationFlagsToUse), EChaosVDParticleDataVisualizationFlags::AngularAcceleration))
+		{
+			FChaosVDDebugDrawUtils::DrawArrowVector(PDI, OwnerLocation, OwnerLocation + ParticleDataViewer->ParticleDynamics.MAngularAcceleration, TEXT("Angular Acceleration"), FColor::Purple);
+		}
+
+		if (EnumHasAnyFlags(static_cast<EChaosVDParticleDataVisualizationFlags>(VisualizationFlagsToUse), EChaosVDParticleDataVisualizationFlags::LinearImpulse))
+		{
+			FChaosVDDebugDrawUtils::DrawArrowVector(PDI, OwnerLocation, OwnerLocation + ParticleDataViewer->ParticleDynamics.MLinearImpulseVelocity, TEXT("Linear Implulse Velocity"), FColor::Purple);
+		}
+
+		if (EnumHasAnyFlags(static_cast<EChaosVDParticleDataVisualizationFlags>(VisualizationFlagsToUse), EChaosVDParticleDataVisualizationFlags::AngularImpulse))
+		{
+			FChaosVDDebugDrawUtils::DrawArrowVector(PDI, OwnerLocation, OwnerLocation + ParticleDataViewer->ParticleDynamics.MAngularImpulseVelocity, TEXT("Angular Implulse Velocity"), FColor::Emerald);
+		}
+	}
+}
diff --git a/Engine/Plugins/Experimental/ChaosVD/Source/ChaosVD/Private/Widgets/SChaosVDMainTab.cpp b/Engine/Plugins/Experimental/ChaosVD/Source/ChaosVD/Private/Widgets/SChaosVDMainTab.cpp
--- a/Engine/Plugins/Experimental/ChaosVD/Source/ChaosVD/Private/Widgets/SChaosVDMainTab.cpp
+++ b/Engine/Plugins/Experimental/ChaosVD/Source/ChaosVD/Private/Widgets/SChaosVDMainTab.cpp
@@ -2,6 +2,7 @@
 
 #include "Widgets/SChaosVDMainTab.h"
 
+#include "ChaosVDEditorSettingsTab.h"
 #include "ChaosVDEngine.h"
 #include "ChaosVDObjectDetailsTab.h"
 #include "ChaosVDOutputLogTab.h"
@@ -36,6 +37,7 @@ void SChaosVDMainTab::Construct(const FArguments& InArgs, TSharedPtr<FChaosVDEng
 	OutputLogTab = MakeShared<FChaosVDOutputLogTab>(FChaosVDTabID::OutputLog, TabManager, this);
 	PlaybackViewportTab = MakeShared<FChaosVDPlaybackViewportTab>(FChaosVDTabID::PlaybackViewport, TabManager, this);
 	SolversTracksTab = MakeShared<FChaosVDSolversTracksTab>(FChaosVDTabID::SolversTrack, TabManager, this);
+	EditorSettingsTab = MakeShared<FChaosVDEditorSettingsTab>(FChaosVDTabID::CVDEditorSettings, TabManager, this);
 
 	GenerateMainWindowMenu();
 
@@ -191,6 +193,7 @@ TSharedRef<FTabManager::FLayout> SChaosVDMainTab::GenerateMainLayout()
 					FTabManager::NewStack()
 					->SetSizeCoefficient(0.3f)
 					->AddTab(FChaosVDTabID::DetailsPanel, ETabState::OpenedTab)
+					->AddTab(FChaosVDTabID::CVDEditorSettings, ETabState::OpenedTab)
 				)
 			)
 		);
diff --git a/Engine/Plugins/Experimental/ChaosVD/Source/ChaosVD/Private/Widgets/SChaosVDPlaybackViewport.cpp b/Engine/Plugins/Experimental/ChaosVD/Source/ChaosVD/Private/Widgets/SChaosVDPlaybackViewport.cpp
--- a/Engine/Plugins/Experimental/ChaosVD/Source/ChaosVD/Private/Widgets/SChaosVDPlaybackViewport.cpp
+++ b/Engine/Plugins/Experimental/ChaosVD/Source/ChaosVD/Private/Widgets/SChaosVDPlaybackViewport.cpp
@@ -91,6 +91,7 @@ void SChaosVDPlaybackViewport::Construct(const FArguments& InArgs, TWeakPtr<FCha
 			[
 				SAssignNew(GameFramesTimelineWidget, SChaosVDTimelineWidget)
 				.HidePlayStopButtons(false)
+				.HideLockButton(true)
 				.OnFrameChanged_Raw(this, &SChaosVDPlaybackViewport::OnFrameSelectionUpdated)
 				.MaxFrames(0)
 			]
diff --git a/Engine/Plugins/Experimental/ChaosVD/Source/ChaosVD/Private/Widgets/SChaosVDSolverPlaybackControls.cpp b/Engine/Plugins/Experimental/ChaosVD/Source/ChaosVD/Private/Widgets/SChaosVDSolverPlaybackControls.cpp
--- a/Engine/Plugins/Experimental/ChaosVD/Source/ChaosVD/Private/Widgets/SChaosVDSolverPlaybackControls.cpp
+++ b/Engine/Plugins/Experimental/ChaosVD/Source/ChaosVD/Private/Widgets/SChaosVDSolverPlaybackControls.cpp
@@ -2,6 +2,7 @@
 
 #include "Widgets/SChaosVDSolverPlaybackControls.h"
 
+#include "ChaosVDEditorSettings.h"
 #include "ChaosVDPlaybackController.h"
 #include "Widgets/SBoxPanel.h"
 #include "Widgets/SChaosVDPlaybackViewport.h"
@@ -39,6 +40,7 @@ void SChaosVDSolverPlaybackControls::Construct(const FArguments& InArgs, int32 I
 				[
 					SAssignNew(FramesTimelineWidget, SChaosVDTimelineWidget)
 						.HidePlayStopButtons(false)
+						.HideLockButton(true)
 						.OnFrameChanged_Raw(this, &SChaosVDSolverPlaybackControls::OnFrameSelectionUpdated)
 						.MaxFrames(0)
 				]
@@ -59,6 +61,7 @@ void SChaosVDSolverPlaybackControls::Construct(const FArguments& InArgs, int32 I
 				[
 					SAssignNew(StepsTimelineWidget, SChaosVDTimelineWidget)
 					.HidePlayStopButtons(true)
+					.OnFrameLockStateChanged_Raw(this, &SChaosVDSolverPlaybackControls::HandleLockStateChanged)
 					.OnFrameChanged_Raw(this, &SChaosVDSolverPlaybackControls::OnStepSelectionUpdated)
 					.MaxFrames(0)
 				]
@@ -158,8 +161,22 @@ void SChaosVDSolverPlaybackControls::HandleControllerTrackFrameUpdated(TWeakPtr<
 		{
 			FramesTimelineWidget->SetCurrentTimelineFrame(SolverTrackInfo->CurrentFrame, EChaosVDSetTimelineFrameFlags::None);
 
-			constexpr int32 StepNumber = 0;
-			UpdateStepsWidgetForFrame(*CurrentPlaybackControllerPtr.Get(), SolverTrackInfo->CurrentFrame, StepNumber);
+			UpdateStepsWidgetForFrame(*CurrentPlaybackControllerPtr.Get(), SolverTrackInfo->CurrentFrame, SolverTrackInfo->CurrentStep);
+		}
+	}
+}
+
+void SChaosVDSolverPlaybackControls::HandleLockStateChanged(bool NewIsLocked)
+{
+	if (const TSharedPtr<FChaosVDPlaybackController> CurrentPlaybackControllerPtr = PlaybackController.Pin())
+	{
+		if (NewIsLocked)
+		{
+			CurrentPlaybackControllerPtr->LockTrackInCurrentStep(EChaosVDTrackType::Solver, SolverID);
+		}
+		else
+		{
+			CurrentPlaybackControllerPtr->UnlockTrackStep(EChaosVDTrackType::Solver, SolverID);
 		}
 	}
 }
@@ -168,7 +185,12 @@ void SChaosVDSolverPlaybackControls::OnFrameSelectionUpdated(int32 NewFrameIndex
 {
 	if (const TSharedPtr<FChaosVDPlaybackController> PlaybackControllerPtr = PlaybackController.Pin())
 	{
-		constexpr int32 StepNumber = 0;
+		const int32 LastStepNumber = PlaybackControllerPtr->GetTrackLastStepAtFrame(EChaosVDTrackType::Solver, SolverID, NewFrameIndex);
+		const int32 CurrentStep = PlaybackControllerPtr->GetTrackCurrentStep(EChaosVDTrackType::Solver, SolverID);
+
+		// If the steps control is unlocked, each time we go to a new frame we should start at step 0.
+		const int32 StepNumber = StepsTimelineWidget->IsUnlocked() ? LastStepNumber : CurrentStep;
+
 		UpdateStepsWidgetForFrame(*PlaybackControllerPtr.Get(), NewFrameIndex, StepNumber);
 
 		PlaybackControllerPtr->GoToTrackFrame(GetInstigatorID(), EChaosVDTrackType::Solver, SolverID, NewFrameIndex, StepNumber);
diff --git a/Engine/Plugins/Experimental/ChaosVD/Source/ChaosVD/Private/Widgets/SChaosVDTimelineWidget.cpp b/Engine/Plugins/Experimental/ChaosVD/Source/ChaosVD/Private/Widgets/SChaosVDTimelineWidget.cpp
--- a/Engine/Plugins/Experimental/ChaosVD/Source/ChaosVD/Private/Widgets/SChaosVDTimelineWidget.cpp
+++ b/Engine/Plugins/Experimental/ChaosVD/Source/ChaosVD/Private/Widgets/SChaosVDTimelineWidget.cpp
@@ -16,6 +16,7 @@ void SChaosVDTimelineWidget::Construct(const FArguments& InArgs)
 {
 	MaxFrames = InArgs._MaxFrames;
 	FrameChangedDelegate = InArgs._OnFrameChanged;
+	FrameLockedDelegate = InArgs._OnFrameLockStateChanged;
 
 	SetCanTick(false);
 
@@ -25,13 +26,14 @@ void SChaosVDTimelineWidget::Construct(const FArguments& InArgs)
 			+SHorizontalBox::Slot()
 			.HAlign(HAlign_Center)
 			.VAlign(VAlign_Center)
-			.FillWidth(0.2f)
+			.FillWidth(0.25f)
 			[
 				SNew(SHorizontalBox)
 				+SHorizontalBox::Slot()
 				[
 					SNew(SButton)
-					.Visibility(InArgs._HidePlayStopButtons ? EVisibility::Collapsed : EVisibility::Visible)
+					.Visibility(InArgs._HidePlayStopButtons.Get() ? EVisibility::Collapsed : EVisibility::Visible)
+					.IsEnabled_Raw(this, &SChaosVDTimelineWidget::IsUnlocked)
 					.OnClicked( FOnClicked::CreateRaw(this, &SChaosVDTimelineWidget::Play))
 					.ContentPadding( 2.0f )
 					.ForegroundColor( FSlateColor::UseForeground() )
@@ -46,7 +48,8 @@ void SChaosVDTimelineWidget::Construct(const FArguments& InArgs)
 				+SHorizontalBox::Slot()
 				[
 					SNew(SButton)
-					.Visibility(InArgs._HidePlayStopButtons ? EVisibility::Collapsed : EVisibility::Visible)
+					.Visibility(InArgs._HidePlayStopButtons.Get() ? EVisibility::Collapsed : EVisibility::Visible)
+					.IsEnabled_Raw(this, &SChaosVDTimelineWidget::IsUnlocked)
 					.OnClicked( FOnClicked::CreateRaw(this, &SChaosVDTimelineWidget::Stop))
 					.ContentPadding( 2.0f )
 					.ForegroundColor( FSlateColor::UseForeground() )
@@ -61,6 +64,8 @@ void SChaosVDTimelineWidget::Construct(const FArguments& InArgs)
 				+SHorizontalBox::Slot()
 				[
 					SNew(SButton)
+					.Visibility(InArgs._HideNextPrevButtons.Get() ? EVisibility::Collapsed : EVisibility::Visible)
+					.IsEnabled_Raw(this, &SChaosVDTimelineWidget::IsUnlocked)
 					.OnClicked( FOnClicked::CreateRaw(this, &SChaosVDTimelineWidget::Prev))
 					.ContentPadding( 2.0f )
 					.ForegroundColor( FSlateColor::UseForeground() )
@@ -75,6 +80,8 @@ void SChaosVDTimelineWidget::Construct(const FArguments& InArgs)
 				+SHorizontalBox::Slot()
 				[
 					SNew(SButton)
+					.Visibility(InArgs._HideNextPrevButtons.Get() ? EVisibility::Collapsed : EVisibility::Visible)
+					.IsEnabled_Raw(this, &SChaosVDTimelineWidget::IsUnlocked)
 					.OnClicked( FOnClicked::CreateRaw(this, &SChaosVDTimelineWidget::Next))
 					.ContentPadding( 2.0f )
 					.ForegroundColor( FSlateColor::UseForeground() )
@@ -86,10 +93,25 @@ void SChaosVDTimelineWidget::Construct(const FArguments& InArgs)
 						.ColorAndOpacity( FSlateColor::UseForeground() )
 					]
 				]
+				+SHorizontalBox::Slot()
+				[
+					SNew(SButton)
+					.Visibility(InArgs._HideLockButton.Get() ? EVisibility::Collapsed : EVisibility::Visible)
+					.OnClicked( FOnClicked::CreateRaw(this, &SChaosVDTimelineWidget::ToggleLockState))
+					.ContentPadding( 2.0f )
+					.ForegroundColor( FSlateColor::UseForeground())
+					.IsFocusable( false )
+					[
+						SNew( SImage )
+						.Image_Raw(this, &SChaosVDTimelineWidget::GetLockStateIcon)
+						.DesiredSizeOverride(FVector2D(16.0f,16.0f))
+						.ColorAndOpacity( FSlateColor::UseForeground() )
+					]
+				]
 			]
 			+SHorizontalBox::Slot()
 			.VAlign(VAlign_Center)
-			.FillWidth(0.7f)
+			.FillWidth(0.65f)
 			[
 			  SAssignNew(TimelineSlider, SSlider)
 			  .ToolTipText_Lambda([this]()-> FText{ return FText::AsNumber(CurrentFrame); })
@@ -207,11 +229,24 @@ FReply SChaosVDTimelineWidget::Prev()
 	return FReply::Handled();
 }
 
+FReply SChaosVDTimelineWidget::ToggleLockState()
+{
+	bIsLocked = !bIsLocked;
+
+	FrameLockedDelegate.ExecuteIfBound(bIsLocked);
+
+	return FReply::Handled();
+}
+
 const FSlateBrush* SChaosVDTimelineWidget::GetPlayOrPauseIcon() const
 {
 	return bIsPlaying ? FChaosVDStyle::Get().GetBrush("PauseIcon") : FChaosVDStyle::Get().GetBrush("PlayIcon");
 }
 
+const FSlateBrush* SChaosVDTimelineWidget::GetLockStateIcon() const
+{
+	return bIsLocked ? FChaosVDStyle::Get().GetBrush("LockIcon") : FChaosVDStyle::Get().GetBrush("UnlockedIcon");
+}
 
 void SChaosVDTimelineWidget::Tick(const FGeometry& AllottedGeometry, const double InCurrentTime, const float InDeltaTime)
 {
diff --git a/Engine/Plugins/Experimental/ChaosVD/Source/ChaosVD/Public/ChaosVDEditorSettings.h b/Engine/Plugins/Experimental/ChaosVD/Source/ChaosVD/Public/ChaosVDEditorSettings.h
--- a/Engine/Plugins/Experimental/ChaosVD/Source/ChaosVD/Public/ChaosVDEditorSettings.h
+++ b/Engine/Plugins/Experimental/ChaosVD/Source/ChaosVD/Public/ChaosVDEditorSettings.h
@@ -15,9 +15,15 @@ class UChaosVDEditorSettings : public UObject
 	GENERATED_BODY()
 public:
 
-	UPROPERTY(Config, EditAnywhere, Category = Debug)
-	TSoftObjectPtr<UStaticMesh> DebugMesh;
+	UPROPERTY(EditAnywhere, Category = "Viewport Visualization Flags", meta = (Bitmask, BitmaskEnum = EChaosVDParticleDataVisualizationFlags))
+	uint8 GlobalParticleDataVisualizationFlags = 0;
+	
+	UPROPERTY(EditAnywhere, Category = "Viewport Visualization Flags", meta = (Bitmask, BitmaskEnum = EChaosVDCollisionVisualizationFlags))
+	uint8 GlobalCollisionDataVisualizationFlags = 0;
 
-	UPROPERTY(Config, EditAnywhere, Category = Debug)
+	UPROPERTY(EditAnywhere, Category = "Viewport Visualization Flags", meta = (Bitmask, BitmaskEnum = EChaosVDCollisionVisualizationFlags))
+	bool bShowDebugText = false;
+
+	UPROPERTY(Config, EditInstanceOnly, Category = "Editor Options")
 	TSoftObjectPtr<UWorld> BasePhysicsVDWorld;
 };
\ No newline at end of file
diff --git a/Engine/Plugins/Experimental/ChaosVD/Source/ChaosVD/Public/ChaosVDEditorSettingsTab.h b/Engine/Plugins/Experimental/ChaosVD/Source/ChaosVD/Public/ChaosVDEditorSettingsTab.h
new file mode 100644
--- /dev/null
+++ b/Engine/Plugins/Experimental/ChaosVD/Source/ChaosVD/Public/ChaosVDEditorSettingsTab.h
@@ -0,0 +1,15 @@
+﻿// Copyright Epic Games, Inc. All Rights Reserved.
+#pragma once
+#include "ChaosVDTabSpawnerBase.h"
+
+class FChaosVDEditorSettingsTab : public FChaosVDTabSpawnerBase
+{
+public:
+	FChaosVDEditorSettingsTab(const FName& InTabID, const TSharedPtr<FTabManager>& InTabManager, SChaosVDMainTab* InOwningTabWidget)
+		: FChaosVDTabSpawnerBase(InTabID, InTabManager, InOwningTabWidget)
+	{
+	}
+
+
+	virtual TSharedRef<SDockTab> HandleTabSpawned(const FSpawnTabArgs& Args) override;
+};
diff --git a/Engine/Plugins/Experimental/ChaosVD/Source/ChaosVD/Public/ChaosVDParticleActor.h b/Engine/Plugins/Experimental/ChaosVD/Source/ChaosVD/Public/ChaosVDParticleActor.h
--- a/Engine/Plugins/Experimental/ChaosVD/Source/ChaosVD/Public/ChaosVDParticleActor.h
+++ b/Engine/Plugins/Experimental/ChaosVD/Source/ChaosVD/Public/ChaosVDParticleActor.h
@@ -5,6 +5,9 @@
 #include "DataWrappers/ChaosVDParticleDataWrapper.h"
 #include "DataWrappers/ChaosVDCollisionDataWrappers.h"
 #include "GameFramework/Actor.h"
+#include "Visualizers/ChaosVDCollisionDataVisualizer.h"
+#include "Visualizers/ChaosVDParticleDataVisualizer.h"
+#include "Visualizers/ChaosVDDataVisualizerBase.h"
 
 #include "ChaosVDParticleActor.generated.h"
 
@@ -30,7 +33,7 @@ ENUM_CLASS_FLAGS(EChaosVDActorGeometryUpdateFlags)
 
 /** Actor used to represent a Chaos Particle in the Visual Debugger's world */
 UCLASS()
-class AChaosVDParticleActor : public AActor
+class AChaosVDParticleActor : public AActor, public IChaosVDParticleVisualizationDataProvider, public IChaosVDVisualizerContainerInterface
 {
 	GENERATED_BODY()
 public:
@@ -48,7 +51,12 @@ public:
 
 	virtual void BeginDestroy() override;
 
-	const FChaosVDParticleDataWrapper& GetParticleData() { return ParticleDataViewer; }
+	virtual const FChaosVDParticleDataWrapper* GetParticleData() override { return &ParticleDataViewer; }
+	virtual void GetVisualizationContext(FChaosVDVisualizationContext& OutVisualizationContext) override;
+
+	void CreateVisualizers();
+
+	virtual void DrawVisualization(const FSceneView* View, FPrimitiveDrawInterface* PDI) override;
 	
 #if WITH_EDITOR
 	virtual bool IsSelectedInEditor() const override;
@@ -56,10 +64,21 @@ public:
 #endif
 
 protected:
+	
+	UPROPERTY(EditAnywhere, Category = "Viewport Visualization Flags", meta = (Bitmask, BitmaskEnum = EChaosVDParticleDataVisualizationFlags))
+	uint8 LocalParticleDataVisualizationFlags;
+	
+	UPROPERTY(EditAnywhere, Category = "Viewport Visualization Flags", meta = (Bitmask, BitmaskEnum = EChaosVDCollisionVisualizationFlags))
+	uint8 LocalCollisionDataVisualizationFlags;
 
-	UPROPERTY(EditAnywhere, Category="Particle Data")
+	UPROPERTY(EditAnywhere, Category = "Viewport Visualization Flags")
+	bool bShowDebugText = false;
+
+	UPROPERTY(EditAnywhere, Category = "Particle Data")
 	FChaosVDParticleDataWrapper ParticleDataViewer;
 
+	FTransform CachedSimulationTransform;
+
 	bool bIsGeometryDataGenerationStarted = false;
 
 	TWeakPtr<FChaosVDScene> OwningScene;
@@ -67,4 +86,6 @@ protected:
 	TArray<TWeakObjectPtr<UMeshComponent>> MeshComponents;
 
 	FDelegateHandle GeometryUpdatedDelegate;
+
+	TMap<FStringView, TUniquePtr<FChaosVDDataVisualizerBase>> CVDVisualizers;
 };
diff --git a/Engine/Plugins/Experimental/ChaosVD/Source/ChaosVD/Public/ChaosVDParticleActorCustomization.h b/Engine/Plugins/Experimental/ChaosVD/Source/ChaosVD/Public/ChaosVDParticleActorCustomization.h
--- a/Engine/Plugins/Experimental/ChaosVD/Source/ChaosVD/Public/ChaosVDParticleActorCustomization.h
+++ b/Engine/Plugins/Experimental/ChaosVD/Source/ChaosVD/Public/ChaosVDParticleActorCustomization.h
@@ -9,10 +9,15 @@
 class FChaosVDParticleActorCustomization : public IDetailCustomization
 {
 public:
-	
+	FChaosVDParticleActorCustomization();
+
 	inline static FName ChaosVDCategoryName = FName("Particle Data");
+	inline static FName ChaosVDVisualizationCategoryName = FName("Viewport Visualization Flags");
 
 	static TSharedRef<IDetailCustomization> MakeInstance();
 
 	virtual void CustomizeDetails(IDetailLayoutBuilder& DetailBuilder) override;
+
+private:
+	TSet<FName> AllowedCategories;
 };
diff --git a/Engine/Plugins/Experimental/ChaosVD/Source/ChaosVD/Public/ChaosVDPlaybackController.h b/Engine/Plugins/Experimental/ChaosVD/Source/ChaosVD/Public/ChaosVDPlaybackController.h
--- a/Engine/Plugins/Experimental/ChaosVD/Source/ChaosVD/Public/ChaosVDPlaybackController.h
+++ b/Engine/Plugins/Experimental/ChaosVD/Source/ChaosVD/Public/ChaosVDPlaybackController.h
@@ -31,6 +31,7 @@ struct FChaosVDTrackInfo
 	EChaosVDTrackType TrackType = EChaosVDTrackType::Invalid;
 	int32 CurrentFrame = INDEX_NONE;
 	int32 CurrentStep = INDEX_NONE;
+	int32 LockedOnStep = INDEX_NONE;
 	int32 MaxFrames = INDEX_NONE;
 	FString TrackName;
 };
@@ -123,6 +124,14 @@ public:
 	 */
 	int32 GetTrackCurrentStep(EChaosVDTrackType TrackType, int32 InTrackID) const;
 
+	/**
+	 * Gets the index number of the last step available (available steps -1)
+	 * @param TrackType Type of the track to evaluate
+	 * @param InTrackID ID of the track to evaluate
+	 * @return Number of the last step
+	 */
+	int32 GetTrackLastStepAtFrame(EChaosVDTrackType TrackType, int32 InTrackID, int32 InFrameNumber) const;
+
 	/** Converts the current frame number of a track, to a frame number in other tracks space time
 	 * @param FromTrack Track info with the current frame number we want to convert
 	 * @param ToTrack Track info we want to use to convert the frame to
@@ -154,6 +163,29 @@ public:
 	 */
 	const FChaosVDTrackInfo* GetTrackInfo(EChaosVDTrackType TrackType, int32 TrackID);
 
+	/**
+	 * Gets the track info of the specified type with the specified ID
+	 * @param TrackType Type of the track to find
+	 * @param TrackID ID of the track to find
+	 * @return Ptr to the found track info data - Null if nothing was found.
+	 */
+	FChaosVDTrackInfo* GetMutableTrackInfo(EChaosVDTrackType TrackType, int32 TrackID);
+
+
+	/**
+	 * Locks the steps timeline of a given track so each time you move between frames, it will automatically scrub to the locked in step
+	 * @param TrackType Type of the track to find
+	 * @param TrackID ID of the track to find
+	 */
+	void LockTrackInCurrentStep(EChaosVDTrackType TrackType, int32 TrackID);
+
+	/**
+	 * UnLocks the steps timeline of a given track so each time you move between frames, it will automatically scrub to the default step
+	 * @param TrackType Type of the track to find
+	 * @param TrackID ID of the track to find
+	 */
+	void UnlockTrackStep(EChaosVDTrackType TrackType, int32 TrackID);
+
 	/** Returns a weak ptr pointer to the loaded recording */
 	TWeakPtr<FChaosVDRecording> GetCurrentRecording() { return LoadedRecording; }
 
diff --git a/Engine/Plugins/Experimental/ChaosVD/Source/ChaosVD/Public/ChaosVDPlaybackViewportClient.h b/Engine/Plugins/Experimental/ChaosVD/Source/ChaosVD/Public/ChaosVDPlaybackViewportClient.h
--- a/Engine/Plugins/Experimental/ChaosVD/Source/ChaosVD/Public/ChaosVDPlaybackViewportClient.h
+++ b/Engine/Plugins/Experimental/ChaosVD/Source/ChaosVD/Public/ChaosVDPlaybackViewportClient.h
@@ -24,6 +24,11 @@ private:
 
 	void HandleObjectFocused(UObject* FocusedObject);
 
+public:
+	virtual void Draw(const FSceneView* View, FPrimitiveDrawInterface* PDI) override;
+	virtual void DrawCanvas(FViewport& InViewport, FSceneView& View, FCanvas& Canvas) override;
+
+private:
 	FDelegateHandle ObjectFocusedDelegateHandle;
 	UWorld* CVDWorld;
 	TWeakPtr<FChaosVDScene> CVDScene;
diff --git a/Engine/Plugins/Experimental/ChaosVD/Source/ChaosVD/Public/ChaosVDScene.h b/Engine/Plugins/Experimental/ChaosVD/Source/ChaosVD/Public/ChaosVDScene.h
--- a/Engine/Plugins/Experimental/ChaosVD/Source/ChaosVD/Public/ChaosVDScene.h
+++ b/Engine/Plugins/Experimental/ChaosVD/Source/ChaosVD/Public/ChaosVDScene.h
@@ -76,6 +76,12 @@ public:
 	/** Event triggered when an object is focused in the scene (double click in the scene outliner)*/
 	FChaosVDOnObjectSelectedDelegate& OnObjectFocused() { return ObjectFocusedDelegate; }
 
+	/** Returns a ptr to the particle actor representing the provided Particle ID
+	 * @param SolverID ID of the solver owning the Particle
+	 * @param ParticleID ID of the particle
+	 */
+	AChaosVDParticleActor* GetParticleActor(int32 SolverID, int32 ParticleID);
+
 	TSharedPtr<FChaosVDRecording> LoadedRecording;
 
 private:
diff --git a/Engine/Plugins/Experimental/ChaosVD/Source/ChaosVD/Public/ChaosVDTabsIDs.h b/Engine/Plugins/Experimental/ChaosVD/Source/ChaosVD/Public/ChaosVDTabsIDs.h
--- a/Engine/Plugins/Experimental/ChaosVD/Source/ChaosVD/Public/ChaosVDTabsIDs.h
+++ b/Engine/Plugins/Experimental/ChaosVD/Source/ChaosVD/Public/ChaosVDTabsIDs.h
@@ -14,4 +14,5 @@ public:
 	static const FName DetailsPanel;
 	static const FName OutputLog;
 	static const FName SolversTrack;
+	static const FName CVDEditorSettings;
 };
diff --git a/Engine/Plugins/Experimental/ChaosVD/Source/ChaosVD/Public/Visualizers/ChaosVDCollisionDataVisualizer.h b/Engine/Plugins/Experimental/ChaosVD/Source/ChaosVD/Public/Visualizers/ChaosVDCollisionDataVisualizer.h
new file mode 100644
--- /dev/null
+++ b/Engine/Plugins/Experimental/ChaosVD/Source/ChaosVD/Public/Visualizers/ChaosVDCollisionDataVisualizer.h
@@ -0,0 +1,39 @@
+﻿// Copyright Epic Games, Inc. All Rights Reserved.
+
+#pragma once
+#include "ChaosVDDataVisualizerBase.h"
+#include "IChaosVDParticleVisualizationDataProvider.h"
+
+#include "ChaosVDCollisionDataVisualizer.generated.h"
+
+struct FChaosVDParticlePairMidPhase;
+struct FChaosVDParticleDataWrapper;
+
+UENUM(meta = (Bitflags, UseEnumValuesAsMaskValuesInEditor = "true"))
+enum class EChaosVDCollisionVisualizationFlags: uint32
+{
+	None					= 0 UMETA(Hidden),
+	ContactPoints			= 1 << 0,
+	NetPushOut				= 1 << 2,
+	NetImpulse				= 1 << 3,
+	ContactNormal			= 1 << 4,
+	AccumulatedImpulse		= 1 << 5,
+	DrawOnlyActiveContacts	= 1 << 6,
+};
+ENUM_CLASS_FLAGS(EChaosVDCollisionVisualizationFlags);
+
+class FChaosVDCollisionDataVisualizer final : public FChaosVDDataVisualizerBase
+{
+public:
+	explicit FChaosVDCollisionDataVisualizer(IChaosVDParticleVisualizationDataProvider& InProvider)
+		: FChaosVDDataVisualizerBase(), DataProvider(InProvider)
+	{
+	}
+
+	virtual void DrawVisualization(const FSceneView* View, FPrimitiveDrawInterface* PDI) override;
+	
+	inline static const FStringView VisualizerID = TEXT("CollisionVisualizer");
+
+protected:
+	IChaosVDParticleVisualizationDataProvider& DataProvider;
+};
diff --git a/Engine/Plugins/Experimental/ChaosVD/Source/ChaosVD/Public/Visualizers/ChaosVDDataVisualizerBase.h b/Engine/Plugins/Experimental/ChaosVD/Source/ChaosVD/Public/Visualizers/ChaosVDDataVisualizerBase.h
new file mode 100644
--- /dev/null
+++ b/Engine/Plugins/Experimental/ChaosVD/Source/ChaosVD/Public/Visualizers/ChaosVDDataVisualizerBase.h
@@ -0,0 +1,52 @@
+﻿// Copyright Epic Games, Inc. All Rights Reserved.
+
+#pragma once
+#include "UObject/Interface.h"
+
+#include "ChaosVDDataVisualizerBase.generated.h"
+
+class FSceneView;
+class FPrimitiveDrawInterface;
+
+/** Interface to be used by any object that contains a visualizer*/
+UINTERFACE(MinimalAPI)
+class UChaosVDVisualizerContainerInterface : public UInterface
+{
+	GENERATED_BODY()
+};
+
+class IChaosVDVisualizerContainerInterface
+{
+	GENERATED_BODY()
+public:
+	virtual void DrawVisualization(const FSceneView* View, FPrimitiveDrawInterface* PDI){};
+};
+
+/**
+ * Base class that represents a Visualizer (anything that will debug draw in the Chaos VD viewport).
+ * It allows add visualizers to anything and has a way limited API, just with what we need.
+ * @note UE Components Visualizers should also be supported and can be used instead if needed.
+ */
+class FChaosVDDataVisualizerBase
+{
+public:
+	FChaosVDDataVisualizerBase()
+	{
+	}
+
+	virtual ~FChaosVDDataVisualizerBase() = default;
+
+	virtual void DrawVisualization(const FSceneView* View, FPrimitiveDrawInterface* PDI) {};
+
+	/** Replaces the local visualization flags with new ones
+	 * @param NewFlags Flags to use as replacement
+	 */
+	void UpdateVisualizationFlags(uint32 NewFlags) { LocalVisualizationFlags = NewFlags; }
+
+protected:
+
+	/** Flags to override the Global visualization flags.
+	 * Chaos VD has a set of global visualization flags that affect all objects, but objects can override these by setting their own flags.
+	 */
+	uint32 LocalVisualizationFlags = 0;
+};
diff --git a/Engine/Plugins/Experimental/ChaosVD/Source/ChaosVD/Public/Visualizers/ChaosVDDebugDrawUtils.h b/Engine/Plugins/Experimental/ChaosVD/Source/ChaosVD/Public/Visualizers/ChaosVDDebugDrawUtils.h
new file mode 100644
--- /dev/null
+++ b/Engine/Plugins/Experimental/ChaosVD/Source/ChaosVD/Public/Visualizers/ChaosVDDebugDrawUtils.h
@@ -0,0 +1,26 @@
+﻿// Copyright Epic Games, Inc. All Rights Reserved.
+
+#pragma once
+#include "Containers/Queue.h"
+#include "SceneManagement.h"
+
+/** Utility methods that allows Debug draw into the Chaos VD Editor */
+class FChaosVDDebugDrawUtils
+{
+public:
+	static void DrawArrowVector(FPrimitiveDrawInterface* PDI, const FVector& StartLocation, const FVector& EndLocation, FStringView DebugText, const FColor& Color, float ArrowSize = 10.0f);
+	static void DrawPoint(FPrimitiveDrawInterface* PDI, const FVector& Location, FStringView DebugText, const FColor& Color, float Size = 10.0f);
+	static void DrawText(FStringView StringToDraw, const FVector& Location, const FColor& Color);
+
+	static void DrawCanvas(FViewport& InViewport, FSceneView& View, FCanvas& Canvas);
+
+private:
+	struct FChaosVDQueuedTextToDraw
+	{
+		FString Text;
+		FVector WorldPosition;
+		FLinearColor Color;
+	};
+
+	static TQueue<FChaosVDQueuedTextToDraw> TexToDrawQueue;
+};
diff --git a/Engine/Plugins/Experimental/ChaosVD/Source/ChaosVD/Public/Visualizers/ChaosVDParticleDataVisualizer.h b/Engine/Plugins/Experimental/ChaosVD/Source/ChaosVD/Public/Visualizers/ChaosVDParticleDataVisualizer.h
new file mode 100644
--- /dev/null
+++ b/Engine/Plugins/Experimental/ChaosVD/Source/ChaosVD/Public/Visualizers/ChaosVDParticleDataVisualizer.h
@@ -0,0 +1,38 @@
+﻿// Copyright Epic Games, Inc. All Rights Reserved.
+
+#pragma once
+#include "ChaosVDDataVisualizerBase.h"
+#include "IChaosVDParticleVisualizationDataProvider.h"
+
+#include "ChaosVDParticleDataVisualizer.generated.h"
+
+struct FChaosVDParticleDataWrapper;
+
+UENUM(meta = (Bitflags, UseEnumValuesAsMaskValuesInEditor = "true"))
+enum class EChaosVDParticleDataVisualizationFlags : uint32
+{
+	None				= 0 UMETA(Hidden),
+	Velocity			= 1 << 0,
+	AngularVelocity		= 1 << 1,
+	Acceleration		= 1 << 2,
+	AngularAcceleration = 1 << 3,
+	LinearImpulse		= 1 << 4,
+	AngularImpulse		= 1 << 5,
+};
+ENUM_CLASS_FLAGS(EChaosVDParticleDataVisualizationFlags);
+
+class FChaosVDParticleDataVisualizer final : public FChaosVDDataVisualizerBase
+{
+public:
+	explicit FChaosVDParticleDataVisualizer(IChaosVDParticleVisualizationDataProvider& InProvider)
+		: FChaosVDDataVisualizerBase(), DataProvider(InProvider)
+	{
+	}
+
+	virtual void DrawVisualization(const FSceneView* View, FPrimitiveDrawInterface* PDI) override;
+
+	inline static const FStringView VisualizerID = TEXT("ParticleDataVisualizer");
+
+protected:
+	IChaosVDParticleVisualizationDataProvider& DataProvider;
+};
diff --git a/Engine/Plugins/Experimental/ChaosVD/Source/ChaosVD/Public/Visualizers/IChaosVDParticleVisualizationDataProvider.h b/Engine/Plugins/Experimental/ChaosVD/Source/ChaosVD/Public/Visualizers/IChaosVDParticleVisualizationDataProvider.h
new file mode 100644
--- /dev/null
+++ b/Engine/Plugins/Experimental/ChaosVD/Source/ChaosVD/Public/Visualizers/IChaosVDParticleVisualizationDataProvider.h
@@ -0,0 +1,30 @@
+﻿// Copyright Epic Games, Inc. All Rights Reserved.
+
+#pragma once
+#include "Math/Transform.h"
+#include "HAL/Platform.h"
+#include "Templates/SharedPointer.h"
+
+struct FChaosVDParticleDataWrapper;
+class FChaosVDScene;
+
+/** Context needed to be able to visualize data in the viewport */
+struct FChaosVDVisualizationContext
+{
+	FTransform SpaceTransform;
+	TWeakPtr<FChaosVDScene> CVDScene;
+	int32 SolverID = INDEX_NONE;
+};
+
+/** Interface to be used by any object that contains Particle Data that needs to be visualized.
+ * @note As we are still not settled on if we will stick with having Actors to represent each particle, this interface allows us to abstract the Particle data access
+ * instead of using directly AChaosVDParticleActor
+ */
+class IChaosVDParticleVisualizationDataProvider
+{
+public:
+	virtual ~IChaosVDParticleVisualizationDataProvider() = default;
+
+	virtual void GetVisualizationContext(FChaosVDVisualizationContext& OutVisualizationContext) {}
+	virtual const FChaosVDParticleDataWrapper* GetParticleData() { return nullptr; }
+};
diff --git a/Engine/Plugins/Experimental/ChaosVD/Source/ChaosVD/Public/Widgets/SChaosVDMainTab.h b/Engine/Plugins/Experimental/ChaosVD/Source/ChaosVD/Public/Widgets/SChaosVDMainTab.h
--- a/Engine/Plugins/Experimental/ChaosVD/Source/ChaosVD/Public/Widgets/SChaosVDMainTab.h
+++ b/Engine/Plugins/Experimental/ChaosVD/Source/ChaosVD/Public/Widgets/SChaosVDMainTab.h
@@ -6,6 +6,7 @@
 #include "Templates/SharedPointer.h"
 #include "Widgets/SCompoundWidget.h"
 
+class FChaosVDEditorSettingsTab;
 class FChaosVDSolversTracksTab;
 class FChaosVDEngine;
 class FChaosVDOutputLogTab;
@@ -37,11 +38,13 @@ private:
 
 	TSharedPtr<FChaosVDEngine> ChaosVDEngine;
 
+	// TODO Convert this to a map ID-> Tab
 	TSharedPtr<FChaosVDWorldOutlinerTab> WorldOutlinerTab;
 	TSharedPtr<FChaosVDObjectDetailsTab> ObjectDetailsTab;
 	TSharedPtr<FChaosVDPlaybackViewportTab> PlaybackViewportTab;
 	TSharedPtr<FChaosVDSolversTracksTab> SolversTracksTab;
 	TSharedPtr<FChaosVDOutputLogTab> OutputLogTab;
+	TSharedPtr<FChaosVDEditorSettingsTab> EditorSettingsTab;
 
 	TSharedPtr<FTabManager> TabManager;
 };
diff --git a/Engine/Plugins/Experimental/ChaosVD/Source/ChaosVD/Public/Widgets/SChaosVDSolverPlaybackControls.h b/Engine/Plugins/Experimental/ChaosVD/Source/ChaosVD/Public/Widgets/SChaosVDSolverPlaybackControls.h
--- a/Engine/Plugins/Experimental/ChaosVD/Source/ChaosVD/Public/Widgets/SChaosVDSolverPlaybackControls.h
+++ b/Engine/Plugins/Experimental/ChaosVD/Source/ChaosVD/Public/Widgets/SChaosVDSolverPlaybackControls.h
@@ -43,8 +43,11 @@ private:
 	void UpdateStepsWidgetForFrame(const FChaosVDPlaybackController& InCurrentPlaybackController, int32 FrameNumber, int32 StepNumber, EChaosVDStepsWidgetUpdateFlags OptionsFlags = EChaosVDStepsWidgetUpdateFlags::Default);
 	virtual void HandleControllerTrackFrameUpdated(TWeakPtr<FChaosVDPlaybackController> InController, const FChaosVDTrackInfo* UpdatedTrackInfo, FGuid InstigatorGuid) override;
 
+	void HandleLockStateChanged(bool NewIsLocked);
+
 	int32 SolverID = INDEX_NONE;
 	FString CurrentStepName;
 	TSharedPtr<SChaosVDTimelineWidget> FramesTimelineWidget;
 	TSharedPtr<SChaosVDTimelineWidget> StepsTimelineWidget;
+	bool bStepsLocked = false;
 };
diff --git a/Engine/Plugins/Experimental/ChaosVD/Source/ChaosVD/Public/Widgets/SChaosVDTimelineWidget.h b/Engine/Plugins/Experimental/ChaosVD/Source/ChaosVD/Public/Widgets/SChaosVDTimelineWidget.h
--- a/Engine/Plugins/Experimental/ChaosVD/Source/ChaosVD/Public/Widgets/SChaosVDTimelineWidget.h
+++ b/Engine/Plugins/Experimental/ChaosVD/Source/ChaosVD/Public/Widgets/SChaosVDTimelineWidget.h
@@ -10,6 +10,7 @@ class FReply;
 struct FSlateBrush;
 
 DECLARE_DELEGATE_OneParam(FChaosVDFrameChangedDelegate, int32)
+DECLARE_DELEGATE_OneParam(FChaosVDFrameLockStateDelegate, bool)
 
 enum class EChaosVDSetTimelineFrameFlags
 {
@@ -28,7 +29,10 @@ public:
 	SLATE_BEGIN_ARGS( SChaosVDTimelineWidget ){}
 		SLATE_ARGUMENT(int32, MaxFrames)
 		SLATE_EVENT(FChaosVDFrameChangedDelegate, OnFrameChanged)
-		SLATE_ARGUMENT(bool, HidePlayStopButtons)
+		SLATE_EVENT(FChaosVDFrameLockStateDelegate, OnFrameLockStateChanged)
+		SLATE_ATTRIBUTE( bool, HidePlayStopButtons)
+		SLATE_ATTRIBUTE( bool, HideNextPrevButtons)
+		SLATE_ATTRIBUTE( bool, HideLockButton)
 	SLATE_END_ARGS()
 
 	void Construct(const FArguments& InArgs);
@@ -37,21 +41,30 @@ public:
 
 	void SetCurrentTimelineFrame(float FrameNumber, EChaosVDSetTimelineFrameFlags Options = EChaosVDSetTimelineFrameFlags::BroadcastChange);
 
+	void SetIsLocked(bool NewIsLocked) { bIsLocked = NewIsLocked; }
+
+	bool IsUnlocked() const { return !bIsLocked; }
+
 	/** Brings back the state of the timeline to its original state*/
 	void ResetTimeline();
 
 	/** Called when a new frame is manually selected or auto-updated during playback */
 	FChaosVDFrameChangedDelegate& OnFrameChanged() { return FrameChangedDelegate; }
+	
+	/** Called when this timeline is locked or unlocked */
+	FChaosVDFrameLockStateDelegate& OnFrameLockStateChanged() { return FrameLockedDelegate; }
 
 protected:
 
-
 	FReply  Play();
 	FReply  Stop();
 	FReply  Next();
 	FReply  Prev();
 
+	FReply ToggleLockState();
+
 	const FSlateBrush* GetPlayOrPauseIcon() const;
+	const FSlateBrush* GetLockStateIcon() const;
 
 	virtual void Tick(const FGeometry& AllottedGeometry, const double InCurrentTime, const float InDeltaTime) override;
 
@@ -65,4 +78,7 @@ protected:
 	float CurrentPlaybackTime = 0.0f;
 
 	FChaosVDFrameChangedDelegate FrameChangedDelegate;
+	FChaosVDFrameLockStateDelegate FrameLockedDelegate;
+
+	bool bIsLocked = false;
 };
diff --git a/Engine/Source/Runtime/Experimental/Chaos/Private/Chaos/PBDRigidsEvolutionGBF.cpp b/Engine/Source/Runtime/Experimental/Chaos/Private/Chaos/PBDRigidsEvolutionGBF.cpp
--- a/Engine/Source/Runtime/Experimental/Chaos/Private/Chaos/PBDRigidsEvolutionGBF.cpp
+++ b/Engine/Source/Runtime/Experimental/Chaos/Private/Chaos/PBDRigidsEvolutionGBF.cpp
@@ -574,7 +574,10 @@ void FPBDRigidsEvolutionGBF::AdvanceOneTimeStepImpl(const FReal Dt, const FSubSt
 	}
 
 	{
-		CVD_SCOPE_TRACE_SOLVER_STEP(TEXT("Evolution End"))
+		CVD_SCOPE_TRACE_SOLVER_STEP(TEXT("Evolution End"));
+
+		CVD_TRACE_MID_PHASES_FROM_COLLISION_CONSTRAINTS(GetCollisionConstraints());
+
 		CVD_TRACE_PARTICLES_SOA(Particles);
 	}
 	
diff --git a/Engine/Source/Runtime/Experimental/Chaos/Private/ChaosVisualDebugger/ChaosVDDataWrapperUtils.cpp b/Engine/Source/Runtime/Experimental/Chaos/Private/ChaosVisualDebugger/ChaosVDDataWrapperUtils.cpp
--- a/Engine/Source/Runtime/Experimental/Chaos/Private/ChaosVisualDebugger/ChaosVDDataWrapperUtils.cpp
+++ b/Engine/Source/Runtime/Experimental/Chaos/Private/ChaosVisualDebugger/ChaosVDDataWrapperUtils.cpp
@@ -100,6 +100,8 @@ FChaosVDConstraint FChaosVDDataWrapperUtils::BuildConstraintDataWrapperFromConst
 
 	Algo::Copy(InConstraint.ShapeWorldTransforms, WrappedConstraintData.ShapeWorldTransforms);
 
+	Algo::Copy(InConstraint.ImplicitTransform, WrappedConstraintData.ImplicitTransforms);
+
 	WrappedConstraintData.CollisionMargins = TArray(InConstraint.CollisionMargins, std::size(InConstraint.CollisionMargins));
 	WrappedConstraintData.LastShapeWorldPositionDelta = FVector(InConstraint.LastShapeWorldPositionDelta);
 	WrappedConstraintData.LastShapeWorldRotationDelta = FQuat(InConstraint.LastShapeWorldRotationDelta);
diff --git a/Engine/Source/Runtime/Experimental/Chaos/Private/ChaosVisualDebugger/ChaosVisualDebuggerTrace.cpp b/Engine/Source/Runtime/Experimental/Chaos/Private/ChaosVisualDebugger/ChaosVisualDebuggerTrace.cpp
--- a/Engine/Source/Runtime/Experimental/Chaos/Private/ChaosVisualDebugger/ChaosVisualDebuggerTrace.cpp
+++ b/Engine/Source/Runtime/Experimental/Chaos/Private/ChaosVisualDebugger/ChaosVisualDebuggerTrace.cpp
@@ -6,6 +6,7 @@
 #include "Chaos/Framework/PhysicsSolverBase.h"
 #include "Chaos/ImplicitObject.h"
 #include "Chaos/ParticleHandle.h"
+#include "Chaos/PBDCollisionConstraints.h"
 #include "Chaos/PBDRigidsSOAs.h"
 #include "ChaosVisualDebugger/ChaosVDDataWrapperUtils.h"
 #include "Compression/OodleDataCompressionUtil.h"
@@ -267,6 +268,29 @@ void FChaosVisualDebuggerTrace::TraceMidPhase(const Chaos::FParticlePairMidPhase
 	TraceBinaryData(CVDBuffer.BufferRef, FChaosVDParticlePairMidPhase::WrapperTypeName);
 }
 
+
+void FChaosVisualDebuggerTrace::TraceMidPhasesFromCollisionConstraints(Chaos::FPBDCollisionConstraints& InCollisionConstraints)
+{
+	if (!IsTracing())
+	{
+		return;
+	}
+
+	const FChaosVDContext* CVDContextData = FChaosVDThreadContext::Get().GetCurrentContext();
+	if (!ensure(CVDContextData))
+	{
+		return;
+	}
+
+	InCollisionConstraints.GetConstraintAllocator().VisitMidPhases([CopyContext = *CVDContextData](const Chaos::FParticlePairMidPhase& MidPhase) -> Chaos::ECollisionVisitorResult
+	{
+		CVD_SCOPE_CONTEXT(CopyContext)
+		CVD_TRACE_MID_PHASE(&MidPhase);
+		return Chaos::ECollisionVisitorResult::Continue;
+	});
+}
+
+
 void FChaosVisualDebuggerTrace::TraceCollisionConstraint(const Chaos::FPBDCollisionConstraint* CollisionConstraint)
 {
 	if (!IsTracing())
diff --git a/Engine/Source/Runtime/Experimental/Chaos/Public/ChaosVisualDebugger/ChaosVisualDebuggerTrace.h b/Engine/Source/Runtime/Experimental/Chaos/Public/ChaosVisualDebugger/ChaosVisualDebuggerTrace.h
--- a/Engine/Source/Runtime/Experimental/Chaos/Public/ChaosVisualDebugger/ChaosVisualDebuggerTrace.h
+++ b/Engine/Source/Runtime/Experimental/Chaos/Public/ChaosVisualDebugger/ChaosVisualDebuggerTrace.h
@@ -78,6 +78,11 @@
 #ifndef CVD_TRACE_PARTICLES_VIEW
 	#define CVD_TRACE_PARTICLES_VIEW(ParticleHandlesView)
 #endif
+
+#ifndef CVD_TRACE_MID_PHASES_FROM_COLLISION_CONSTRAINTS
+	#define CVD_TRACE_MID_PHASES_FROM_COLLISION_CONSTRAINTS(CollisionConstraints)
+#endif
+
 #else
 
 #include "ChaosVDRuntimeModule.h"
@@ -268,10 +273,16 @@ UE_TRACE_EVENT_END()
 	FChaosVisualDebuggerTrace::TraceCollisionConstraintView(ConstraintView);
 #endif
 
+#ifndef CVD_TRACE_MID_PHASES_FROM_COLLISION_CONSTRAINTS
+	#define CVD_TRACE_MID_PHASES_FROM_COLLISION_CONSTRAINTS(CollisionConstraints) \
+	FChaosVisualDebuggerTrace::TraceMidPhasesFromCollisionConstraints(CollisionConstraints);
+#endif
+
 struct FChaosVDContext;
 
 namespace Chaos
 {
+	class FPBDCollisionConstraints;
 	class FPBDRigidsSOAs;
 	class FImplicitObject;
 	class FPhysicsSolverBase;
@@ -327,6 +338,9 @@ public:
 	/** Traces a Particle pair MidPhase as binary data */
 	static CHAOS_API void TraceMidPhase(const Chaos::FParticlePairMidPhase* MidPhase);
 
+	/** Traces a Particle pair MidPhase as binary data from a provided CollisionConstraints object */
+	static CHAOS_API void TraceMidPhasesFromCollisionConstraints(Chaos::FPBDCollisionConstraints& InCollisionConstraints);
+
 	/** Traces a Particle pair MidPhase as binary data */
 	static CHAOS_API void TraceCollisionConstraint(const Chaos::FPBDCollisionConstraint* CollisionConstraint);
 
diff --git a/Engine/Source/Runtime/Experimental/ChaosVisualDebugger/Private/DataWrappers/ChaosVDCollisionDataWrappers.cpp b/Engine/Source/Runtime/Experimental/ChaosVisualDebugger/Private/DataWrappers/ChaosVDCollisionDataWrappers.cpp
--- a/Engine/Source/Runtime/Experimental/ChaosVisualDebugger/Private/DataWrappers/ChaosVDCollisionDataWrappers.cpp
+++ b/Engine/Source/Runtime/Experimental/ChaosVisualDebugger/Private/DataWrappers/ChaosVDCollisionDataWrappers.cpp
@@ -54,6 +54,7 @@ bool FChaosVDConstraint::Serialize(FArchive& Ar)
 	Ar << AccumulatedImpulse;
 	Ar << ShapesType;
 	Ar << ShapeWorldTransforms;
+	Ar << ImplicitTransforms;
 	Ar << CullDistance;
 	Ar << CollisionMargins;	
 	Ar << CollisionTolerance;	
diff --git a/Engine/Source/Runtime/Experimental/ChaosVisualDebugger/Public/DataWrappers/ChaosVDCollisionDataWrappers.h b/Engine/Source/Runtime/Experimental/ChaosVisualDebugger/Public/DataWrappers/ChaosVDCollisionDataWrappers.h
--- a/Engine/Source/Runtime/Experimental/ChaosVisualDebugger/Public/DataWrappers/ChaosVDCollisionDataWrappers.h
+++ b/Engine/Source/Runtime/Experimental/ChaosVisualDebugger/Public/DataWrappers/ChaosVDCollisionDataWrappers.h
@@ -190,6 +190,9 @@ struct CHAOSVDRUNTIME_API FChaosVDConstraint
 	
 	UPROPERTY(EditAnywhere, Category=Collision)
 	TArray<FTransform> ShapeWorldTransforms;
+
+	UPROPERTY(EditAnywhere, Category=Collision)
+	TArray<FTransform> ImplicitTransforms;
 	
 	UPROPERTY(EditAnywhere, Category=Collision)
 	float CullDistance;
