[Chaos VD] Fix static analysis waring in ChaosVDPlaybackController.cpp by adding missing null check

diff --git a/Engine/Plugins/Experimental/ChaosVD/Source/ChaosVD/Private/ChaosVDPlaybackController.cpp b/Engine/Plugins/Experimental/ChaosVD/Source/ChaosVD/Private/ChaosVDPlaybackController.cpp
--- a/Engine/Plugins/Experimental/ChaosVD/Source/ChaosVD/Private/ChaosVDPlaybackController.cpp
+++ b/Engine/Plugins/Experimental/ChaosVD/Source/ChaosVD/Private/ChaosVDPlaybackController.cpp
@@ -110,12 +110,18 @@ void FChaosVDPlaybackController::PlayFromClosestKeyFrame(const int32 InTrackID,
 
 	for (int32 CurrentFrameNumber = KeyFrameNumber; CurrentFrameNumber < FrameNumber; CurrentFrameNumber++)
 	{				
-		const FChaosVDSolverFrameData* SolverFrameData = LoadedRecording->GetSolverFrameData(InTrackID, CurrentFrameNumber);
-		const int32 LastStepNumber = SolverFrameData->SolverSteps.Num() - 1;
+		if (const FChaosVDSolverFrameData* SolverFrameData = LoadedRecording->GetSolverFrameData(InTrackID, CurrentFrameNumber))
+		{
+			const int32 LastStepNumber = SolverFrameData->SolverSteps.Num() - 1;
 
-		if (SolverFrameData && ensure(SolverFrameData->SolverSteps.IsValidIndex(LastStepNumber)))
+			if (SolverFrameData && ensure(SolverFrameData->SolverSteps.IsValidIndex(LastStepNumber)))
+			{
+				InSceneToControl.UpdateFromRecordedStepData(InTrackID, SolverFrameData->DebugName, SolverFrameData->SolverSteps[LastStepNumber], *SolverFrameData);
+			}
+		}
+		else
 		{
-			InSceneToControl.UpdateFromRecordedStepData(InTrackID, SolverFrameData->DebugName, SolverFrameData->SolverSteps[LastStepNumber], *SolverFrameData);
+			UE_LOG(LogChaosVDEditor, Error, TEXT("[%s] Failed to read solver frame data for frame [%d] in track [%d]"), ANSI_TO_TCHAR(__FUNCTION__), CurrentFrameNumber, InTrackID);
 		}
 	}
 }
