[ChoasVD] Fix static analysis warning on FChaosVDPlaybackController

diff --git a/Engine/Plugins/Experimental/ChaosVD/Source/ChaosVD/Private/ChaosVDPlaybackController.cpp b/Engine/Plugins/Experimental/ChaosVD/Source/ChaosVD/Private/ChaosVDPlaybackController.cpp
--- a/Engine/Plugins/Experimental/ChaosVD/Source/ChaosVD/Private/ChaosVDPlaybackController.cpp
+++ b/Engine/Plugins/Experimental/ChaosVD/Source/ChaosVD/Private/ChaosVDPlaybackController.cpp
@@ -114,7 +114,7 @@ void FChaosVDPlaybackController::PlayFromClosestKeyFrame(const int32 InTrackID,
 		{
 			const int32 LastStepNumber = SolverFrameData->SolverSteps.Num() - 1;
 
-			if (SolverFrameData && ensure(SolverFrameData->SolverSteps.IsValidIndex(LastStepNumber)))
+			if (ensure(SolverFrameData->SolverSteps.IsValidIndex(LastStepNumber)))
 			{
 				InSceneToControl.UpdateFromRecordedStepData(InTrackID, SolverFrameData->DebugName, SolverFrameData->SolverSteps[LastStepNumber], *SolverFrameData);
 			}
@@ -149,13 +149,14 @@ void FChaosVDPlaybackController::GoToRecordedSolverStep(const int32 InTrackID, c
 				if (TSharedPtr<FChaosVDTrackInfo>* TrackInfo = TrackInfoByID->Find(InTrackID))
 				{
 					CurrentTrackInfo = *TrackInfo;
-					if (!ensure(CurrentTrackInfo.IsValid()))
-					{
-						UE_LOG(LogChaosVDEditor, Warning, TEXT("[%s] Track info for track ID [%d]. We can't continue..."), ANSI_TO_TCHAR(__FUNCTION__), InTrackID);
-						return;
-					}
 				}
 			}
+			
+			if (!ensure(CurrentTrackInfo.IsValid()))
+			{
+				UE_LOG(LogChaosVDEditor, Warning, TEXT("[%s] Track info for track ID [%d]. We can't continue..."), ANSI_TO_TCHAR(__FUNCTION__), InTrackID);
+				return;
+			}
 
 			const int32 FrameDiff = FrameNumber - CurrentTrackInfo->CurrentFrame;
 			constexpr int32 FrameDriftTolerance = 1;
@@ -178,7 +179,7 @@ void FChaosVDPlaybackController::GoToRecordedSolverStep(const int32 InTrackID, c
 					}
 					else
 					{
-						UE_LOG(LogChaosVDEditor, Error, TEXT("[%s] Tried to scrub to an invalid step | Step Number [%d] | Available Steps [%d]..."), ANSI_TO_TCHAR(__FUNCTION__), CurrentTrackInfo->LockedOnStep, SolverFrameData->SolverSteps.Num());
+						UE_LOG(LogChaosVDEditor, Error, TEXT("[%s] Tried to scrub to an invalid step | Step Number [%d] ..."), ANSI_TO_TCHAR(__FUNCTION__), CurrentTrackInfo->LockedOnStep);
 						return;
 					}
 				}
@@ -191,7 +192,7 @@ void FChaosVDPlaybackController::GoToRecordedSolverStep(const int32 InTrackID, c
 				}
 				else
 				{
-					UE_LOG(LogChaosVDEditor, Error, TEXT("[%s] Tried to scrub to an invalid step | Step Number [%d] | Available Steps [%d]..."), ANSI_TO_TCHAR(__FUNCTION__), CurrentTrackInfo->LockedOnStep, SolverFrameData->SolverSteps.Num());
+					UE_LOG(LogChaosVDEditor, Error, TEXT("[%s] Tried to scrub to an invalid step | Step Number [%d] ..."), ANSI_TO_TCHAR(__FUNCTION__), CurrentTrackInfo->LockedOnStep);
 				}
 			}
 			
