[ChaosVD] Fixed synchronization issues between selection in viewport and scene outliner.

Before this change, when you selected something in the viewport, the scene outliner didn't auto-scroll to it or sometimes it crashed if the item was not in the outliner tree (because it was filtered out).

diff --git a/Engine/Plugins/Experimental/ChaosVD/Source/ChaosVD/Private/ChaosVDWorldOutlinerMode.cpp b/Engine/Plugins/Experimental/ChaosVD/Source/ChaosVD/Private/ChaosVDWorldOutlinerMode.cpp
--- a/Engine/Plugins/Experimental/ChaosVD/Source/ChaosVD/Private/ChaosVDWorldOutlinerMode.cpp
+++ b/Engine/Plugins/Experimental/ChaosVD/Source/ChaosVD/Private/ChaosVDWorldOutlinerMode.cpp
@@ -3,6 +3,7 @@
 #include "ChaosVDWorldOutlinerMode.h"
 
 #include "ActorTreeItem.h"
+#include "ChaosVDModule.h"
 #include "ChaosVDScene.h"
 #include "Elements/Framework/TypedElementSelectionSet.h"
 
@@ -91,8 +92,16 @@ void FChaosVDWorldOutlinerMode::HandlePostSelectionChange(const UTypedElementSel
 	{
 		// We don't support multi selection yet
 		ensure(SelectedActors.Num() == 1);
-
-		FSceneOutlinerTreeItemPtr ObjectItem = SceneOutliner->GetTreeItem(SelectedActors[0]);
-		SceneOutliner->SetItemSelection(ObjectItem, true, ESelectInfo::OnMouseClick);
+		AActor* SelectedActor = SelectedActors[0];
+		
+		if (FSceneOutlinerTreeItemPtr TreeItem = SceneOutliner->GetTreeItem(SelectedActor, false))
+		{
+			SceneOutliner->ScrollItemIntoView(TreeItem);
+			SceneOutliner->SetItemSelection(TreeItem, true, ESelectInfo::OnMouseClick);
+		}
+		else
+		{
+			UE_LOG(LogChaosVDEditor, Verbose, TEXT("Selected actor is not in the outliner. It might be filtered out"))	
+		}
 	}
 }
