[Chaos VD] Making all delegates from ChaosVDPlaybackController trigger on the Game Thread.

During Trace analysis, we fill up the ChaosVDRecording structure. At the same time, as soon a new frame data is available, the playback controller will process it (allowing to play a recording before it finished loading).

That is useful to reduce waiting times, but is also a requirement for remote debugging using trace as in that mode we will be receiving frames constantly as we play.

There are a few issues with this and the current implementation:

- Trace Analysis updating the recording, and the playback controller processing it will trigger an UI update, and when that is done outside the game thread, it will crash.
- Playback is done in the GameThread, to do so we get a reference to the solver frame data from the recording we want to play. if the container where that data lives is re-allocated while Trace Analysis is adding data, the tool will crash.

This CL ensure all events reporting changes used by UI code are triggered in the Game Thread instead of the Trace Analysis thread.

In addition a new RW lock was added to the recording structure (plus AssumesLocked versions of existing methods), to protect the data from being accessed during playback when we are writing a new frame.

The Assumeslocked version of some methods are required because we need to lock the data between multiple scopes, and in some scenarios in a  recursive manner. An alternative would to use a recursive mutex.

diff --git a/Engine/Plugins/Experimental/ChaosVD/Source/ChaosVD/Private/ChaosVDPlaybackController.cpp b/Engine/Plugins/Experimental/ChaosVD/Source/ChaosVD/Private/ChaosVDPlaybackController.cpp
--- a/Engine/Plugins/Experimental/ChaosVD/Source/ChaosVD/Private/ChaosVDPlaybackController.cpp
+++ b/Engine/Plugins/Experimental/ChaosVD/Source/ChaosVD/Private/ChaosVDPlaybackController.cpp
@@ -49,13 +49,6 @@ bool FChaosVDPlaybackController::LoadChaosVDRecordingFromTraceSession(const FStr
 
 	HandleCurrentRecordingUpdated();
 
-	for (TMap<int32, TArray<FChaosVDSolverFrameData>>::TConstIterator SolversIterator = LoadedRecording->GetAvailableSolvers().CreateConstIterator(); SolversIterator; ++SolversIterator)
-	{
-		constexpr int32 FrameNumber = 0;
-		constexpr int32 StepNumber = 0;
-		GoToTrackFrame(IChaosVDPlaybackControllerInstigator::InvalidGuid, EChaosVDTrackType::Solver, SolversIterator->Key, FrameNumber, StepNumber);
-	}
-
 	LoadedRecording->OnGeometryDataLoaded().AddLambda([this](const TSharedPtr<const Chaos::FImplicitObject>& NewGeometry, const uint32 GeometryID)
 	{
 		if (const TSharedPtr<FChaosVDScene> ScenePtr = SceneToControl.Pin())
@@ -69,13 +62,15 @@ bool FChaosVDPlaybackController::LoadChaosVDRecordingFromTraceSession(const FStr
 		ScenePtr->LoadedRecording = LoadedRecording;
 	}
 	
-	OnDataUpdated().Broadcast(AsWeak());
+	bHasPendingGTUpdateBroadcast = true;
 
 	return true;
 }
 
 void FChaosVDPlaybackController::UnloadCurrentRecording(EChaosVDUnloadRecordingFlags UnloadOptions)
 {
+	TrackInfoUpdateGTQueue.Empty();
+	
 	if (LoadedRecording.IsValid())
 	{
 		LoadedRecording->OnRecordingUpdated().RemoveAll(this);
@@ -94,14 +89,15 @@ void FChaosVDPlaybackController::UnloadCurrentRecording(EChaosVDUnloadRecordingF
 
 	if (EnumHasAnyFlags(UnloadOptions, EChaosVDUnloadRecordingFlags::BroadcastChanges))
 	{
-		const TWeakPtr<FChaosVDPlaybackController> ThisWeakPtr = DoesSharedInstanceExist() ? AsWeak() : nullptr;
-		OnDataUpdated().Broadcast(ThisWeakPtr);
+		bHasPendingGTUpdateBroadcast = true;
 	}
+
+	bPlayedFirsFrame = false;
 }
 
-void FChaosVDPlaybackController::PlayFromClosestKeyFrame(const int32 InTrackID, const int32 FrameNumber, FChaosVDScene& InSceneToControl)
+void FChaosVDPlaybackController::PlayFromClosestKeyFrame_AssumesLocked(const int32 InTrackID, const int32 FrameNumber, FChaosVDScene& InSceneToControl)
 {
-	const int32 KeyFrameNumber = LoadedRecording->FindFirstSolverKeyFrameNumberFromFrame(InTrackID, FrameNumber);
+	const int32 KeyFrameNumber = LoadedRecording->FindFirstSolverKeyFrameNumberFromFrame_AssumesLocked(InTrackID, FrameNumber);
 
 	if (!ensure(KeyFrameNumber >= 0))
 	{
@@ -109,98 +105,105 @@ void FChaosVDPlaybackController::PlayFromClosestKeyFrame(const int32 InTrackID,
 	}
 
 	for (int32 CurrentFrameNumber = KeyFrameNumber; CurrentFrameNumber < FrameNumber; CurrentFrameNumber++)
-	{				
-		if (const FChaosVDSolverFrameData* SolverFrameData = LoadedRecording->GetSolverFrameData(InTrackID, CurrentFrameNumber))
+	{
+		if (const FChaosVDSolverFrameData* SolverFrameData = LoadedRecording->GetSolverFrameData_AssumesLocked(InTrackID, CurrentFrameNumber))
 		{
 			const int32 LastStepNumber = SolverFrameData->SolverSteps.Num() - 1;
 
-			if (ensure(SolverFrameData->SolverSteps.IsValidIndex(LastStepNumber)))
+			if (SolverFrameData->SolverSteps.IsValidIndex(LastStepNumber))
 			{
 				InSceneToControl.UpdateFromRecordedStepData(InTrackID, SolverFrameData->DebugName, SolverFrameData->SolverSteps[LastStepNumber], *SolverFrameData);
 			}
 		}
 		else
 		{
-			UE_LOG(LogChaosVDEditor, Error, TEXT("[%s] Failed to read solver frame data for frame [%d] in track [%d]"), ANSI_TO_TCHAR(__FUNCTION__), CurrentFrameNumber, InTrackID);
+			// This is common if we stop PIE, change worlds, and PIE again without stopping the recording
+			UE_LOG(LogChaosVDEditor, Verbose, TEXT("[%s] Failed to read solver frame data for frame [%d] in track [%d]"), ANSI_TO_TCHAR(__FUNCTION__), CurrentFrameNumber, InTrackID);
 		}
 	}
 }
 
-void FChaosVDPlaybackController::GoToRecordedSolverStep(const int32 InTrackID, const int32 FrameNumber, const int32 Step, FGuid InstigatorID)
+void FChaosVDPlaybackController::EnqueueTrackInfoUpdate(const FChaosVDTrackInfo& InTrackInfo, FGuid InstigatorID)
+{
+	// This will be used in the Game Thread on the first tick after this was added, so we need to make a copy of the state right now
+	FChaosVDQueuedTrackInfoUpdate InfoUpdate;
+	InfoUpdate.TrackInfo = InTrackInfo;
+	InfoUpdate.InstigatorID = InstigatorID;
+
+	// TODO: FChaosVDTrackInfo has the track name as an string which will be copied as well.
+	// We should move the names to another structure as these don't change often.
+	// It wasn't an issue for now, but now this copy might have an impact due to the number of updates there could be done.
+	
+	TrackInfoUpdateGTQueue.Enqueue(InfoUpdate);
+}
+
+void FChaosVDPlaybackController::GoToRecordedSolverStep_AssumesLocked(const int32 InTrackID, const int32 FrameNumber, const int32 Step, FGuid InstigatorID)
 {
 	if (const TSharedPtr<FChaosVDScene> SceneToControlSharedPtr = SceneToControl.Pin())
 	{
 		if (ensure(LoadedRecording.IsValid()))
 		{
-			const TSharedPtr<const TraceServices::IAnalysisSession> TraceSession = FChaosVDModule::Get().GetTraceManager()->GetSession(LoadedRecording->SessionName);
-			if (!ensure(TraceSession))
-			{
-				return;	
-			}
-
-			// TODO: This will not cover a future case where the recording is not own/populated by Trace
-			// If in the future we decide implement the CVD format standalone with streaming support again,
-			// we will need to add a lock to the file. A feature that might need that is Recording Clips
-			TraceServices::FAnalysisSessionReadScope SessionReadScope(*TraceSession);
-
 			TSharedPtr<FChaosVDTrackInfo> CurrentTrackInfo;
 			if (TrackInfoByIDMap* TrackInfoByID = TrackInfoPerType.Find(EChaosVDTrackType::Solver))
 			{
 				if (TSharedPtr<FChaosVDTrackInfo>* TrackInfo = TrackInfoByID->Find(InTrackID))
 				{
 					CurrentTrackInfo = *TrackInfo;
+					if (!ensure(CurrentTrackInfo.IsValid()))
+					{
+						UE_LOG(LogChaosVDEditor, Warning, TEXT("[%s] Track info for track ID [%d]. We can't continue..."), ANSI_TO_TCHAR(__FUNCTION__), InTrackID);
+						return;
+					}
 				}
 			}
-			
-			if (!ensure(CurrentTrackInfo.IsValid()))
-			{
-				UE_LOG(LogChaosVDEditor, Warning, TEXT("[%s] Track info for track ID [%d]. We can't continue..."), ANSI_TO_TCHAR(__FUNCTION__), InTrackID);
-				return;
-			}
 
 			const int32 FrameDiff = FrameNumber - CurrentTrackInfo->CurrentFrame;
 			constexpr int32 FrameDriftTolerance = 1;
 			if (FMath::Abs(FrameDiff) > FrameDriftTolerance || CurrentTrackInfo->CurrentFrame == 0)
 			{
 				// As Frames are recorded as delta, we need to make sure of playing back all the deltas since the closest keyframe
-				PlayFromClosestKeyFrame(InTrackID, FrameNumber, *SceneToControlSharedPtr.Get());
+				PlayFromClosestKeyFrame_AssumesLocked(InTrackID, FrameNumber, *SceneToControlSharedPtr.Get());
 			}
-
-			const FChaosVDSolverFrameData* SolverFrameData = LoadedRecording->GetSolverFrameData(InTrackID, FrameNumber);
-			if (CurrentTrackInfo->LockedOnStep != INDEX_NONE)
+			
 			{
-				// If this track is locked to a specific step, we need to play back the previous steps on the current frame, because not all steps capture the same data.
-				// For example, Particles positions are fully captured in the first sub-step and the last one
-				for (int32 StepIndex = 0; StepIndex <= CurrentTrackInfo->LockedOnStep; StepIndex++)
+				const FChaosVDSolverFrameData* SolverFrameData = LoadedRecording->GetSolverFrameData_AssumesLocked(InTrackID, FrameNumber);
+				if (CurrentTrackInfo->LockedOnStep != INDEX_NONE)
 				{
-					if (SolverFrameData && ensure(SolverFrameData->SolverSteps.IsValidIndex(CurrentTrackInfo->LockedOnStep)))
+					// If this track is locked to a specific step, we need to play back the previous steps on the current frame, because not all steps capture the same data.
+					// For example, Particles positions are fully captured in the first sub-step and the last one
+					for (int32 StepIndex = 0; StepIndex <= CurrentTrackInfo->LockedOnStep; StepIndex++)
 					{
-						SceneToControlSharedPtr->UpdateFromRecordedStepData(InTrackID, SolverFrameData->DebugName, SolverFrameData->SolverSteps[StepIndex], *SolverFrameData);
+						if (SolverFrameData && SolverFrameData->SolverSteps.IsValidIndex(CurrentTrackInfo->LockedOnStep))
+						{
+							SceneToControlSharedPtr->UpdateFromRecordedStepData(InTrackID, SolverFrameData->DebugName, SolverFrameData->SolverSteps[StepIndex], *SolverFrameData);
+						}
+						else
+						{
+							// This is common if we stop PIE, change worlds, and PIE again without stopping the recording
+							UE_LOG(LogChaosVDEditor, Verbose, TEXT("[%s] Tried to scrub to an invalid step | Step Number [%d] ..."), ANSI_TO_TCHAR(__FUNCTION__), CurrentTrackInfo->LockedOnStep);
+							return;
+						}
 					}
-					else
-					{
-						UE_LOG(LogChaosVDEditor, Error, TEXT("[%s] Tried to scrub to an invalid step | Step Number [%d] ..."), ANSI_TO_TCHAR(__FUNCTION__), CurrentTrackInfo->LockedOnStep);
-						return;
-					}
-				}
-			}
-			else
-			{
-				if (SolverFrameData && ensure(SolverFrameData->SolverSteps.IsValidIndex(Step)))
-				{
-					SceneToControlSharedPtr->UpdateFromRecordedStepData(InTrackID, SolverFrameData->DebugName, SolverFrameData->SolverSteps[Step], *SolverFrameData);
 				}
 				else
 				{
-					UE_LOG(LogChaosVDEditor, Error, TEXT("[%s] Tried to scrub to an invalid step | Step Number [%d] ..."), ANSI_TO_TCHAR(__FUNCTION__), CurrentTrackInfo->LockedOnStep);
+					if (SolverFrameData && SolverFrameData->SolverSteps.IsValidIndex(Step))
+					{
+						SceneToControlSharedPtr->UpdateFromRecordedStepData(InTrackID, SolverFrameData->DebugName, SolverFrameData->SolverSteps[Step], *SolverFrameData);
+					}
+					else
+					{
+						// This is common if we stop PIE, change worlds, and PIE again without stopping the recording
+						UE_LOG(LogChaosVDEditor, Verbose, TEXT("[%s] Tried to scrub to an invalid step | Step Number [%d] ..."), ANSI_TO_TCHAR(__FUNCTION__), CurrentTrackInfo->LockedOnStep);
+					}
 				}
 			}
 			
 
 			CurrentTrackInfo->CurrentFrame = FrameNumber;
 			CurrentTrackInfo->CurrentStep = Step;
-
-			OnTrackFrameUpdated().Broadcast(AsWeak(), CurrentTrackInfo.Get(), InstigatorID);
+			
+			EnqueueTrackInfoUpdate(*CurrentTrackInfo.Get(), InstigatorID);
 		}
 	}
 	else
@@ -209,18 +212,18 @@ void FChaosVDPlaybackController::GoToRecordedSolverStep(const int32 InTrackID, c
 	}
 }
 
-void FChaosVDPlaybackController::GoToRecordedGameFrame(const int32 FrameNumber, FGuid InstigatorID)
+void FChaosVDPlaybackController::GoToRecordedGameFrame_AssumesLocked(const int32 FrameNumber, FGuid InstigatorID)
 {
 	if (const TSharedPtr<FChaosVDScene> SceneToControlSharedPtr = SceneToControl.Pin())
 	{
 		if (ensure(LoadedRecording.IsValid()))
 		{
-			if (const FChaosVDGameFrameData* FrameData = LoadedRecording->GetGameFrameData(FrameNumber))
+			if (const FChaosVDGameFrameData* FrameData = LoadedRecording->GetGameFrameData_AssumesLocked(FrameNumber))
 			{
 				if (TrackInfoByIDMap* TrackInfoByID = TrackInfoPerType.Find(EChaosVDTrackType::Game))
 				{
 					TArray<int32> AvailableSolversID;
-					LoadedRecording->GetAvailableSolverIDsAtGameFrameNumber(FrameNumber, AvailableSolversID);
+					LoadedRecording->GetAvailableSolverIDsAtGameFrameNumber_AssumesLocked(FrameNumber, AvailableSolversID);
 
 					SceneToControlSharedPtr->HandleEnterNewGameFrame(FrameNumber, AvailableSolversID);
 
@@ -228,18 +231,18 @@ void FChaosVDPlaybackController::GoToRecordedGameFrame(const int32 FrameNumber,
 					{
 						// When Scrubbing the timeline by game frames instead of solvers, try to go to the first solver frame on the first platform cycle of the game frame.
 						// Game Frames are not in sync with Solver Frames and Solver steps.
-						const int32 SolverFrameNumber = LoadedRecording->GetLowestSolverFrameNumberAtCycle(SolverID, FrameData->FirstCycle);
+						const int32 SolverFrameNumber = LoadedRecording->GetLowestSolverFrameNumberAtCycle_AssumesLocked(SolverID, FrameData->FirstCycle);
 	
-						const int32 StepNumber = GetTrackLastStepAtFrame(EChaosVDTrackType::Solver, SolverID,SolverFrameNumber);
+						const int32 StepNumber = GetTrackLastStepAtFrame(EChaosVDTrackType::Solver, SolverID, SolverFrameNumber);
 
-						GoToTrackFrame(InstigatorID, EChaosVDTrackType::Solver, SolverID, SolverFrameNumber, StepNumber);
+						GoToTrackFrame_AssumesLocked(InstigatorID, EChaosVDTrackType::Solver, SolverID, SolverFrameNumber, StepNumber);
 					}
 
 					TSharedPtr<FChaosVDTrackInfo>* TrackInfoPtrPtr = TrackInfoByID->Find(GameTrackID);
 					if (TSharedPtr<FChaosVDTrackInfo> TrackInfoSharedPtr = TrackInfoPtrPtr? *TrackInfoPtrPtr : nullptr)
 					{
 						TrackInfoSharedPtr->CurrentFrame = FrameNumber;
-						OnTrackFrameUpdated().Broadcast(AsWeak(), TrackInfoSharedPtr.Get(), InstigatorID);
+						EnqueueTrackInfoUpdate(*TrackInfoSharedPtr.Get(), InstigatorID);
 					}
 				}
 			}
@@ -248,14 +251,20 @@ void FChaosVDPlaybackController::GoToRecordedGameFrame(const int32 FrameNumber,
 }
 
 void FChaosVDPlaybackController::GoToTrackFrame(FGuid InstigatorID, EChaosVDTrackType TrackType, int32 InTrackID, int32 FrameNumber, int32 Step)
+{
+	FReadScopeLock ReadLock(LoadedRecording->GetRecordingDataLock());
+	GoToTrackFrame_AssumesLocked(InstigatorID, TrackType, InTrackID, FrameNumber, Step);
+}
+
+void FChaosVDPlaybackController::GoToTrackFrame_AssumesLocked(FGuid InstigatorID, EChaosVDTrackType TrackType, int32 InTrackID, int32 FrameNumber, int32 Step)
 {
 	switch (TrackType)
 	{
 	case EChaosVDTrackType::Game:
-		GoToRecordedGameFrame(FrameNumber, InstigatorID);
+		GoToRecordedGameFrame_AssumesLocked(FrameNumber, InstigatorID);
 		break;
 	case EChaosVDTrackType::Solver:
-		GoToRecordedSolverStep(InTrackID, FrameNumber, Step, InstigatorID);
+		GoToRecordedSolverStep_AssumesLocked(InTrackID, FrameNumber, Step, InstigatorID);
 		break;
 	default:
 		ensure(false);
@@ -263,7 +272,7 @@ void FChaosVDPlaybackController::GoToTrackFrame(FGuid InstigatorID, EChaosVDTrac
 	}
 }
 
-int32 FChaosVDPlaybackController::GetTrackStepsNumberAtFrame(EChaosVDTrackType TrackType, const int32 InTrackID, const int32 FrameNumber) const
+int32 FChaosVDPlaybackController::GetTrackStepsNumberAtFrame_AssumesLocked(EChaosVDTrackType TrackType, const int32 InTrackID, const int32 FrameNumber) const
 {
 	if (!LoadedRecording.IsValid())
 	{
@@ -277,24 +286,24 @@ int32 FChaosVDPlaybackController::GetTrackStepsNumberAtFrame(EChaosVDTrackType T
 			return 0;
 			break;
 		case EChaosVDTrackType::Solver:
-		{
-			if (const FChaosVDSolverFrameData* FrameData = LoadedRecording->GetSolverFrameData(InTrackID, FrameNumber))
 			{
-				return FrameData->SolverSteps.Num() > 0 ? FrameData->SolverSteps.Num() : INDEX_NONE;
-			}
-			else
-			{
-				return INDEX_NONE;
+				if (const FChaosVDSolverFrameData* FrameData = LoadedRecording->GetSolverFrameData_AssumesLocked(InTrackID, FrameNumber))
+				{
+					return FrameData->SolverSteps.Num() > 0 ? FrameData->SolverSteps.Num() : INDEX_NONE;
+				}
+				else
+				{
+					return INDEX_NONE;
+				}
+				break;
 			}
-			break;
-		}
 	default:
 		return INDEX_NONE;
 		break;
 	}
 }
 
-const FChaosVDStepsContainer* FChaosVDPlaybackController::GetTrackStepsDataAtFrame(EChaosVDTrackType TrackType, int32 InTrackID, int32 FrameNumber) const
+const FChaosVDStepsContainer* FChaosVDPlaybackController::GetTrackStepsDataAtFrame_AssumesLocked(EChaosVDTrackType TrackType, int32 InTrackID, int32 FrameNumber) const
 {
 	if (!LoadedRecording.IsValid())
 	{
@@ -309,7 +318,7 @@ const FChaosVDStepsContainer* FChaosVDPlaybackController::GetTrackStepsDataAtFra
 		break;
 	case EChaosVDTrackType::Solver:
 		{
-			if (const FChaosVDSolverFrameData* FrameData = LoadedRecording->GetSolverFrameData(InTrackID, FrameNumber))
+			if (const FChaosVDSolverFrameData* FrameData = LoadedRecording->GetSolverFrameData_AssumesLocked(InTrackID, FrameNumber))
 			{
 				return &FrameData->SolverSteps;
 			}
@@ -434,7 +443,7 @@ int32 FChaosVDPlaybackController::GetTrackLastStepAtFrame(EChaosVDTrackType Trac
 	{
 		case EChaosVDTrackType::Solver:
 		{
-			const int32 AvailableSteps = GetTrackStepsNumberAtFrame(EChaosVDTrackType::Solver, InTrackID, InFrameNumber);
+			const int32 AvailableSteps = GetTrackStepsNumberAtFrame_AssumesLocked(EChaosVDTrackType::Solver, InTrackID, InFrameNumber);
 			return AvailableSteps == INDEX_NONE ? INDEX_NONE: AvailableSteps -1;
 			break;
 		}
@@ -530,9 +539,46 @@ void FChaosVDPlaybackController::GetAvailableTrackInfosAtTrackFrame(EChaosVDTrac
 	}
 }
 
+bool FChaosVDPlaybackController::Tick(float DeltaTime)
+{
+	const TWeakPtr<FChaosVDPlaybackController> ThisWeakPtr = DoesSharedInstanceExist() ? AsWeak() : nullptr;
+	if (!ThisWeakPtr.IsValid())
+	{
+		return true;
+	}
+
+	if (bHasPendingGTUpdateBroadcast)
+	{
+		ControllerUpdatedDelegate.Broadcast(ThisWeakPtr);
+		bHasPendingGTUpdateBroadcast = false;
+	}
+
+	if (!TrackInfoUpdateGTQueue.IsEmpty())
+	{
+		FChaosVDQueuedTrackInfoUpdate TrackInfoUpdate;
+		while (TrackInfoUpdateGTQueue.Dequeue(TrackInfoUpdate))
+		{
+			OnTrackFrameUpdated().Broadcast(ThisWeakPtr, &TrackInfoUpdate.TrackInfo, TrackInfoUpdate.InstigatorID);
+		}
+	}
+
+	if (!bPlayedFirsFrame)
+	{
+		if (LoadedRecording.IsValid() && LoadedRecording->GetAvailableSolversNumber_AssumesLocked() > 0)
+		{
+			constexpr int32 GameFrame = 0;
+			constexpr int32 Step = 0;
+			GoToTrackFrame(IChaosVDPlaybackControllerInstigator::InvalidGuid, EChaosVDTrackType::Game, GameTrackID, GameFrame, Step);
+			bPlayedFirsFrame = true;
+		}
+	}
+
+	return true;
+}
+
 void FChaosVDPlaybackController::UpdateSolverTracksData()
 {
-	const TMap<int32, TArray<FChaosVDSolverFrameData>>& SolversByID = LoadedRecording->GetAvailableSolvers();
+	const TMap<int32, TArray<FChaosVDSolverFrameData>>& SolversByID = LoadedRecording->GetAvailableSolvers_AssumesLocked();
 	for (const TPair<int32, TArray<FChaosVDSolverFrameData>>& SolverIDPair : SolversByID)
 	{
 		TSharedPtr<FChaosVDTrackInfo>& SolverTrackInfo = TrackInfoPerType[EChaosVDTrackType::Solver].FindOrAdd(SolverIDPair.Key);
@@ -566,11 +612,11 @@ void FChaosVDPlaybackController::HandleCurrentRecordingUpdated()
 		GameTrackInfo->CurrentStep = 0;
 	};
 
-	GameTrackInfo->MaxFrames = LoadedRecording->GetAvailableGameFrames().Num();
+	GameTrackInfo->MaxFrames = LoadedRecording->GetAvailableGameFrames_AssumesLocked().Num();
 	GameTrackInfo->TrackType = EChaosVDTrackType::Game;
 
 	// Each time the recording is updated, populate or update the existing solver tracks data
 	UpdateSolverTracksData();
 
-	OnDataUpdated().Broadcast(AsWeak());
+	bHasPendingGTUpdateBroadcast = true;
 }
diff --git a/Engine/Plugins/Experimental/ChaosVD/Source/ChaosVD/Private/Trace/ChaosVDTraceAnalyzer.cpp b/Engine/Plugins/Experimental/ChaosVD/Source/ChaosVD/Private/Trace/ChaosVDTraceAnalyzer.cpp
--- a/Engine/Plugins/Experimental/ChaosVD/Source/ChaosVD/Private/Trace/ChaosVDTraceAnalyzer.cpp
+++ b/Engine/Plugins/Experimental/ChaosVD/Source/ChaosVD/Private/Trace/ChaosVDTraceAnalyzer.cpp
@@ -63,7 +63,8 @@ bool FChaosVDTraceAnalyzer::OnEvent(uint16 RouteId, EStyle Style, const FOnEvent
 			uint8 FrameType = EventData.GetValue<uint8>("FrameType");
 			if (static_cast<ETraceFrameType>(FrameType) == TraceFrameType_Game)
 			{
-				if (FChaosVDGameFrameData* CurrentFrameData = ChaosVDTraceProvider->GetLastGameFrame())
+				FWriteScopeLock WriteLock(ChaosVDTraceProvider->GetDataLock());
+				if (FChaosVDGameFrameData* CurrentFrameData = ChaosVDTraceProvider->GetLastGameFrame_AssumesLocked())
 				{
 					CurrentFrameData->LastCycle = EventData.GetValue<uint64>("Cycle");
 				}
@@ -94,8 +95,9 @@ bool FChaosVDTraceAnalyzer::OnEvent(uint16 RouteId, EStyle Style, const FOnEvent
 		{
 			const int32 SolverID = EventData.GetValue<int32>("SolverID");
 
+			FWriteScopeLock WriteLock(ChaosVDTraceProvider->GetDataLock());
 			// This can be null if the recording started Mid-Frame. In this case we just discard the data for now
-			if (FChaosVDSolverFrameData* FrameData  = ChaosVDTraceProvider->GetLastSolverFrame(SolverID))
+			if (FChaosVDSolverFrameData* FrameData  = ChaosVDTraceProvider->GetLastSolverFrame_AssumesLocked(SolverID))
 			{
 				// Add an empty step. It will be filled out by the particle (and later on other objects/elements) events
 				FChaosVDStepData& StepData = FrameData->SolverSteps.AddDefaulted_GetRef();
@@ -115,7 +117,8 @@ bool FChaosVDTraceAnalyzer::OnEvent(uint16 RouteId, EStyle Style, const FOnEvent
 		{
 			const int32 SolverID = EventData.GetValue<int32>("SolverID");
 
-			if (FChaosVDSolverFrameData* FrameData = ChaosVDTraceProvider->GetLastSolverFrame(SolverID))
+			FWriteScopeLock WriteLock(ChaosVDTraceProvider->GetDataLock());
+			if (FChaosVDSolverFrameData* FrameData = ChaosVDTraceProvider->GetLastSolverFrame_AssumesLocked(SolverID))
 			{
 				int32 ParticleDestroyedID = EventData.GetValue<int32>("ParticleID");
 
@@ -171,28 +174,29 @@ bool FChaosVDTraceAnalyzer::OnEvent(uint16 RouteId, EStyle Style, const FOnEvent
 			break;
 		}
 	case RouteId_ChaosVDSolverSimulationSpace:
-	{
-		const int32 SolverID = EventData.GetValue<int32>("SolverID");
+		{
+			const int32 SolverID = EventData.GetValue<int32>("SolverID");
 
-		FVector Position;
-		Position.X = EventData.GetValue<float>("PositionX");
-		Position.Y = EventData.GetValue<float>("PositionY");
-		Position.Z = EventData.GetValue<float>("PositionZ");
+			FVector Position;
+			Position.X = EventData.GetValue<float>("PositionX");
+			Position.Y = EventData.GetValue<float>("PositionY");
+			Position.Z = EventData.GetValue<float>("PositionZ");
 
-		FQuat Rotation;
-		Rotation.X = EventData.GetValue<float>("RotationX");
-		Rotation.Y = EventData.GetValue<float>("RotationY");
-		Rotation.Z = EventData.GetValue<float>("RotationZ");
-		Rotation.W = EventData.GetValue<float>("RotationW");
+			FQuat Rotation;
+			Rotation.X = EventData.GetValue<float>("RotationX");
+			Rotation.Y = EventData.GetValue<float>("RotationY");
+			Rotation.Z = EventData.GetValue<float>("RotationZ");
+			Rotation.W = EventData.GetValue<float>("RotationW");
 
-		// This can be null if the recording started Mid-Frame. In this case we just discard the data for now
-		if (FChaosVDSolverFrameData* FrameData = ChaosVDTraceProvider->GetLastSolverFrame(SolverID))
-		{
-			FrameData->SimulationTransform.SetLocation(Position);
-			FrameData->SimulationTransform.SetRotation(Rotation);
+			FWriteScopeLock WriteLock(ChaosVDTraceProvider->GetDataLock());
+			// This can be null if the recording started Mid-Frame. In this case we just discard the data for now
+			if (FChaosVDSolverFrameData* FrameData = ChaosVDTraceProvider->GetLastSolverFrame_AssumesLocked(SolverID))
+			{
+				FrameData->SimulationTransform.SetLocation(Position);
+				FrameData->SimulationTransform.SetRotation(Rotation);
+			}
+			break;
 		}
-		break;
-	}
 	default:
 		break;
 	}
diff --git a/Engine/Plugins/Experimental/ChaosVD/Source/ChaosVD/Private/Trace/ChaosVDTraceProvider.cpp b/Engine/Plugins/Experimental/ChaosVD/Source/ChaosVD/Private/Trace/ChaosVDTraceProvider.cpp
--- a/Engine/Plugins/Experimental/ChaosVD/Source/ChaosVD/Private/Trace/ChaosVDTraceProvider.cpp
+++ b/Engine/Plugins/Experimental/ChaosVD/Source/ChaosVD/Private/Trace/ChaosVDTraceProvider.cpp
@@ -46,54 +46,55 @@ void FChaosVDTraceProvider::AddGameFrame(FChaosVDGameFrameData&& FrameData)
 	if (InternalRecording.IsValid())
 	{
 		// In PIE, we can have a lot of empty frames at the beginning, so we discard them here
-		if (InternalRecording->GetAvailableSolvers().IsEmpty())
 		{
-			if (FChaosVDGameFrameData* GameFrame = InternalRecording->GetLastGameFrameData())
+			// This is a kind of made up number, this includes RBAN solvers so we have more than 1 for sure.
+			constexpr int32 MinSolverIDsExpected = 10;
+			
+			static TArray<int32> SolverIDs;
+			SolverIDs.Reserve(MinSolverIDsExpected);
+			
+			FWriteScopeLock WriteLock(GetDataLock());
+			const int32 AvailableGameFrames = InternalRecording->GetAvailableGameFramesNumber_AssumesLocked();
+			InternalRecording->GetAvailableSolverIDsAtGameFrameNumber_AssumesLocked(AvailableGameFrames -1, SolverIDs);
+			if (SolverIDs.IsEmpty())
 			{
-				GameFrame = &FrameData;
-				return;
+				if (FChaosVDGameFrameData* GameFrame = InternalRecording->GetLastGameFrameData_AssumesLocked())
+				{
+					*GameFrame = FrameData;
+					return;
+				}
 			}
+
+			SolverIDs.Reset();
 		}
 
 		InternalRecording->AddGameFrameData(MoveTemp(FrameData));
 	}
 }
 
-FChaosVDSolverFrameData* FChaosVDTraceProvider::GetSolverFrame(const int32 InSolverID, const int32 FrameNumber) const
+FChaosVDSolverFrameData* FChaosVDTraceProvider::GetSolverFrame_AssumesLocked(const int32 InSolverID, const int32 FrameNumber) const
 {
-	return InternalRecording.IsValid() ? InternalRecording->GetSolverFrameData(InSolverID, FrameNumber) : nullptr;
+	return InternalRecording.IsValid() ? InternalRecording->GetSolverFrameData_AssumesLocked(InSolverID, FrameNumber) : nullptr;
 }
 
-FChaosVDSolverFrameData* FChaosVDTraceProvider::GetLastSolverFrame(const int32 InSolverID) const
+FChaosVDSolverFrameData* FChaosVDTraceProvider::GetLastSolverFrame_AssumesLocked(const int32 InSolverID) const
 {
-	if (InternalRecording.IsValid() && InternalRecording->GetAvailableSolverFramesNumber(InSolverID) > 0)
+	if (InternalRecording.IsValid())
 	{
-		const int32 AvailableFramesNumber = InternalRecording->GetAvailableSolverFramesNumber(InSolverID);
+		const int32 AvailableFramesNumber = InternalRecording->GetAvailableSolverFramesNumber_AssumesLocked(InSolverID);
 
-		if (AvailableFramesNumber != INDEX_NONE)
+		if (AvailableFramesNumber > 0)
 		{
-			return GetSolverFrame(InSolverID, InternalRecording->GetAvailableSolverFramesNumber(InSolverID) - 1);
+			return GetSolverFrame_AssumesLocked(InSolverID, AvailableFramesNumber - 1);
 		}
 	}
 
 	return nullptr;
 }
 
-FChaosVDGameFrameData* FChaosVDTraceProvider::GetSolverFrame(uint64 FrameStartCycle) const
-{
-	FChaosVDGameFrameData* FoundFrameData = nullptr;
-
-	if (InternalRecording.IsValid() && InternalRecording->GetAvailableGameFrames().Num() > 0)
-	{
-		return InternalRecording->GetGameFrameDataAtCycle(FrameStartCycle);
-	}
-	
-	return FoundFrameData;
-}
-
-FChaosVDGameFrameData* FChaosVDTraceProvider::GetLastGameFrame() const
+FChaosVDGameFrameData* FChaosVDTraceProvider::GetLastGameFrame_AssumesLocked() const
 {
-	return InternalRecording.IsValid() ? InternalRecording->GetLastGameFrameData() : nullptr;
+	return InternalRecording.IsValid() ? InternalRecording->GetLastGameFrameData_AssumesLocked() : nullptr;
 }
 
 FChaosVDBinaryDataContainer& FChaosVDTraceProvider::FindOrAddUnprocessedData(const int32 DataID)
@@ -165,6 +166,12 @@ void FChaosVDTraceProvider::RegisterDataProcessor(TSharedPtr<IChaosVDDataProcess
 	RegisteredDataProcessors.Add(InDataProcessor->GetCompatibleTypeName(), InDataProcessor);
 }
 
+FRWLock& FChaosVDTraceProvider::GetDataLock()
+{
+	check(InternalRecording.IsValid());
+	return InternalRecording->GetRecordingDataLock();
+}
+
 void FChaosVDTraceProvider::RegisterDefaultDataProcessorsIfNeeded()
 {
 	if (bDefaultDataProcessorsRegistered)
diff --git a/Engine/Plugins/Experimental/ChaosVD/Source/ChaosVD/Private/Trace/DataProcessors/ChaosVDConstraintDataProcessor.cpp b/Engine/Plugins/Experimental/ChaosVD/Source/ChaosVD/Private/Trace/DataProcessors/ChaosVDConstraintDataProcessor.cpp
--- a/Engine/Plugins/Experimental/ChaosVD/Source/ChaosVD/Private/Trace/DataProcessors/ChaosVDConstraintDataProcessor.cpp
+++ b/Engine/Plugins/Experimental/ChaosVD/Source/ChaosVD/Private/Trace/DataProcessors/ChaosVDConstraintDataProcessor.cpp
@@ -24,7 +24,8 @@ bool FChaosVDConstraintDataProcessor::ProcessRawData(const TArray<uint8>& InData
 	RecordedConstraint.Serialize(MemReader);
 
 	// This can be null if the recording started Mid-Frame. In this case we just discard the data for now
-	if (FChaosVDSolverFrameData* FrameData = ProviderSharedPtr->GetLastSolverFrame(RecordedConstraint.SolverID))
+	FWriteScopeLock WriteLock(ProviderSharedPtr->GetDataLock());
+	if (FChaosVDSolverFrameData* FrameData = ProviderSharedPtr->GetLastSolverFrame_AssumesLocked(RecordedConstraint.SolverID))
 	{
 		if (ensureMsgf(FrameData->SolverSteps.Num() > 0, TEXT("A MidPhase was traced without a valid step scope")))
 		{
diff --git a/Engine/Plugins/Experimental/ChaosVD/Source/ChaosVD/Private/Trace/DataProcessors/ChaosVDMidPhaseDataProcessor.cpp b/Engine/Plugins/Experimental/ChaosVD/Source/ChaosVD/Private/Trace/DataProcessors/ChaosVDMidPhaseDataProcessor.cpp
--- a/Engine/Plugins/Experimental/ChaosVD/Source/ChaosVD/Private/Trace/DataProcessors/ChaosVDMidPhaseDataProcessor.cpp
+++ b/Engine/Plugins/Experimental/ChaosVD/Source/ChaosVD/Private/Trace/DataProcessors/ChaosVDMidPhaseDataProcessor.cpp
@@ -23,8 +23,9 @@ bool FChaosVDMidPhaseDataProcessor::ProcessRawData(const TArray<uint8>& InData)
 
 	MidPhase->Serialize(MemReader);
 
+	FWriteScopeLock WriteLock(ProviderSharedPtr->GetDataLock());
 	// This can be null if the recording started Mid-Frame. In this case we just discard the data for now
-	if (FChaosVDSolverFrameData* FrameData = ProviderSharedPtr->GetLastSolverFrame(MidPhase->SolverID))
+	if (FChaosVDSolverFrameData* FrameData = ProviderSharedPtr->GetLastSolverFrame_AssumesLocked(MidPhase->SolverID))
 	{
 		if (ensureMsgf(FrameData->SolverSteps.Num() > 0, TEXT("A MidPhase was traced without a valid step scope")))
 		{
diff --git a/Engine/Plugins/Experimental/ChaosVD/Source/ChaosVD/Private/Trace/DataProcessors/ChaosVDTraceParticleDataProcessor.cpp b/Engine/Plugins/Experimental/ChaosVD/Source/ChaosVD/Private/Trace/DataProcessors/ChaosVDTraceParticleDataProcessor.cpp
--- a/Engine/Plugins/Experimental/ChaosVD/Source/ChaosVD/Private/Trace/DataProcessors/ChaosVDTraceParticleDataProcessor.cpp
+++ b/Engine/Plugins/Experimental/ChaosVD/Source/ChaosVD/Private/Trace/DataProcessors/ChaosVDTraceParticleDataProcessor.cpp
@@ -27,8 +27,9 @@ bool FChaosVDTraceParticleDataProcessor::ProcessRawData(const TArray<uint8>& InD
 	FChaosVDParticleDataWrapper ParticleData;
 	ParticleData.Serialize(Ar);
 
+	FWriteScopeLock WriteLock(ProviderSharedPtr->GetDataLock());
 	// This can be null if the recording started Mid-Frame. In this case we just discard the data for now
-	if (FChaosVDSolverFrameData* FrameData = ProviderSharedPtr->GetLastSolverFrame(ParticleData.SolverID))
+	if (FChaosVDSolverFrameData* FrameData = ProviderSharedPtr->GetLastSolverFrame_AssumesLocked(ParticleData.SolverID))
 	{
 		if (ensureMsgf(FrameData->SolverSteps.Num() > 0, TEXT("A particle was traced without a valid step scope")))
 		{
diff --git a/Engine/Plugins/Experimental/ChaosVD/Source/ChaosVD/Private/Widgets/SChaosVDPlaybackViewport.cpp b/Engine/Plugins/Experimental/ChaosVD/Source/ChaosVD/Private/Widgets/SChaosVDPlaybackViewport.cpp
--- a/Engine/Plugins/Experimental/ChaosVD/Source/ChaosVD/Private/Widgets/SChaosVDPlaybackViewport.cpp
+++ b/Engine/Plugins/Experimental/ChaosVD/Source/ChaosVD/Private/Widgets/SChaosVDPlaybackViewport.cpp
@@ -115,6 +115,7 @@ void SChaosVDPlaybackViewport::HandlePlaybackControllerDataUpdated(TWeakPtr<FCha
 		{
 			// Max is inclusive and we use this to request as the index on the recorded frames/steps arrays so we need to -1 to the available frames/steps
 			GameFramesTimelineWidget->UpdateMinMaxValue(0, TrackInfo->MaxFrames != INDEX_NONE ? TrackInfo->MaxFrames -1  : 0);
+			GameFramesTimelineWidget->SetCurrentTimelineFrame(TrackInfo->CurrentFrame);
 		}
 	}
 	else
diff --git a/Engine/Plugins/Experimental/ChaosVD/Source/ChaosVD/Private/Widgets/SChaosVDSolverPlaybackControls.cpp b/Engine/Plugins/Experimental/ChaosVD/Source/ChaosVD/Private/Widgets/SChaosVDSolverPlaybackControls.cpp
--- a/Engine/Plugins/Experimental/ChaosVD/Source/ChaosVD/Private/Widgets/SChaosVDSolverPlaybackControls.cpp
+++ b/Engine/Plugins/Experimental/ChaosVD/Source/ChaosVD/Private/Widgets/SChaosVDSolverPlaybackControls.cpp
@@ -93,7 +93,7 @@ void SChaosVDSolverPlaybackControls::HandlePlaybackControllerDataUpdated(TWeakPt
 		const int32 AvailableFrames = ControllerSharedPtr->GetTrackFramesNumber(EChaosVDTrackType::Solver, SolverID);
 
 		int32 AvailableSteps = INDEX_NONE;
-		if (const FChaosVDStepsContainer* StepData = ControllerSharedPtr->GetTrackStepsDataAtFrame(EChaosVDTrackType::Solver, SolverID, ControllerSharedPtr->GetTrackCurrentFrame(EChaosVDTrackType::Solver, SolverID)))
+		if (const FChaosVDStepsContainer* StepData = ControllerSharedPtr->GetTrackStepsDataAtFrame_AssumesLocked(EChaosVDTrackType::Solver, SolverID, ControllerSharedPtr->GetTrackCurrentFrame(EChaosVDTrackType::Solver, SolverID)))
 		{
 			AvailableSteps = StepData->Num() > 0 ? StepData->Num() : INDEX_NONE;
 			CurrentStepName =  StepData->Num() > 0 ? (*StepData)[0].StepName : TEXT("NONE");
@@ -117,7 +117,7 @@ void SChaosVDSolverPlaybackControls::HandlePlaybackControllerDataUpdated(TWeakPt
 
 void SChaosVDSolverPlaybackControls::UpdateStepsWidgetForFrame(const FChaosVDPlaybackController& InCurrentPlaybackController, int32 FrameNumber, int32 StepNumber, EChaosVDStepsWidgetUpdateFlags OptionsFlags)
 {
-	if (const FChaosVDStepsContainer* StepsData = InCurrentPlaybackController.GetTrackStepsDataAtFrame(EChaosVDTrackType::Solver, SolverID, FrameNumber))
+	if (const FChaosVDStepsContainer* StepsData = InCurrentPlaybackController.GetTrackStepsDataAtFrame_AssumesLocked(EChaosVDTrackType::Solver, SolverID, FrameNumber))
 	{
 		constexpr TCHAR const* UnknownStepName = TEXT("Unknown");
 
diff --git a/Engine/Plugins/Experimental/ChaosVD/Source/ChaosVD/Public/ChaosVDPlaybackController.h b/Engine/Plugins/Experimental/ChaosVD/Source/ChaosVD/Public/ChaosVDPlaybackController.h
--- a/Engine/Plugins/Experimental/ChaosVD/Source/ChaosVD/Public/ChaosVDPlaybackController.h
+++ b/Engine/Plugins/Experimental/ChaosVD/Source/ChaosVD/Public/ChaosVDPlaybackController.h
@@ -3,9 +3,12 @@
 #pragma once
 
 #include "ChaosVDRecording.h"
+#include "Containers/Queue.h"
+#include "Containers/Ticker.h"
 #include "Templates/SharedPointer.h"
 #include "Delegates/DelegateCombinations.h"
 #include "Delegates/Delegate.h"
+#include "HAL/ThreadSafeBool.h"
 
 struct FChaosVDTrackInfo;
 class FChaosVDScene;
@@ -36,6 +39,12 @@ struct FChaosVDTrackInfo
 	FString TrackName;
 };
 
+struct FChaosVDQueuedTrackInfoUpdate
+{
+	FChaosVDTrackInfo TrackInfo;
+	FGuid InstigatorID;
+};
+
 /** Flags used to control how the unload of a recording is performed */
 enum class EChaosVDUnloadRecordingFlags : uint8
 {
@@ -48,7 +57,7 @@ ENUM_CLASS_FLAGS(EChaosVDUnloadRecordingFlags)
 typedef TMap<int32, TSharedPtr<FChaosVDTrackInfo>> TrackInfoByIDMap;
 
 /** Loads,unloads and owns a Chaos VD recording file */
-class FChaosVDPlaybackController : public TSharedFromThis<FChaosVDPlaybackController>
+class FChaosVDPlaybackController : public TSharedFromThis<FChaosVDPlaybackController>, public FTSTickerObjectBase
 {
 public:
 
@@ -56,7 +65,7 @@ public:
 	static constexpr int32 GameTrackID  = 0 ;
 
 	FChaosVDPlaybackController(const TWeakPtr<FChaosVDScene>& InSceneToControl);
-	~FChaosVDPlaybackController();
+	virtual ~FChaosVDPlaybackController() override;
 
 	/** Loads a recording using a Trace Session Name */
 	bool LoadChaosVDRecordingFromTraceSession(const FString& InSessionName);
@@ -81,6 +90,7 @@ public:
 	 * @param Step Step number to go
 	 */
 	void GoToTrackFrame(FGuid InstigatorID, EChaosVDTrackType TrackType, int32 InTrackID, int32 FrameNumber, int32 Step);
+	void GoToTrackFrame_AssumesLocked(FGuid InstigatorID, EChaosVDTrackType TrackType, int32 InTrackID, int32 FrameNumber, int32 Step);
 
 	/**
 	 * Gets the number of available steps in a track at the specified frame
@@ -89,7 +99,7 @@ public:
 	 * @param FrameNumber Frame number to evaluate
 	 * @return Number of available steps
 	 */
-	int32 GetTrackStepsNumberAtFrame(EChaosVDTrackType TrackType, int32 InTrackID, int32 FrameNumber) const;
+	int32 GetTrackStepsNumberAtFrame_AssumesLocked(EChaosVDTrackType TrackType, int32 InTrackID, int32 FrameNumber) const;
 
 	/**
 	 * Gets the available steps container in a track at the specified frame
@@ -98,7 +108,7 @@ public:
 	 * @param FrameNumber Frame number to evaluate
 	 * @return Ptr to the steps data container
 	 */
-	const FChaosVDStepsContainer* GetTrackStepsDataAtFrame(EChaosVDTrackType TrackType, int32 InTrackID, int32 FrameNumber) const;
+	const FChaosVDStepsContainer* GetTrackStepsDataAtFrame_AssumesLocked(EChaosVDTrackType TrackType, int32 InTrackID, int32 FrameNumber) const;
 
 	/**
 	 * Gets the number of available frames for the specified track
@@ -171,7 +181,6 @@ public:
 	 */
 	FChaosVDTrackInfo* GetMutableTrackInfo(EChaosVDTrackType TrackType, int32 TrackID);
 
-
 	/**
 	 * Locks the steps timeline of a given track so each time you move between frames, it will automatically scrub to the locked in step
 	 * @param TrackType Type of the track to find
@@ -194,6 +203,8 @@ public:
 
 	/** Called when a frame on a track is updated */
 	FChaosVDPlaybackControllerFrameUpdated& OnTrackFrameUpdated() { return ControllerFrameUpdatedDelegate; }
+	
+	virtual bool Tick(float DeltaTime) override;
 
 protected:
 
@@ -201,17 +212,20 @@ protected:
 	void UpdateSolverTracksData();
 
 	/** Updates the controlled scene with the loaded data at specified game frame */
-	void GoToRecordedGameFrame(int32 FrameNumber, FGuid InstigatorID);
+	void GoToRecordedGameFrame_AssumesLocked(int32 FrameNumber, FGuid InstigatorID);
 
 	/** Updates the controlled scene with the loaded data at specified solver frame and solver step */
-	void GoToRecordedSolverStep(int32 InTrackID, int32 FrameNumber, int32 Step, FGuid InstigatorID);
+	void GoToRecordedSolverStep_AssumesLocked(int32 InTrackID, int32 FrameNumber, int32 Step, FGuid InstigatorID);
 
 	/** Handles any data changes on the loaded recording - Usually called during Trace analysis */
 	void HandleCurrentRecordingUpdated();
 
 	/** Finds the closest Key frame to the provided frame number, and plays all the following frames until the specified frame number (no inclusive) */
-	void PlayFromClosestKeyFrame(int32 InTrackID, int32 FrameNumber, FChaosVDScene& InSceneToControl);
+	void PlayFromClosestKeyFrame_AssumesLocked(int32 InTrackID, int32 FrameNumber, FChaosVDScene& InSceneToControl);
 
+	/** Add the provided track info update to the queue. The update will be broadcast in the game thread */
+	void EnqueueTrackInfoUpdate(const FChaosVDTrackInfo& InTrackInfo, FGuid InstigatorID);
+	
 	/** Map containing all track info, by track type*/
 	TMap<EChaosVDTrackType, TrackInfoByIDMap> TrackInfoPerType;
 
@@ -226,4 +240,12 @@ protected:
 
 	/** Delegate called when the a in a track changes */
 	FChaosVDPlaybackControllerFrameUpdated ControllerFrameUpdatedDelegate;
+
+	/** Set to true when the recording data controlled by this Playback Controller is updated, the update delegate will be called on the GT */
+	FThreadSafeBool bHasPendingGTUpdateBroadcast;
+
+	/** Queue with a copy of all Track Info Updates that needs to be done in the Game thread */
+	TQueue<FChaosVDQueuedTrackInfoUpdate, EQueueMode::Mpsc> TrackInfoUpdateGTQueue;
+
+	bool bPlayedFirsFrame = false;
 };
diff --git a/Engine/Plugins/Experimental/ChaosVD/Source/ChaosVD/Public/Trace/ChaosVDTraceProvider.h b/Engine/Plugins/Experimental/ChaosVD/Source/ChaosVD/Public/Trace/ChaosVDTraceProvider.h
--- a/Engine/Plugins/Experimental/ChaosVD/Source/ChaosVD/Public/Trace/ChaosVDTraceProvider.h
+++ b/Engine/Plugins/Experimental/ChaosVD/Source/ChaosVD/Public/Trace/ChaosVDTraceProvider.h
@@ -54,11 +54,10 @@ public:
 	void DeleteRecordingInstanceForSession();
 	void AddSolverFrame(const int32 InSolverGUID, FChaosVDSolverFrameData&& FrameData);
 	void AddGameFrame(FChaosVDGameFrameData&& FrameData);
-	FChaosVDSolverFrameData* GetSolverFrame(const int32 InSolverGUID, const int32 FrameNumber) const;
-	FChaosVDSolverFrameData* GetLastSolverFrame(const int32 InSolverGUID) const;
-
-	FChaosVDGameFrameData* GetSolverFrame(uint64 FrameStartCycle) const;
-	FChaosVDGameFrameData* GetLastGameFrame() const;
+	FChaosVDSolverFrameData* GetSolverFrame_AssumesLocked(const int32 InSolverGUID, const int32 FrameNumber) const;
+	FChaosVDSolverFrameData* GetLastSolverFrame_AssumesLocked(const int32 InSolverGUID) const;
+	
+	FChaosVDGameFrameData* GetLastGameFrame_AssumesLocked() const;
 
 	FChaosVDBinaryDataContainer& FindOrAddUnprocessedData(const int32 DataID);
 
@@ -68,6 +67,8 @@ public:
 
 	void RegisterDataProcessor(TSharedPtr<IChaosVDDataProcessor> InDataProcessor);
 
+	FRWLock& GetDataLock();
+
 private:
 
 	void RegisterDefaultDataProcessorsIfNeeded();
diff --git a/Engine/Source/Runtime/Experimental/ChaosVDData/Private/ChaosVDRecording.cpp b/Engine/Source/Runtime/Experimental/ChaosVDData/Private/ChaosVDRecording.cpp
--- a/Engine/Source/Runtime/Experimental/ChaosVDData/Private/ChaosVDRecording.cpp
+++ b/Engine/Source/Runtime/Experimental/ChaosVDData/Private/ChaosVDRecording.cpp
@@ -3,7 +3,26 @@
 #include "ChaosVDRecording.h"
 #include "Chaos/ImplicitObject.h"
 
+int32 FChaosVDRecording::GetAvailableGameFramesNumber() const
+{
+	FReadScopeLock ReadLock(RecordingDataLock);
+	
+	return GetAvailableGameFramesNumber_AssumesLocked();
+}
+
+int32 FChaosVDRecording::GetAvailableGameFramesNumber_AssumesLocked() const
+{
+	return GameFrames.Num(); 
+}
+
 int32 FChaosVDRecording::GetAvailableSolverFramesNumber(const int32 SolverID) const
+{
+	FReadScopeLock ReadLock(RecordingDataLock);
+
+	return GetAvailableSolverFramesNumber_AssumesLocked(SolverID);
+}
+
+int32 FChaosVDRecording::GetAvailableSolverFramesNumber_AssumesLocked(int32 SolverID) const
 {
 	if (const TArray<FChaosVDSolverFrameData>* SolverData = RecordedFramesDataPerSolver.Find(SolverID))
 	{
@@ -17,6 +36,8 @@ FString FChaosVDRecording::GetSolverName(int32 SolverID)
 {
 	static FString DefaultName(TEXT("Invalid"));
 
+	FReadScopeLock ReadLock(RecordingDataLock);
+
 	// Currently we don't create an entry per solver, so we need to get the name from the frame data
 	// TODO: Record Solver specific data per instance and not per frame
 	if (TArray<FChaosVDSolverFrameData>* SolverFramesData = RecordedFramesDataPerSolver.Find(SolverID))
@@ -30,7 +51,7 @@ FString FChaosVDRecording::GetSolverName(int32 SolverID)
 	return DefaultName;
 }
 
-FChaosVDSolverFrameData* FChaosVDRecording::GetSolverFrameData(const int32 SolverID, const int32 FrameNumber)
+FChaosVDSolverFrameData* FChaosVDRecording::GetSolverFrameData_AssumesLocked(const int32 SolverID, const int32 FrameNumber)
 {
 	if (TArray<FChaosVDSolverFrameData>* SolverFrames = RecordedFramesDataPerSolver.Find(SolverID))
 	{
@@ -41,7 +62,7 @@ FChaosVDSolverFrameData* FChaosVDRecording::GetSolverFrameData(const int32 Solve
 	return nullptr;
 }
 
-FChaosVDSolverFrameData* FChaosVDRecording::GetSolverFrameDataAtCycle(int32 SolverID, uint64 Cycle)
+FChaosVDSolverFrameData* FChaosVDRecording::GetSolverFrameDataAtCycle_AssumesLocked(int32 SolverID, uint64 Cycle)
 {
 	if (TArray<FChaosVDSolverFrameData>* SolverFramesPtr = RecordedFramesDataPerSolver.Find(SolverID))
 	{
@@ -54,6 +75,13 @@ FChaosVDSolverFrameData* FChaosVDRecording::GetSolverFrameDataAtCycle(int32 Solv
 }
 
 int32 FChaosVDRecording::GetLowestSolverFrameNumberAtCycle(int32 SolverID, uint64 Cycle)
+{
+	FReadScopeLock ReadLock(RecordingDataLock);
+
+	return GetLowestSolverFrameNumberAtCycle_AssumesLocked(SolverID, Cycle);
+}
+
+int32 FChaosVDRecording::GetLowestSolverFrameNumberAtCycle_AssumesLocked(int32 SolverID, uint64 Cycle)
 {
 	if (TArray<FChaosVDSolverFrameData>* SolverFramesPtr = RecordedFramesDataPerSolver.Find(SolverID))
 	{
@@ -64,7 +92,7 @@ int32 FChaosVDRecording::GetLowestSolverFrameNumberAtCycle(int32 SolverID, uint6
 	return INDEX_NONE;
 }
 
-int32 FChaosVDRecording::FindFirstSolverKeyFrameNumberFromFrame(int32 SolverID, int32 StartFrameNumber)
+int32 FChaosVDRecording::FindFirstSolverKeyFrameNumberFromFrame_AssumesLocked(int32 SolverID, int32 StartFrameNumber)
 {
 	if (TArray<int32>* KeyFrameNumbersPtr = RecordedKeyFramesNumberPerSolver.Find(SolverID))
 	{
@@ -104,6 +132,7 @@ int32 FChaosVDRecording::FindFirstSolverKeyFrameNumberFromFrame(int32 SolverID,
 
 int32 FChaosVDRecording::GetLowestSolverFrameNumberGameFrame(int32 SolverID, int32 GameFrame)
 {
+	FReadScopeLock ReadLock(RecordingDataLock);
 	if (!GameFrames.IsValidIndex(GameFrame))
 	{
 		return INDEX_NONE;
@@ -121,6 +150,7 @@ int32 FChaosVDRecording::GetLowestSolverFrameNumberGameFrame(int32 SolverID, int
 
 int32 FChaosVDRecording::GetLowestGameFrameAtSolverFrameNumber(int32 SolverID, int32 SolverFrame)
 {
+	FReadScopeLock ReadLock(RecordingDataLock);
 	if (TArray<FChaosVDSolverFrameData>* SolverFramesPtr = RecordedFramesDataPerSolver.Find(SolverID))
 	{
 		TArray<FChaosVDSolverFrameData>& SolverFrames = *SolverFramesPtr;
@@ -136,6 +166,7 @@ int32 FChaosVDRecording::GetLowestGameFrameAtSolverFrameNumber(int32 SolverID, i
 
 void FChaosVDRecording::AddKeyFrameNumberForSolver(const int32 SolverID, int32 FrameNumber)
 {
+	FWriteScopeLock WriteLock(RecordingDataLock);
 	if (TArray<int32>* KeyFrameNumber = RecordedKeyFramesNumberPerSolver.Find(SolverID))
 	{
 		KeyFrameNumber->Add(FrameNumber);
@@ -150,16 +181,21 @@ void FChaosVDRecording::AddFrameForSolver(const int32 SolverID, FChaosVDSolverFr
 {
 	int32 FrameNumber;
 	const bool bIsKeyFrame = InFrameData.bIsKeyFrame;
-	if (TArray<FChaosVDSolverFrameData>* SolverFrames = RecordedFramesDataPerSolver.Find(SolverID))
+
 	{
-		FrameNumber = SolverFrames->Num();
+		FWriteScopeLock WriteLock(RecordingDataLock);
 
-		SolverFrames->Add(MoveTemp(InFrameData));	
-	}
-	else
-	{	
-		FrameNumber = 0;
-		RecordedFramesDataPerSolver.Add(SolverID, { MoveTemp(InFrameData) });
+		if (TArray<FChaosVDSolverFrameData>* SolverFrames = RecordedFramesDataPerSolver.Find(SolverID))
+		{
+			FrameNumber = SolverFrames->Num();
+
+			SolverFrames->Add(MoveTemp(InFrameData));	
+		}
+		else
+		{	
+			FrameNumber = 0;
+			RecordedFramesDataPerSolver.Add(SolverID, { MoveTemp(InFrameData) });
+		}
 	}
 
 	if (bIsKeyFrame)
@@ -172,10 +208,11 @@ void FChaosVDRecording::AddFrameForSolver(const int32 SolverID, FChaosVDSolverFr
 
 void FChaosVDRecording::AddGameFrameData(FChaosVDGameFrameData&& InFrameData)
 {
+	FWriteScopeLock WriteLock(RecordingDataLock);
 	GameFrames.Add(MoveTemp(InFrameData));
 }
 
-FChaosVDGameFrameData* FChaosVDRecording::GetGameFrameDataAtCycle(uint64 Cycle)
+FChaosVDGameFrameData* FChaosVDRecording::GetGameFrameDataAtCycle_AssumesLocked(uint64 Cycle)
 {
 	int32 FrameIndex = Algo::BinarySearchBy(GameFrames, Cycle, &FChaosVDGameFrameData::FirstCycle);
 
@@ -187,22 +224,30 @@ FChaosVDGameFrameData* FChaosVDRecording::GetGameFrameDataAtCycle(uint64 Cycle)
 	return nullptr;
 }
 
-FChaosVDGameFrameData* FChaosVDRecording::GetGameFrameData(int32 FrameNumber)
+FChaosVDGameFrameData* FChaosVDRecording::GetGameFrameData_AssumesLocked(int32 FrameNumber)
 {
 	return GameFrames.IsValidIndex(FrameNumber) ? &GameFrames[FrameNumber] : nullptr;
 }
 
-FChaosVDGameFrameData* FChaosVDRecording::GetLastGameFrameData()
+FChaosVDGameFrameData* FChaosVDRecording::GetLastGameFrameData_AssumesLocked()
 {
 	return GameFrames.Num() > 0 ? &GameFrames.Last() : nullptr;
 }
 
 int32 FChaosVDRecording::GetLowestGameFrameNumberAtCycle(uint64 Cycle)
 {
+	FReadScopeLock ReadLock(RecordingDataLock);
 	return Algo::LowerBoundBy(GameFrames, Cycle, &FChaosVDGameFrameData::FirstCycle);
 }
 
 void FChaosVDRecording::GetAvailableSolverIDsAtGameFrameNumber(int32 FrameNumber, TArray<int32>& OutSolversID)
+{
+	FReadScopeLock ReadLock(RecordingDataLock);
+
+	return GetAvailableSolverIDsAtGameFrameNumber_AssumesLocked(FrameNumber, OutSolversID);
+}
+
+void FChaosVDRecording::GetAvailableSolverIDsAtGameFrameNumber_AssumesLocked(int32 FrameNumber, TArray<int32>& OutSolversID)
 {
 	if (!GameFrames.IsValidIndex(FrameNumber))
 	{
@@ -236,6 +281,7 @@ void FChaosVDRecording::GetAvailableSolverIDsAtGameFrameNumber(int32 FrameNumber
 
 void FChaosVDRecording::AddImplicitObject(const uint32 ID, const TSharedPtr<Chaos::FImplicitObject>& InImplicitObject)
 {
+	FWriteScopeLock WriteLock(RecordingDataLock);
 	if (!ImplicitObjects.Contains(ID))
 	{
 		AddImplicitObject_Internal(ID, InImplicitObject);
@@ -244,6 +290,7 @@ void FChaosVDRecording::AddImplicitObject(const uint32 ID, const TSharedPtr<Chao
 
 void FChaosVDRecording::AddImplicitObject(const uint32 ID, const Chaos::FImplicitObject* InImplicitObject)
 {
+	FWriteScopeLock WriteLock(RecordingDataLock);
 	if (!ImplicitObjects.Contains(ID))
 	{
 		// Only take ownership after we know we will add it to the map
diff --git a/Engine/Source/Runtime/Experimental/ChaosVDData/Public/ChaosVDRecording.h b/Engine/Source/Runtime/Experimental/ChaosVDData/Public/ChaosVDRecording.h
--- a/Engine/Source/Runtime/Experimental/ChaosVDData/Public/ChaosVDRecording.h
+++ b/Engine/Source/Runtime/Experimental/ChaosVDData/Public/ChaosVDRecording.h
@@ -53,22 +53,24 @@ struct FChaosVDGameFrameData
 struct CHAOSVDDATA_API FChaosVDRecording
 {	
 	/** Returns the current available recorded solvers number */
-	int32 GetAvailableSolversNumber() const { return RecordedFramesDataPerSolver.Num(); }
+	int32 GetAvailableSolversNumber_AssumesLocked() const { return RecordedFramesDataPerSolver.Num(); }
 	
 	/** Returns the current available Game Frames */
-	int32 GetAvailableGameFramesNumber() const { return GameFrames.Num(); }
+	int32 GetAvailableGameFramesNumber() const;
+	int32 GetAvailableGameFramesNumber_AssumesLocked() const;
 
 	/** Returns a reference to the array holding all the available game frames */
-	const TArray<FChaosVDGameFrameData>& GetAvailableGameFrames() const { return GameFrames; }
+	const TArray<FChaosVDGameFrameData>& GetAvailableGameFrames_AssumesLocked() const { return GameFrames; }
 
 	/** Returns a reference to the map containing the available solver data */
-	const TMap<int32, TArray<FChaosVDSolverFrameData>>& GetAvailableSolvers() const { return RecordedFramesDataPerSolver; }
+	const TMap<int32, TArray<FChaosVDSolverFrameData>>& GetAvailableSolvers_AssumesLocked() const { return RecordedFramesDataPerSolver; }
 
 	/**
 	 * Returns the number of available frame data for the specified solver ID
 	 * @param SolverID ID of the solver 
 	 */
 	int32 GetAvailableSolverFramesNumber(int32 SolverID) const;
+	int32 GetAvailableSolverFramesNumber_AssumesLocked(int32 SolverID) const;
 	
 	/**
 	 * Returns the name of the specified solver id
@@ -82,7 +84,7 @@ struct CHAOSVDDATA_API FChaosVDRecording
 	 * @param FrameNumber Frame number
 	 * @return Ptr to the existing solver frame data from the specified ID and Frame number - It is a ptr to the array element, Do not store
 	 */
-	FChaosVDSolverFrameData* GetSolverFrameData(int32 SolverID, int32 FrameNumber);
+	FChaosVDSolverFrameData* GetSolverFrameData_AssumesLocked(int32 SolverID, int32 FrameNumber);
 	
 	/**
 	 * Return a ptr to the existing solver frame data from the specified ID and Frame number
@@ -90,7 +92,7 @@ struct CHAOSVDDATA_API FChaosVDRecording
 	 * @param Cycle Platform cycle at which the solver frame was recorded
 	 * @return Ptr to the existing solver frame data from the specified ID and Frame number - It is a ptr to the array element, Do not store
 	 */
-	FChaosVDSolverFrameData* GetSolverFrameDataAtCycle(int32 SolverID, uint64 Cycle);
+	FChaosVDSolverFrameData* GetSolverFrameDataAtCycle_AssumesLocked(int32 SolverID, uint64 Cycle);
 
 	/**
 	 * Searches and returns the lowest frame number of a solver at the specified cycle
@@ -99,8 +101,9 @@ struct CHAOSVDDATA_API FChaosVDRecording
 	 * @return Found frame number. INDEX_NONE if no frame is found for the specified cycle
 	 */
 	int32 GetLowestSolverFrameNumberAtCycle(int32 SolverID, uint64 Cycle);
+	int32 GetLowestSolverFrameNumberAtCycle_AssumesLocked(int32 SolverID, uint64 Cycle);
 
-	int32 FindFirstSolverKeyFrameNumberFromFrame(int32 SolverID, int32 StartFrameNumber);
+	int32 FindFirstSolverKeyFrameNumberFromFrame_AssumesLocked(int32 SolverID, int32 StartFrameNumber);
 	
 	/**
 	 * Searches and returns the lowest frame number of a solver at the specified cycle
@@ -142,17 +145,17 @@ struct CHAOSVDDATA_API FChaosVDRecording
 	 * @param Cycle Platform Cycle to be used in the search
 	 * @return A ptr to the recorded game frame data - This is a ptr to the array element. Do not store
 	 */
-	FChaosVDGameFrameData* GetGameFrameDataAtCycle(uint64 Cycle);
+	FChaosVDGameFrameData* GetGameFrameDataAtCycle_AssumesLocked(uint64 Cycle);
 
 	/**
 	 * Searches for a recorded Game frame at the specified cycle 
 	 * @param FrameNumber Frame Number
 	 * @return A ptr to the recorded game frame data - This is a ptr to the array element. Do not store
 	 */
-	FChaosVDGameFrameData* GetGameFrameData(int32 FrameNumber);
+	FChaosVDGameFrameData* GetGameFrameData_AssumesLocked(int32 FrameNumber);
 
 	/** Returns a ptr to the last recorded game frame - This is a ptr to the array element. Do not store */
-	FChaosVDGameFrameData* GetLastGameFrameData();
+	FChaosVDGameFrameData* GetLastGameFrameData_AssumesLocked();
 
 	/**
 	 * Searches and returns the lowest game frame number at the specified cycle
@@ -167,6 +170,7 @@ struct CHAOSVDDATA_API FChaosVDRecording
      * @param OutSolversID Solver's ID array to be filled with any IDs found
      */
 	void GetAvailableSolverIDsAtGameFrameNumber(int32 FrameNumber, TArray<int32>& OutSolversID);
+	void GetAvailableSolverIDsAtGameFrameNumber_AssumesLocked(int32 FrameNumber, TArray<int32>& OutSolversID);
 
 	/** Returns a reference to the GeometryID-ImplicitObject map of this recording */
 	const TMap<uint32, TSharedPtr<const Chaos::FImplicitObject>>& GetGeometryDataMap() const { return ImplicitObjects; };
@@ -177,6 +181,8 @@ struct CHAOSVDDATA_API FChaosVDRecording
 	/** Session name of the trace session used to re-build this recording */
 	FString SessionName;
 
+	FRWLock& GetRecordingDataLock() { return RecordingDataLock; }
+
 protected:
 
 	/** Adds an Implicit Object to the recording and takes ownership of it */
@@ -198,6 +204,8 @@ protected:
 	/** Id to Ptr map of all shared geometry data required to visualize */
 	TMap<uint32, TSharedPtr<const Chaos::FImplicitObject>> ImplicitObjects;
 
+	mutable FRWLock RecordingDataLock;
+
 	friend class FChaosVDTraceProvider;
 	friend class FChaosVDTraceImplicitObjectProcessor;
 };
