[ChaosVD] Added UI feedback for slow and sync operations

Now if the processing of particle data takes more than 1.5 sec (can happen when processing the first frame of a huge map as it needs to spanw the actors), a pop up will be shown.

In addition, when geometry is being generated, as it is an async task, an notification progress bar will be shown instead of a full blocking progress bar

diff --git a/Engine/Plugins/Experimental/ChaosVD/Source/ChaosVD/ChaosVD.Build.cs b/Engine/Plugins/Experimental/ChaosVD/Source/ChaosVD/ChaosVD.Build.cs
--- a/Engine/Plugins/Experimental/ChaosVD/Source/ChaosVD/ChaosVD.Build.cs
+++ b/Engine/Plugins/Experimental/ChaosVD/Source/ChaosVD/ChaosVD.Build.cs
@@ -42,7 +42,8 @@ public class ChaosVD : ModuleRules
 				"DynamicMesh",
 				"SceneOutliner",
 				"TypedElementRuntime",
-				"TypedElementFramework"
+				"TypedElementFramework",
+				"StatusBar"
 			}
 			);
 
diff --git a/Engine/Plugins/Experimental/ChaosVD/Source/ChaosVD/Private/ChaosVDGeometryBuilder.cpp b/Engine/Plugins/Experimental/ChaosVD/Source/ChaosVD/Private/ChaosVDGeometryBuilder.cpp
--- a/Engine/Plugins/Experimental/ChaosVD/Source/ChaosVD/Private/ChaosVDGeometryBuilder.cpp
+++ b/Engine/Plugins/Experimental/ChaosVD/Source/ChaosVD/Private/ChaosVDGeometryBuilder.cpp
@@ -68,12 +68,9 @@ UDynamicMesh* FChaosVDGeometryBuilder::CreateAndCacheDynamicMesh(const uint32 Ge
 {
 	{
 		FReadScopeLock ReadLock(RWLock);
-		//TODO: Make this return what is cached when the system is more robust
-		// For now this should not happen and we want to catch it and make it visually noticeable
-		if (DynamicMeshCacheMap.Contains(GeometryCacheKey))
+		if (TObjectPtr<UDynamicMesh>* DynamicMeshPtrPtr = DynamicMeshCacheMap.Find(GeometryCacheKey))
 		{
-			ensureMsgf(false, TEXT("Tried to create a new mesh with an existing Cache key"));
-			return nullptr;
+			return *DynamicMeshPtrPtr;
 		}
 	}
 
@@ -94,12 +91,9 @@ UStaticMesh* FChaosVDGeometryBuilder::CreateAndCacheStaticMesh(const uint32 Geom
 {
 	{
 		FReadScopeLock ReadLock(RWLock);
-		//TODO: Make this return what is cached when the system is more robust
-		// For now this should not happen and we want to catch it and make it visually noticeable
-		if (StaticMeshCacheMap.Contains(GeometryCacheKey))
+		if (TObjectPtr<UStaticMesh>* StaticMeshPtrPtr = StaticMeshCacheMap.Find(GeometryCacheKey))
 		{
-			ensureMsgf(false, TEXT("Tried to create a new mesh with an existing Cache key"));
-			return nullptr;
+			return *StaticMeshPtrPtr;
 		}
 	}
 
@@ -201,17 +195,30 @@ void FChaosVDGeometryBuilder::ApplyMeshToComponentFromKey(TWeakObjectPtr<UMeshCo
 
 bool FChaosVDGeometryBuilder::GameThreadTick(float DeltaTime)
 {
+	if (bHasPendingJobs)
+	{
+		FReadScopeLock ReadLock(RWLock);
+		int32 PendingJobs = MeshComponentsWaitingForGeometryByKey.Num();
+		GeometryGenerationNotification.Update(PendingJobs);
+		
+		if (PendingJobs == 0)
+		{
+			bHasPendingJobs = false;
+		}
+	}
+	
 	int32 CurrentGeometryProcessed = 0;
 	while (!GeometryReadyToApplyQueue.IsEmpty() && CurrentGeometryProcessed < CVarChaosVDGeometryToProcessPerTick->GetInt())
 	{
 		uint32 GeometryKey = 0;
 		GeometryReadyToApplyQueue.Dequeue(GeometryKey);
-		CurrentGeometryProcessed++;
 
+		CurrentGeometryProcessed++;
+		
 		TArray<TWeakObjectPtr<UMeshComponent>>* MeshComponentsWaiting = nullptr;
 		{
 			FReadScopeLock ReadLock(RWLock);
-			 MeshComponentsWaiting = MeshComponentsWaitingForGeometryByKey.Find(GeometryKey);
+			MeshComponentsWaiting = MeshComponentsWaitingForGeometryByKey.Find(GeometryKey);
 		}
 
 		if (MeshComponentsWaiting)
@@ -227,7 +234,7 @@ bool FChaosVDGeometryBuilder::GameThreadTick(float DeltaTime)
 			MeshComponentsWaitingForGeometryByKey.Remove(GeometryKey);
 		}
 	}
-
+	
 	return true;
 }
 
@@ -254,6 +261,7 @@ void FChaosVDGeometryBuilder::RegisterMeshComponentWaitingForGeometry(uint32 Geo
 		else
 		{
 			MeshComponentsWaitingForGeometryByKey.Add(GeometryKey, {MesComponent});
+			bHasPendingJobs = true;
 		}
 	}
 }
diff --git a/Engine/Plugins/Experimental/ChaosVD/Source/ChaosVD/Private/ChaosVDScene.cpp b/Engine/Plugins/Experimental/ChaosVD/Source/ChaosVD/Private/ChaosVDScene.cpp
--- a/Engine/Plugins/Experimental/ChaosVD/Source/ChaosVD/Private/ChaosVDScene.cpp
+++ b/Engine/Plugins/Experimental/ChaosVD/Source/ChaosVD/Private/ChaosVDScene.cpp
@@ -17,13 +17,23 @@
 #include "Engine/Engine.h"
 #include "Engine/LevelStreamingDynamic.h"
 #include "Engine/World.h"
+#include "Misc/ScopedSlowTask.h"
 
 #include "UObject/Package.h"
 
+#define LOCTEXT_NAMESPACE "ChaosVisualDebugger"
+
 FChaosVDScene::FChaosVDScene() = default;
 
 FChaosVDScene::~FChaosVDScene() = default;
 
+namespace ChaosVDSceneUIOptions
+{
+	constexpr float DelayToShowProgressDialogThreshold = 1.5f;
+	constexpr bool bShowCancelButton = false;
+	constexpr bool bAllowInPIE = false;
+}
+
 void FChaosVDScene::Initialize()
 {
 	if (!ensure(!bIsInitialized))
@@ -83,48 +93,57 @@ void FChaosVDScene::UpdateFromRecordedStepData(const int32 SolverID, const FStri
 	
 	TSet<int32> ParticlesIDsInRecordedStepData;
 	ParticlesIDsInRecordedStepData.Reserve(InRecordedStepData.RecordedParticlesData.Num());
-
-	// Go over existing Particle VD Instances and update them or create them if needed 
-	for (const FChaosVDParticleDataWrapper& Particle : InRecordedStepData.RecordedParticlesData)
+	
 	{
-		const int32 ParticleVDInstanceID = GetIDForRecordedParticleData(Particle);
-		ParticlesIDsInRecordedStepData.Add(ParticleVDInstanceID);
-
-		if (InRecordedStepData.ParticlesDestroyedIDs.Contains(ParticleVDInstanceID))
+		constexpr float AmountOfWork = 1.0f;
+		const float PercentagePerElement = 1.0f / InRecordedStepData.RecordedParticlesData.Num();
+		FScopedSlowTask UpdatingSceneSlowTask(AmountOfWork, LOCTEXT("ProcessingParticleData", "Processing Particle Data..."));
+		UpdatingSceneSlowTask.MakeDialogDelayed(ChaosVDSceneUIOptions::DelayToShowProgressDialogThreshold, ChaosVDSceneUIOptions::bShowCancelButton, ChaosVDSceneUIOptions::bAllowInPIE);
+	
+		// Go over existing Particle VD Instances and update them or create them if needed 
+		for (const FChaosVDParticleDataWrapper& Particle : InRecordedStepData.RecordedParticlesData)
 		{
-			// Do not process the particle if it was destroyed in the same step
-			continue;
-		}
+			const int32 ParticleVDInstanceID = GetIDForRecordedParticleData(Particle);
+			ParticlesIDsInRecordedStepData.Add(ParticleVDInstanceID);
 
-		if (AChaosVDParticleActor** ExistingParticleVDInstancePtrPtr = SolverParticlesByID.Find(ParticleVDInstanceID))
-		{
-			if (AChaosVDParticleActor* ExistingParticleVDInstancePtr = *ExistingParticleVDInstancePtrPtr)
+			if (InRecordedStepData.ParticlesDestroyedIDs.Contains(ParticleVDInstanceID))
 			{
-				ExistingParticleVDInstancePtr->UpdateFromRecordedParticleData(Particle, InFrameData.SimulationTransform);
+				// Do not process the particle if it was destroyed in the same step
+				continue;
 			}
-			else
+
+			if (AChaosVDParticleActor** ExistingParticleVDInstancePtrPtr = SolverParticlesByID.Find(ParticleVDInstanceID))
 			{
-				//TODO: Handle this error
-				ensure(false);
+				if (AChaosVDParticleActor* ExistingParticleVDInstancePtr = *ExistingParticleVDInstancePtrPtr)
+				{
+					ExistingParticleVDInstancePtr->UpdateFromRecordedParticleData(Particle, InFrameData.SimulationTransform);
+				}
+				else
+				{
+					//TODO: Handle this error
+					ensure(false);
+				}
 			}
-		}
-		else
-		{
-			if (AChaosVDParticleActor* NewParticleVDInstance = SpawnParticleFromRecordedData(Particle, InFrameData))
+			else
 			{
-				FStringFormatOrderedArguments Args {SolverName, FString::FromInt(SolverID)};
-				const FName FolderPath = *FPaths::Combine(FString::Format(TEXT("Solver {0} | ID {1}"), Args), UEnum::GetDisplayValueAsText(NewParticleVDInstance->GetParticleData()->Type).ToString());
+				if (AChaosVDParticleActor* NewParticleVDInstance = SpawnParticleFromRecordedData(Particle, InFrameData))
+				{
+					FStringFormatOrderedArguments Args {SolverName, FString::FromInt(SolverID)};
+					const FName FolderPath = *FPaths::Combine(FString::Format(TEXT("Solver {0} | ID {1}"), Args), UEnum::GetDisplayValueAsText(NewParticleVDInstance->GetParticleData()->Type).ToString());
 
-				NewParticleVDInstance->SetFolderPath(FolderPath);
+					NewParticleVDInstance->SetFolderPath(FolderPath);
 
-				// TODO: Precalculate the max num of entries we would see in the loaded file, and use that number to pre-allocate this map
-				SolverParticlesByID.Add(ParticleVDInstanceID, NewParticleVDInstance);
-			}
-			else
-			{
-				//TODO: Handle this error
-				ensure(false);
+					// TODO: Precalculate the max num of entries we would see in the loaded file, and use that number to pre-allocate this map
+					SolverParticlesByID.Add(ParticleVDInstanceID, NewParticleVDInstance);
+				}
+				else
+				{
+					//TODO: Handle this error
+					ensure(false);
+				}
 			}
+			
+			UpdatingSceneSlowTask.EnterProgressFrame(PercentagePerElement);
 		}
 	}
 
@@ -157,30 +176,48 @@ void FChaosVDScene::UpdateFromRecordedStepData(const int32 SolverID, const FStri
 	OnSceneUpdated().Broadcast();
 }
 
-
 void FChaosVDScene::UpdateParticlesCollisionData(const FChaosVDStepData& InRecordedStepData, int32 SolverID)
 {
-	FChaosVDParticlesByIDMap& SolverParticlesByID = ParticlesBySolverID.FindChecked(SolverID);
+	constexpr float AmountOfWork = 1.0f;
 	
-	for (const TPair<int32, TArray<TSharedPtr<FChaosVDParticlePairMidPhase>>>& ParticleIDMidPhasePair : InRecordedStepData.RecordedMidPhasesByParticleID)
+	FChaosVDParticlesByIDMap& SolverParticlesByID = ParticlesBySolverID.FindChecked(SolverID);
 	{
-		if (AChaosVDParticleActor** ExistingParticleVDInstancePtrPtr = SolverParticlesByID.Find(ParticleIDMidPhasePair.Key))
+		FScopedSlowTask UpdatingSceneSlowTask(AmountOfWork, LOCTEXT("ProcessingMidphaseData", "Processing MidPhase Data..."));
+		UpdatingSceneSlowTask.MakeDialogDelayed(ChaosVDSceneUIOptions::DelayToShowProgressDialogThreshold,ChaosVDSceneUIOptions::bShowCancelButton, ChaosVDSceneUIOptions::bAllowInPIE);
+		
+		const float PercentagePerElement = 1.0f / InRecordedStepData.RecordedMidPhasesByParticleID.Num();
+		
+		for (const TPair<int32, TArray<TSharedPtr<FChaosVDParticlePairMidPhase>>>& ParticleIDMidPhasePair : InRecordedStepData.RecordedMidPhasesByParticleID)
 		{
-			if (AChaosVDParticleActor* ExistingParticleVDInstancePtr = *ExistingParticleVDInstancePtrPtr)
+			if (AChaosVDParticleActor** ExistingParticleVDInstancePtrPtr = SolverParticlesByID.Find(ParticleIDMidPhasePair.Key))
 			{
-				ExistingParticleVDInstancePtr->UpdateCollisionData(ParticleIDMidPhasePair.Value);
+				if (AChaosVDParticleActor* ExistingParticleVDInstancePtr = *ExistingParticleVDInstancePtrPtr)
+				{
+					ExistingParticleVDInstancePtr->UpdateCollisionData(ParticleIDMidPhasePair.Value);
+				}
 			}
+
+			UpdatingSceneSlowTask.EnterProgressFrame(PercentagePerElement);
 		}
 	}
 
-	for (const TPair<int32, TArray<FChaosVDConstraint>>& ParticleIDMidPhasePair : InRecordedStepData.RecordedConstraintsByParticleID)
 	{
-		if (AChaosVDParticleActor** ExistingParticleVDInstancePtrPtr = SolverParticlesByID.Find(ParticleIDMidPhasePair.Key))
+		FScopedSlowTask UpdatingSceneSlowTask(AmountOfWork, LOCTEXT("ProcessingConstraingtData", "Processing Constraint Data..."));
+		UpdatingSceneSlowTask.MakeDialogDelayed(ChaosVDSceneUIOptions::DelayToShowProgressDialogThreshold, ChaosVDSceneUIOptions::bShowCancelButton, ChaosVDSceneUIOptions::bAllowInPIE);
+		
+		const float PercentagePerElement = 1.0f / InRecordedStepData.RecordedMidPhasesByParticleID.Num();
+		
+		for (const TPair<int32, TArray<FChaosVDConstraint>>& ParticleIDMidPhasePair : InRecordedStepData.RecordedConstraintsByParticleID)
 		{
-			if (AChaosVDParticleActor* ExistingParticleVDInstancePtr = *ExistingParticleVDInstancePtrPtr)
+			if (AChaosVDParticleActor** ExistingParticleVDInstancePtrPtr = SolverParticlesByID.Find(ParticleIDMidPhasePair.Key))
 			{
-				ExistingParticleVDInstancePtr->UpdateCollisionData(ParticleIDMidPhasePair.Value);
+				if (AChaosVDParticleActor* ExistingParticleVDInstancePtr = *ExistingParticleVDInstancePtrPtr)
+				{
+					ExistingParticleVDInstancePtr->UpdateCollisionData(ParticleIDMidPhasePair.Value);
+				}
 			}
+
+			UpdatingSceneSlowTask.EnterProgressFrame(PercentagePerElement);
 		}
 	}
 }
@@ -453,3 +490,5 @@ bool FChaosVDScene::IsObjectSelected(const UObject* Object)
 
 	return SelectionSet->IsElementSelected(GetSelectionHandleForObject(Object), FTypedElementIsSelectedOptions());;
 }
+
+#undef LOCTEXT_NAMESPACE
diff --git a/Engine/Plugins/Experimental/ChaosVD/Source/ChaosVD/Private/ChaosVDTabsIDs.cpp b/Engine/Plugins/Experimental/ChaosVD/Source/ChaosVD/Private/ChaosVDTabsIDs.cpp
--- a/Engine/Plugins/Experimental/ChaosVD/Source/ChaosVD/Private/ChaosVDTabsIDs.cpp
+++ b/Engine/Plugins/Experimental/ChaosVD/Source/ChaosVD/Private/ChaosVDTabsIDs.cpp
@@ -9,4 +9,5 @@ const FName FChaosVDTabID::WorldOutliner = TEXT("WorldOutliner");
 const FName FChaosVDTabID::DetailsPanel = TEXT("DetailsPanel");
 const FName FChaosVDTabID::OutputLog = TEXT("ChaosVDOutputLog");
 const FName FChaosVDTabID::SolversTrack = TEXT("SolversTrack");
-const FName FChaosVDTabID::CVDEditorSettings = TEXT("CVDEditorSettings");
\ No newline at end of file
+const FName FChaosVDTabID::CVDEditorSettings = TEXT("CVDEditorSettings");
+const FName FChaosVDTabID::StatusBar = TEXT("StatusBar");
\ No newline at end of file
diff --git a/Engine/Plugins/Experimental/ChaosVD/Source/ChaosVD/Private/Widgets/SChaosVDMainTab.cpp b/Engine/Plugins/Experimental/ChaosVD/Source/ChaosVD/Private/Widgets/SChaosVDMainTab.cpp
--- a/Engine/Plugins/Experimental/ChaosVD/Source/ChaosVD/Private/Widgets/SChaosVDMainTab.cpp
+++ b/Engine/Plugins/Experimental/ChaosVD/Source/ChaosVD/Private/Widgets/SChaosVDMainTab.cpp
@@ -12,18 +12,22 @@
 #include "ChaosVDTabsIDs.h"
 #include "ChaosVDWorldOutlinerTab.h"
 #include "DesktopPlatformModule.h"
+#include "Editor.h"
 #include "Framework/Application/SlateApplication.h"
 #include "Framework/Application/SWindowTitleBar.h"
 #include "Framework/Commands/UICommandList.h"
 #include "Framework/MultiBox/MultiBoxBuilder.h"
 #include "IDesktopPlatform.h"
+#include "StatusBarSubsystem.h"
 #include "Styling/StyleColors.h"
+#include "Styling/ToolBarStyle.h"
 #include "Widgets/Docking/SDockTab.h"
 #include "Widgets/Input/SComboButton.h"
 #include "Widgets/SToolTip.h"
 
 #define LOCTEXT_NAMESPACE "ChaosVisualDebugger"
 
+
 void SChaosVDMainTab::Construct(const FArguments& InArgs, TSharedPtr<FChaosVDEngine> InChaosVDEngine)
 {
 	ChaosVDEngine = InChaosVDEngine;
@@ -39,6 +43,15 @@ void SChaosVDMainTab::Construct(const FArguments& InArgs, TSharedPtr<FChaosVDEng
 	SolversTracksTab = MakeShared<FChaosVDSolversTracksTab>(FChaosVDTabID::SolversTrack, TabManager, this);
 	EditorSettingsTab = MakeShared<FChaosVDEditorSettingsTab>(FChaosVDTabID::CVDEditorSettings, TabManager, this);
 
+	StatusBarID = FName(FChaosVDTabID::StatusBar.ToString() + InChaosVDEngine->GetInstanceGuid().ToString());
+	
+	UStatusBarSubsystem* StatusBarSubsystem = GEditor ? GEditor->GetEditorSubsystem<UStatusBarSubsystem>() : nullptr;
+	check(StatusBarSubsystem);
+	TSharedRef<SWidget> StatusBarWidget = StatusBarSubsystem->MakeStatusBarWidget(StatusBarID, TabManager->GetOwnerTab().ToSharedRef());
+
+	StatusBarSubsystem->UnregisterDrawer(StatusBarID, "ContentBrowser");
+	StatusBarSubsystem->UnregisterDrawer(StatusBarID, "OutputLog");
+	
 	GenerateMainWindowMenu();
 
 	ChildSlot
@@ -137,12 +150,18 @@ void SChaosVDMainTab::Construct(const FArguments& InArgs, TSharedPtr<FChaosVDEng
 			]
 				
 		]
+		
 		// Main Visual Debugger Interface content
 		+SVerticalBox::Slot()
 		.Padding(FMargin(0.0f,5.0f,0.0f,0.0f))
 		[
 			TabManager->RestoreFrom(GenerateMainLayout(), TabManager->GetOwnerTab()->GetParentWindow()).ToSharedRef()
-			
+		]
+		+SVerticalBox::Slot()
+		.Padding(0.0f, 2.0f, 0.0f, 0.0f)
+		.AutoHeight()
+		[
+			StatusBarWidget
 		]
 	];
 
diff --git a/Engine/Plugins/Experimental/ChaosVD/Source/ChaosVD/Public/ChaosVDGeometryBuilder.h b/Engine/Plugins/Experimental/ChaosVD/Source/ChaosVD/Public/ChaosVDGeometryBuilder.h
--- a/Engine/Plugins/Experimental/ChaosVD/Source/ChaosVD/Public/ChaosVDGeometryBuilder.h
+++ b/Engine/Plugins/Experimental/ChaosVD/Source/ChaosVD/Public/ChaosVDGeometryBuilder.h
@@ -2,6 +2,7 @@
 
 #pragma once
 
+#include "AsyncCompilationHelpers.h"
 #include "ChaosVDConvexMeshGenerator.h"
 #include "ChaosVDHeightfieldMeshGenerator.h"
 #include "ChaosVDTriMeshGenerator.h"
@@ -16,6 +17,7 @@
 #include "Generators/MinimalBoxMeshGenerator.h"
 #include "Generators/SphereGenerator.h"
 #include "UDynamicMesh.h"
+#include "Framework/Notifications/NotificationManager.h"
 #include "Tasks/Task.h"
 #include "UObject/GCObject.h"
 #include "UObject/UObjectGlobals.h"
@@ -40,7 +42,7 @@ class FChaosVDGeometryBuilder : public FGCObject, public TSharedFromThis<FChaosV
 {
 public:
 
-	FChaosVDGeometryBuilder()
+	FChaosVDGeometryBuilder() : GeometryGenerationNotification(NSLOCTEXT("ChaosVisualDebugger", "GeometryGenNotification","Particle Geometry"))
 	{
 		GameThreadTickDelegate = FTSTicker::GetCoreTicker().AddTicker(FTickerDelegate::CreateRaw(this, &FChaosVDGeometryBuilder::GameThreadTick));
 	}
@@ -183,6 +185,10 @@ private:
 	/** Queue of geometry keys already generated and waiting to be applied */
 	TQueue<uint32, EQueueMode::Mpsc> GeometryReadyToApplyQueue;
 
+	FAsyncCompilationNotification GeometryGenerationNotification;
+
+	FThreadSafeBool bHasPendingJobs = false;
+
 	friend class FGeometryGenerationTask;
 
 };
diff --git a/Engine/Plugins/Experimental/ChaosVD/Source/ChaosVD/Public/ChaosVDTabsIDs.h b/Engine/Plugins/Experimental/ChaosVD/Source/ChaosVD/Public/ChaosVDTabsIDs.h
--- a/Engine/Plugins/Experimental/ChaosVD/Source/ChaosVD/Public/ChaosVDTabsIDs.h
+++ b/Engine/Plugins/Experimental/ChaosVD/Source/ChaosVD/Public/ChaosVDTabsIDs.h
@@ -15,4 +15,5 @@ public:
 	static const FName OutputLog;
 	static const FName SolversTrack;
 	static const FName CVDEditorSettings;
+	static const FName StatusBar;
 };
diff --git a/Engine/Plugins/Experimental/ChaosVD/Source/ChaosVD/Public/Widgets/SChaosVDMainTab.h b/Engine/Plugins/Experimental/ChaosVD/Source/ChaosVD/Public/Widgets/SChaosVDMainTab.h
--- a/Engine/Plugins/Experimental/ChaosVD/Source/ChaosVD/Public/Widgets/SChaosVDMainTab.h
+++ b/Engine/Plugins/Experimental/ChaosVD/Source/ChaosVD/Public/Widgets/SChaosVDMainTab.h
@@ -27,6 +27,7 @@ public:
 	void Construct(const FArguments& InArgs, TSharedPtr<FChaosVDEngine> InChaosVDEngine);
 
 	TSharedRef<FChaosVDEngine> GetChaosVDEngineInstance() const { return ChaosVDEngine.ToSharedRef(); };
+	const FName& GetStatusBarID() const { return StatusBarID; };
 
 private:
 
@@ -46,5 +47,7 @@ private:
 	TSharedPtr<FChaosVDOutputLogTab> OutputLogTab;
 	TSharedPtr<FChaosVDEditorSettingsTab> EditorSettingsTab;
 
+	FName StatusBarID;
+
 	TSharedPtr<FTabManager> TabManager;
 };
