[ChaosVD] Added Record/Stop buttons to the UI.

Now we can start a recording from the tool itself instead of using the commandline arguments

diff --git a/Engine/Plugins/Experimental/ChaosVD/Source/ChaosVD/Private/ChaosVDStyle.cpp b/Engine/Plugins/Experimental/ChaosVD/Source/ChaosVD/Private/ChaosVDStyle.cpp
--- a/Engine/Plugins/Experimental/ChaosVD/Source/ChaosVD/Private/ChaosVDStyle.cpp
+++ b/Engine/Plugins/Experimental/ChaosVD/Source/ChaosVD/Private/ChaosVDStyle.cpp
@@ -36,6 +36,7 @@ FName FChaosVDStyle::GetStyleSetName()
 
 const FVector2D Icon16x16(16.0f, 16.0f);
 const FVector2D Icon20x20(20.0f, 20.0f);
+const FVector2D Icon24x24(24.0f, 24.0f);
 const FVector2D Icon50x50(50.0f, 50.0f);
 
 TSharedRef< FSlateStyleSet > FChaosVDStyle::Create()
@@ -57,6 +58,8 @@ TSharedRef< FSlateStyleSet > FChaosVDStyle::Create()
 	
 	Style->Set("LockIcon", new IMAGE_BRUSH_SVG(TEXT("lock"), Icon16x16));
 	Style->Set("UnlockedIcon", new IMAGE_BRUSH_SVG(TEXT("lock-unlocked"), Icon16x16));
+	Style->Set("RecordIcon", new IMAGE_BRUSH(TEXT("CircleBox2"), Icon16x16));
+	Style->Set("ConnectionIcon", new IMAGE_BRUSH_SVG(TEXT("Connection"), Icon16x16));
 
 	return Style;
 }
diff --git a/Engine/Plugins/Experimental/ChaosVD/Source/ChaosVD/Private/Widgets/SChaosVDMainTab.cpp b/Engine/Plugins/Experimental/ChaosVD/Source/ChaosVD/Private/Widgets/SChaosVDMainTab.cpp
--- a/Engine/Plugins/Experimental/ChaosVD/Source/ChaosVD/Private/Widgets/SChaosVDMainTab.cpp
+++ b/Engine/Plugins/Experimental/ChaosVD/Source/ChaosVD/Private/Widgets/SChaosVDMainTab.cpp
@@ -9,6 +9,7 @@
 #include "ChaosVDPlaybackController.h"
 #include "ChaosVDPlaybackViewportTab.h"
 #include "ChaosVDSolversTracksTab.h"
+#include "ChaosVDStyle.h"
 #include "ChaosVDTabsIDs.h"
 #include "ChaosVDWorldOutlinerTab.h"
 #include "DesktopPlatformModule.h"
@@ -21,8 +22,10 @@
 #include "StatusBarSubsystem.h"
 #include "Styling/StyleColors.h"
 #include "Styling/ToolBarStyle.h"
+#include "Widgets/SChaosVDRecordingControls.h"
 #include "Widgets/Docking/SDockTab.h"
 #include "Widgets/Input/SComboButton.h"
+#include "Widgets/Layout/SSeparator.h"
 #include "Widgets/SToolTip.h"
 
 #define LOCTEXT_NAMESPACE "ChaosVisualDebugger"
@@ -49,6 +52,7 @@ void SChaosVDMainTab::Construct(const FArguments& InArgs, TSharedPtr<FChaosVDEng
 	check(StatusBarSubsystem);
 	TSharedRef<SWidget> StatusBarWidget = StatusBarSubsystem->MakeStatusBarWidget(StatusBarID, TabManager->GetOwnerTab().ToSharedRef());
 
+	// Status bars come with the output log and content browser drawers by default, therefore we need to remove them otherwise they will be on the tool's window
 	StatusBarSubsystem->UnregisterDrawer(StatusBarID, "ContentBrowser");
 	StatusBarSubsystem->UnregisterDrawer(StatusBarID, "OutputLog");
 	
@@ -70,7 +74,7 @@ void SChaosVDMainTab::Construct(const FArguments& InArgs, TSharedPtr<FChaosVDEng
 				+SHorizontalBox::Slot()
 				.VAlign(VAlign_Center)
 				.HAlign(HAlign_Left)
-				.Padding(FMargin(12, 7, 18, 7))
+				.Padding(FMargin(12, 7, 6, 7))
 				.AutoWidth()
 				[
 					SNew(SButton)
@@ -84,7 +88,7 @@ void SChaosVDMainTab::Construct(const FArguments& InArgs, TSharedPtr<FChaosVDEng
 					.Content()
 					[
 						SNew(SHorizontalBox)
-						+ SHorizontalBox::Slot()
+						+SHorizontalBox::Slot()
 						.HAlign(HAlign_Center)
 						.VAlign(VAlign_Center)
 						[
@@ -92,7 +96,7 @@ void SChaosVDMainTab::Construct(const FArguments& InArgs, TSharedPtr<FChaosVDEng
 							.Image(FAppStyle::Get().GetBrush("Icons.Plus"))
 							.ColorAndOpacity(FStyleColors::AccentGreen)
 						]
-						+ SHorizontalBox::Slot()
+						+SHorizontalBox::Slot()
 						.Padding(FMargin(3, 0, 0, 0))
 						.VAlign(VAlign_Center)
 						.AutoWidth()
@@ -104,8 +108,73 @@ void SChaosVDMainTab::Construct(const FArguments& InArgs, TSharedPtr<FChaosVDEng
 					]
 				]
 
-				+ SHorizontalBox::Slot()
-				.Padding(0, 7, 0, 7)
+				// Remote Connection Button
+				+SHorizontalBox::Slot()
+				.VAlign(VAlign_Center)
+				.HAlign(HAlign_Left)
+				.Padding(FMargin(6, 7, 18, 7))
+				.AutoWidth()
+				[
+					SNew(SButton)
+					.ToolTip(SNew(SToolTip).Text(LOCTEXT("RemoteConnectionDesc", "Not Supported yet - Click here to connect to a remote debug session.")))
+					.ContentPadding(FMargin(0, 5.f, 0, 4.f))
+					.IsEnabled(false) // We don't support remote debugging yet
+					.OnClicked_Lambda([this]()
+					{
+						//TODO : Add a call to the method that will open the Trace Session Browser
+						return FReply::Handled();
+					})
+					.Content()
+					[
+						SNew(SHorizontalBox)
+						+SHorizontalBox::Slot()
+						.HAlign(HAlign_Center)
+						.VAlign(VAlign_Center)
+						[
+							SNew(SImage)
+							.Image(FChaosVDStyle::Get().GetBrush("ConnectionIcon"))
+							.ColorAndOpacity(FStyleColors::AccentGreen)
+						]
+						+SHorizontalBox::Slot()
+						.Padding(FMargin(3, 0, 0, 0))
+						.VAlign(VAlign_Center)
+						.AutoWidth()
+						[
+							SNew(STextBlock)
+							.TextStyle(FAppStyle::Get(), "SmallButtonText")
+							.Text(LOCTEXT("ConnectToSession", "Connect to Session"))
+						]
+					]
+				]
+
+
+				+SHorizontalBox::Slot()
+				.AutoWidth()
+				.Padding(1.0f)
+				[
+					SNew(SSeparator)
+						.Orientation(Orient_Vertical)
+						.Thickness(2.0f)
+						.ColorAndOpacity(FColor::Black)
+						.SeparatorImage(FAppStyle::Get().GetBrush("Menu.Separator"))
+				]
+				
+				+SHorizontalBox::Slot()
+				[
+					
+					SNew(SChaosVDRecordingControls, StaticCastWeakPtr<SChaosVDMainTab>(AsWeak()))
+				]
+				
+				+SHorizontalBox::Slot()
+				.AutoWidth()
+				.Padding(1.0f, 5.0f, 1.0f, 5.0f)
+				[
+					SNew(SSeparator)
+						.Orientation(Orient_Vertical)
+						.Thickness(2.0f)
+						.ColorAndOpacity(FColor::Black)
+						.SeparatorImage(FAppStyle::Get().GetBrush("Menu.Separator"))
+				]
 
 				// Settings button
 				+ SHorizontalBox::Slot()
@@ -117,6 +186,7 @@ void SChaosVDMainTab::Construct(const FArguments& InArgs, TSharedPtr<FChaosVDEng
 					SNew(SComboButton)
 					.ContentPadding(0)
 					.HasDownArrow(false)
+					.IsEnabled(false) // Disabled until we create content for it
 					.ForegroundColor(FSlateColor::UseForeground())
 					.ComboButtonStyle(FAppStyle::Get(), "SimpleComboButton")
 					.AddMetaData<FTagMetaData>(FTagMetaData(TEXT("ViewOptions")))
diff --git a/Engine/Plugins/Experimental/ChaosVD/Source/ChaosVD/Private/Widgets/SChaosVDRecordingControls.cpp b/Engine/Plugins/Experimental/ChaosVD/Source/ChaosVD/Private/Widgets/SChaosVDRecordingControls.cpp
new file mode 100644
--- /dev/null
+++ b/Engine/Plugins/Experimental/ChaosVD/Source/ChaosVD/Private/Widgets/SChaosVDRecordingControls.cpp
@@ -0,0 +1,157 @@
+// Copyright Epic Games, Inc. All Rights Reserved.
+
+#include "Widgets/SChaosVDRecordingControls.h"
+
+#include "ChaosVDEngine.h"
+#include "ChaosVDStyle.h"
+#include "ChaosVDRuntimeModule.h"
+#include "ChaosVisualDebugger/ChaosVisualDebuggerTrace.h"
+#include "Editor.h"
+#include "Input/Reply.h"
+#include "Misc/MessageDialog.h"
+#include "StatusBarSubsystem.h"
+#include "Styling/SlateBrush.h"
+#include "Widgets/Images/SImage.h"
+#include "Widgets/Input/SButton.h"
+#include "Widgets/SChaosVDMainTab.h"
+#include "Widgets/Text/STextBlock.h"
+
+#define LOCTEXT_NAMESPACE "ChaosVisualDebugger"
+
+void SChaosVDRecordingControls::Construct(const FArguments& InArgs, const TWeakPtr<SChaosVDMainTab>& InMainTabWeakPtr)
+{
+	if (TSharedPtr<SChaosVDMainTab> MainTabSharedPtr = InMainTabWeakPtr.Pin())
+	{
+		MainTabWeakPtr = MainTabSharedPtr;
+		StatusBarID = MainTabSharedPtr->GetStatusBarID();
+	}
+	else
+	{
+		ensureMsgf(false, TEXT("Constructed with an invalid Main Tab"));
+	}
+	
+	ChildSlot
+	[
+		SNew(SHorizontalBox)
+		+SHorizontalBox::Slot()
+		.HAlign(HAlign_Left)
+		.Padding(FMargin(12, 7, 2, 7))
+		[
+			SNew(SHorizontalBox)
+			+SHorizontalBox::Slot()
+			.HAlign(HAlign_Left)
+			[
+				SNew(SButton)
+					.OnClicked(FOnClicked::CreateRaw(this, &SChaosVDRecordingControls::ToggleRecordingState))
+					.ForegroundColor(FSlateColor::UseForeground() )
+					.IsFocusable(false)
+					.VAlign(VAlign_Center)
+					.HAlign(HAlign_Center)
+					[
+						SNew( SImage )
+						.ToolTipText_Lambda([this]()
+						{
+							return IsRecording() ? LOCTEXT("StopRecordButtonDesc", "Stop the current recording ")
+													: LOCTEXT("RecordButtonDesc", "Starts a recording for the current sessions");
+						})
+						.Image_Raw(this, &SChaosVDRecordingControls::GetRecordOrStopButton)
+						.ColorAndOpacity(FColor::Red)
+					]
+					]
+
+			+SHorizontalBox::Slot()
+			.HAlign(HAlign_Left)
+			.VAlign(VAlign_Center)
+			.AutoWidth()
+			.Padding(FMargin(12, 0, 0, 0))
+			[
+				SNew( STextBlock )
+					.TextStyle(FAppStyle::Get(), "SmallButtonText")
+					.Text_Raw(this, &SChaosVDRecordingControls::GetRecordingTimeText)
+					.ColorAndOpacity(FColor::White)
+			]
+		]
+	];
+
+	RecordingStartedHandle = FChaosVDRuntimeModule::Get().RegisterRecordingStartedCallback(FChaosVDRecordingStateChangedDelegate::FDelegate::CreateRaw(this, &SChaosVDRecordingControls::HandleRecordingStart));
+	RecordingStoppedHandle = FChaosVDRuntimeModule::Get().RegisterRecordingStopCallback(FChaosVDRecordingStateChangedDelegate::FDelegate::CreateRaw(this, &SChaosVDRecordingControls::HandleRecordingStop));
+}
+
+SChaosVDRecordingControls::~SChaosVDRecordingControls()
+{
+	 FChaosVDRuntimeModule::Get().RemoveRecordingStartedCallback(RecordingStartedHandle);
+	 FChaosVDRuntimeModule::Get().RemoveRecordingStopCallback(RecordingStoppedHandle);
+}
+
+const FSlateBrush* SChaosVDRecordingControls::GetRecordOrStopButton() const
+{
+	return IsRecording() ? FChaosVDStyle::Get().GetBrush("StopIcon") : FChaosVDStyle::Get().GetBrush("RecordIcon");
+}
+
+void SChaosVDRecordingControls::HandleRecordingStop()
+{
+	UStatusBarSubsystem* StatusBarSubsystem = GEditor ? GEditor->GetEditorSubsystem<UStatusBarSubsystem>() : nullptr;
+	check(StatusBarSubsystem);
+	StatusBarSubsystem->PopStatusBarMessage(StatusBarID, RecordingMessageHandle);
+
+	FText RecordingPathMessage = FText::Format(LOCTEXT("RecordingSavedPathMessage"," Recoring saved at {0} "), FText::AsCultureInvariant(FChaosVDRuntimeModule::Get().GetActiveRecordingFileName()));
+	RecordingPathMessageHandle = StatusBarSubsystem->PushStatusBarMessage(StatusBarID, RecordingPathMessage);
+
+	if (FMessageDialog::Open(EAppMsgType::YesNo, LOCTEXT("OpenLastRecordingMessage", "Do you want to load the recorded file now? ")) == EAppReturnType::Yes)
+	{
+		if (TSharedPtr<SChaosVDMainTab> MainTabSharedPtr = MainTabWeakPtr.Pin())
+		{
+			MainTabSharedPtr->GetChaosVDEngineInstance()->LoadRecording(FChaosVDRuntimeModule::Get().GetActiveRecordingFileName());
+		}
+	}
+}
+
+void SChaosVDRecordingControls::HandleRecordingStart()
+{
+	UStatusBarSubsystem* StatusBarSubsystem = GEditor ? GEditor->GetEditorSubsystem<UStatusBarSubsystem>() : nullptr;
+	if (!ensure(StatusBarSubsystem))
+	{
+		return;
+	}
+	
+	if (RecordingPathMessageHandle.IsValid())
+	{
+		StatusBarSubsystem->PopStatusBarMessage(StatusBarID, RecordingPathMessageHandle);
+	}
+
+	RecordingMessageHandle = StatusBarSubsystem->PushStatusBarMessage(StatusBarID, LOCTEXT("RecordingMessgae", "Recording..."));
+}
+
+FReply SChaosVDRecordingControls::ToggleRecordingState()
+{
+	if (!IsRecording())
+	{
+		FChaosVDRuntimeModule::Get().StartRecording({});
+	}
+	else
+	{
+		FChaosVDRuntimeModule::Get().StopRecording();
+	}
+
+	return FReply::Handled();
+}
+
+bool SChaosVDRecordingControls::IsRecording() const
+{
+#if WITH_CHAOS_VISUAL_DEBUGGER
+	return FChaosVisualDebuggerTrace::IsTracing();
+#endif
+	return false;
+}
+
+FText SChaosVDRecordingControls::GetRecordingTimeText() const
+{
+	FNumberFormattingOptions FormatOptions;
+	FormatOptions.MinimumFractionalDigits = 2;
+	FormatOptions.MaximumFractionalDigits = 2;
+	FText SecondsText = FText::AsNumber(FChaosVDRuntimeModule::Get().GetAccumulatedRecordingTime(), &FormatOptions);
+	
+	return FText::Format(LOCTEXT("RecordingTimer","{0} s"), SecondsText);
+}
+
+#undef LOCTEXT_NAMESPACE 
diff --git a/Engine/Plugins/Experimental/ChaosVD/Source/ChaosVD/Public/Widgets/SChaosVDRecordingControls.h b/Engine/Plugins/Experimental/ChaosVD/Source/ChaosVD/Public/Widgets/SChaosVDRecordingControls.h
new file mode 100644
--- /dev/null
+++ b/Engine/Plugins/Experimental/ChaosVD/Source/ChaosVD/Public/Widgets/SChaosVDRecordingControls.h
@@ -0,0 +1,44 @@
+// Copyright Epic Games, Inc. All Rights Reserved.
+
+#pragma once
+
+#include "StatusBarSubsystem.h"
+#include "Templates/SharedPointer.h"
+#include "Widgets/SCompoundWidget.h"
+
+class SChaosVDMainTab;
+class FReply;
+struct FSlateBrush;
+
+class SChaosVDRecordingControls : public SCompoundWidget
+{
+public:
+	SLATE_BEGIN_ARGS( SChaosVDRecordingControls ){}
+	SLATE_END_ARGS()
+	
+	void Construct(const FArguments& InArgs, const TWeakPtr<SChaosVDMainTab>& InMainTabWeakPtr);
+
+	virtual ~SChaosVDRecordingControls() override;
+
+protected:
+	
+	const FSlateBrush* GetRecordOrStopButton() const;
+	
+	void HandleRecordingStop();
+	void HandleRecordingStart();
+
+	FReply ToggleRecordingState();
+
+	bool IsRecording() const;
+
+	FText GetRecordingTimeText() const;
+	
+	FName StatusBarID;
+	
+	FStatusBarMessageHandle RecordingMessageHandle;
+	FStatusBarMessageHandle RecordingPathMessageHandle;
+	FDelegateHandle RecordingStartedHandle;
+	FDelegateHandle RecordingStoppedHandle;
+
+	TWeakPtr<SChaosVDMainTab> MainTabWeakPtr;
+};
diff --git a/Engine/Source/Runtime/Experimental/Chaos/Public/ChaosVisualDebugger/ChaosVisualDebuggerTrace.h b/Engine/Source/Runtime/Experimental/Chaos/Public/ChaosVisualDebugger/ChaosVisualDebuggerTrace.h
--- a/Engine/Source/Runtime/Experimental/Chaos/Public/ChaosVisualDebugger/ChaosVisualDebuggerTrace.h
+++ b/Engine/Source/Runtime/Experimental/Chaos/Public/ChaosVisualDebugger/ChaosVisualDebuggerTrace.h
@@ -398,7 +398,7 @@ public:
 	static FString GetDebugName(T& ObjectWithDebugName);
 
 	/** Returns true if a CVD trace is running */
-	static bool IsTracing() { return bIsTracing;}
+	static CHAOS_API bool IsTracing() { return bIsTracing;}
 
 private:
 
diff --git a/Engine/Source/Runtime/Experimental/ChaosVisualDebugger/Private/ChaosVDRuntimeModule.cpp b/Engine/Source/Runtime/Experimental/ChaosVisualDebugger/Private/ChaosVDRuntimeModule.cpp
--- a/Engine/Source/Runtime/Experimental/ChaosVisualDebugger/Private/ChaosVDRuntimeModule.cpp
+++ b/Engine/Source/Runtime/Experimental/ChaosVisualDebugger/Private/ChaosVDRuntimeModule.cpp
@@ -92,6 +92,16 @@ bool FChaosVDRuntimeModule::RequestFullCapture(float DeltaTime)
 	return true;
 }
 
+bool FChaosVDRuntimeModule::RecordingTimerTick(float DeltaTime)
+{
+	if (bIsRecording)
+	{
+		AccumulatedRecordingTime += DeltaTime;
+	}
+	
+	return true;
+}
+
 void FChaosVDRuntimeModule::StartRecording(const TArray<FString>& Args)
 {
 	if (bIsRecording)
@@ -114,6 +124,8 @@ void FChaosVDRuntimeModule::StartRecording(const TArray<FString>& Args)
 	FullCaptureRequesterHandle = FTSTicker::GetCoreTicker().AddTicker(FTickerDelegate::CreateRaw(this, &FChaosVDRuntimeModule::RequestFullCapture),
 		FMath::Clamp(ConfiguredTimeBetweenCaptures, MinAllowedTimeInSecondsBetweenCaptures, TNumericLimits<int32>::Max()));
 
+	RecordingTimerHandle = FTSTicker::GetCoreTicker().AddTicker(FTickerDelegate::CreateRaw(this, &FChaosVDRuntimeModule::RecordingTimerTick));
+
 #if UE_TRACE_ENABLED
 
 	// Other tools could bee using trace
@@ -140,10 +152,10 @@ void FChaosVDRuntimeModule::StartRecording(const TArray<FString>& Args)
 
 	if (Args.Num() == 0 || Args[0] == TEXT("File"))
 	{
-		FString RecordingFileName;
-		GenerateRecordingFileName(RecordingFileName);
+		ActiveRecordingFileName.Empty();
+		GenerateRecordingFileName(ActiveRecordingFileName);
 
-		bIsRecording = FTraceAuxiliary::Start(FTraceAuxiliary::EConnectionType::File, *RecordingFileName);
+		bIsRecording = FTraceAuxiliary::Start(FTraceAuxiliary::EConnectionType::File, *ActiveRecordingFileName);
 	}
 	else if(Args[0] == TEXT("Server"))
 	{
@@ -155,7 +167,9 @@ void FChaosVDRuntimeModule::StartRecording(const TArray<FString>& Args)
 		nullptr);
 	}
 #endif
-
+	
+	AccumulatedRecordingTime = 0.0f;
+	
 	ensure(bIsRecording);
 }
 
@@ -186,6 +200,7 @@ void FChaosVDRuntimeModule::StopRecording()
 	}
 
 	bIsRecording = false;
+	AccumulatedRecordingTime = 0.0f;
 }
 
 void FChaosVDRuntimeModule::HandleTraceStopRequest(FTraceAuxiliary::EConnectionType TraceType, const FString& TraceDestination)
diff --git a/Engine/Source/Runtime/Experimental/ChaosVisualDebugger/Public/ChaosVDRuntimeModule.h b/Engine/Source/Runtime/Experimental/ChaosVisualDebugger/Public/ChaosVDRuntimeModule.h
--- a/Engine/Source/Runtime/Experimental/ChaosVisualDebugger/Public/ChaosVDRuntimeModule.h
+++ b/Engine/Source/Runtime/Experimental/ChaosVisualDebugger/Public/ChaosVDRuntimeModule.h
@@ -34,7 +34,8 @@ public:
 	 * @param Args : Arguments array provided by the commandline. Used to determine if we want to record to file or a local trace server
 	 */
 	void StartRecording(const TArray<FString>& Args);
-	/**/
+	
+	/* Stops an active recording */
 	void StopRecording();
 
 	/** Returns true if we are currently recording a Physics simulation */
@@ -79,6 +80,12 @@ public:
 		return PerformFullCaptureDelegate.Remove(InDelegateToRemove);
 	}
 
+	/** Returns the accumulated recording time in seconds since the recording started */
+	float GetAccumulatedRecordingTime() const { return AccumulatedRecordingTime; }
+
+	/** Returns the relative file name path of the active recording */
+	const FString& GetActiveRecordingFileName() const { return ActiveRecordingFileName; }
+
 private:
 
 	/** Stops the current Trace session */
@@ -89,6 +96,9 @@ private:
 	/** Queues a full Capture of the simulation on the next frame */
 	bool RequestFullCapture(float DeltaTime);
 
+	/** Queues a full Capture of the simulation on the next frame */
+	bool RecordingTimerTick(float DeltaTime);
+
 	/** Used to handle stop requests to the active trace session that were not done by us
 	 * That is a possible scenario because Trace is shared by other In-Editor tools
 	 */
@@ -97,7 +107,10 @@ private:
 	bool bIsRecording = false;
 	bool bRequestedStop = false;
 
+	float AccumulatedRecordingTime = 0.0f;
+
 	FTSTicker::FDelegateHandle FullCaptureRequesterHandle;
+	FTSTicker::FDelegateHandle RecordingTimerHandle;
 
 	static FChaosVDRecordingStateChangedDelegate RecordingStartedDelegate;
 	static FChaosVDRecordingStateChangedDelegate RecordingStopDelegate;
@@ -105,5 +118,7 @@ private:
 
 	FThreadSafeCounter LastGeneratedID;
 
+	FString ActiveRecordingFileName;
+
 	static FRWLock DelegatesRWLock;
 };
