[Chaos VD] Added In tool object tracking and Locations/Transform recording support

Now you can record any location or transform at any given frame using two new macros, and when the recording is loaded, they can be selected in the tool.

This can be used to mark Point of interests, or a more immediate use case, to record the Player's camera Transform (or Players positions), then you can use the Tracking feature in the tool to track the recorded camera transform.

Current known issues:
As recorded Locations/Transform are recorded as game frame data, they are only tracked when you are playin back from the GameFrames timeline, not a solver timeline. That will be addressed in a future CL

diff --git a/Engine/Plugins/Experimental/ChaosVD/Source/ChaosVD/Private/ChaosVDEditorSettingsTab.cpp b/Engine/Plugins/Experimental/ChaosVD/Source/ChaosVD/Private/ChaosVDEditorSettingsTab.cpp
deleted file mode 100644
diff --git a/Engine/Plugins/Experimental/ChaosVD/Source/ChaosVD/Private/ChaosVDEditorVisualizationSettingsTab.cpp b/Engine/Plugins/Experimental/ChaosVD/Source/ChaosVD/Private/ChaosVDEditorVisualizationSettingsTab.cpp
new file mode 100644
--- /dev/null
+++ b/Engine/Plugins/Experimental/ChaosVD/Source/ChaosVD/Private/ChaosVDEditorVisualizationSettingsTab.cpp
@@ -0,0 +1,75 @@
+﻿// Copyright Epic Games, Inc. All Rights Reserved.
+
+#include "ChaosVDEditorVisualizationSettingsTab.h"
+
+#include "ChaosVDEditorSettings.h"
+#include "ChaosVDEngine.h"
+#include "ChaosVDPlaybackController.h"
+#include "ChaosVDPlaybackViewportTab.h"
+#include "ChaosVDStyle.h"
+#include "DetailsViewArgs.h"
+#include "Framework/Docking/TabManager.h"
+#include "IDetailsView.h"
+#include "Modules/ModuleManager.h"
+#include "PropertyEditorModule.h"
+#include "Templates/SharedPointer.h"
+#include "Widgets/SChaosVDVisualizationControls.h"
+#include "Widgets/Docking/SDockTab.h"
+#include "Widgets/Text/STextBlock.h"
+
+#define LOCTEXT_NAMESPACE "ChaosVisualDebugger"
+
+TSharedRef<SDockTab> FChaosVDEditorVisualizationSettingsTab::HandleTabSpawned(const FSpawnTabArgs& Args)
+{
+
+	TSharedPtr<SChaosVDPlaybackViewport> ViewportWidget;
+	if (TSharedPtr<FChaosVDPlaybackViewportTab> ViewportTabSharedPtr = OwningTabWidget->GetPlaybackViewportTab().Pin())
+	{
+		ViewportWidget = ViewportTabSharedPtr->GetPlaybackViewportWidget().Pin();
+	}
+	
+	FPropertyEditorModule& PropertyEditorModule = FModuleManager::GetModuleChecked<FPropertyEditorModule>("PropertyEditor");
+	FDetailsViewArgs DetailsViewArgs;
+	DetailsViewArgs.bAllowSearch = false;
+	DetailsViewArgs.NameAreaSettings = FDetailsViewArgs::HideNameArea;
+
+	TSharedRef<IDetailsView> DetailsPanel = PropertyEditorModule.CreateDetailView(DetailsViewArgs);
+
+	DetailsPanel->SetObject(GetMutableDefault<UChaosVDEditorSettings>());
+
+	TSharedRef<SDockTab> DetailsPanelTab =
+		SNew(SDockTab)
+		.TabRole(ETabRole::MajorTab)
+		.Label(LOCTEXT("ChaosVDEditorSettings", "Chaos VD Settings"))
+		.ToolTipText(LOCTEXT("ChaosVDEditorSettingsTip", "See the available settings for the editor"));
+
+	if (ensure(ViewportWidget.IsValid()))
+	{
+		DetailsPanelTab->SetContent
+		(
+			SNew(SVerticalBox)
+			+SVerticalBox::Slot()
+			[
+				SNew(SChaosVDVisualizationControls, OwningTabWidget->GetChaosVDEngineInstance()->GetPlaybackController(), ViewportWidget)
+			]	
+		);
+	}
+	else
+	{
+		DetailsPanelTab->SetContent
+		(
+			SNew(SVerticalBox)
+			+SVerticalBox::Slot()
+			[
+				SNew(STextBlock)
+				.Text(LOCTEXT("ChaosVDEditorSettingsLoadError", "Failed to load Visualization Controls"))
+			]	
+		);
+	}
+
+	DetailsPanelTab->SetTabIcon(FChaosVDStyle::Get().GetBrush("TabIconDetailsPanel"));
+
+	return DetailsPanelTab;
+}
+
+#undef LOCTEXT_NAMESPACE 
diff --git a/Engine/Plugins/Experimental/ChaosVD/Source/ChaosVD/Private/ChaosVDPlaybackViewportClient.cpp b/Engine/Plugins/Experimental/ChaosVD/Source/ChaosVD/Private/ChaosVDPlaybackViewportClient.cpp
--- a/Engine/Plugins/Experimental/ChaosVD/Source/ChaosVD/Private/ChaosVDPlaybackViewportClient.cpp
+++ b/Engine/Plugins/Experimental/ChaosVD/Source/ChaosVD/Private/ChaosVDPlaybackViewportClient.cpp
@@ -2,11 +2,14 @@
 
 #include "ChaosVDPlaybackViewportClient.h"
 
+#include "CameraController.h"
+#include "ChaosVDEditorSettings.h"
 #include "ChaosVDParticleActor.h"
 #include "ChaosVDScene.h"
 #include "EngineUtils.h"
 #include "Elements/Framework/TypedElementSelectionSet.h"
 #include "Visualizers/ChaosVDDebugDrawUtils.h"
+#include "SEditorViewport.h"
 
 
 FChaosVDPlaybackViewportClient::FChaosVDPlaybackViewportClient() : FEditorViewportClient(nullptr), CVDWorld(nullptr)
@@ -63,6 +66,93 @@ void FChaosVDPlaybackViewportClient::HandleObjectFocused(UObject* FocusedObject)
 	}
 }
 
+void FChaosVDPlaybackViewportClient::TrackActor(AActor* ActorToTrack, EChaosVDActorTrackingMode TrackingMode)
+{
+	if (!ActorToTrack)
+	{
+		return;
+	}
+	
+	if (const UChaosVDEditorSettings* CVDEditorSettings = GetDefault<UChaosVDEditorSettings>())
+	{
+		switch(TrackingMode)
+		{
+			case EChaosVDActorTrackingMode::ByBoundingBox:
+				{
+					FBox ActorBounds = ActorToTrack->GetComponentsBoundingBox(false);
+					FocusViewportOnBox(ActorBounds.ExpandBy(CVDEditorSettings->ExpandViewTrackingBy), true);
+					break;
+				}
+			case EChaosVDActorTrackingMode::ByDistanceOffset:
+			case EChaosVDActorTrackingMode::MatchTransform:
+				{
+					TrackTransform(ActorToTrack->GetActorTransform(), TrackingMode);
+					break;
+				}
+			default:
+				{
+					ensureMsgf(false, TEXT("Actor tracking requested with invalid options. The actor will not be tracked"));
+					break;
+				}
+		}	
+	}
+}
+
+void FChaosVDPlaybackViewportClient::TrackTransform(const FTransform& TransformToTrack, EChaosVDActorTrackingMode TrackingMode)
+{
+	if (const UChaosVDEditorSettings* CVDEditorSettings = GetDefault<UChaosVDEditorSettings>())
+	{
+		switch(TrackingMode)
+		{
+		case EChaosVDActorTrackingMode::ByBoundingBox:
+			{
+				ensureMsgf(false, TEXT("Tracking by Bounding box only supported with Actors"));
+				break;
+			}
+		case EChaosVDActorTrackingMode::ByDistanceOffset:
+			{
+				FViewportCameraTransform& ViewTransform = GetViewTransform();
+				FVector ActorLocation = TransformToTrack.GetLocation();
+					
+				constexpr bool bEnable = false;
+				ToggleOrbitCamera(bEnable);
+
+				FVector ActorToCamDir = ViewTransform.GetLocation() - ActorLocation;
+				ActorToCamDir.Normalize();
+
+				FVector TargetLocation = ActorLocation + ActorToCamDir * CVDEditorSettings->TrackingDistanceOffset;
+
+				ViewTransform.SetRotation((-ActorToCamDir).Rotation());
+				ViewTransform.SetLookAt(ActorLocation);
+				ViewTransform.TransitionToLocation(TargetLocation, EditorViewportWidget, true);
+					
+				// Tell the viewport to redraw itself.
+				Invalidate();
+				break;
+			}
+		case EChaosVDActorTrackingMode::MatchTransform:
+			{
+				constexpr bool bEnable = false;
+				ToggleOrbitCamera(bEnable);
+					
+				FViewportCameraTransform& ViewTransform = GetViewTransform();
+				ViewTransform.SetRotation(TransformToTrack.GetRotation().Rotator());
+
+				ViewTransform.SetLookAt(TransformToTrack.GetLocation());
+				ViewTransform.TransitionToLocation(TransformToTrack.GetLocation(), EditorViewportWidget, true);
+
+				Invalidate();
+				break;
+			}		
+		default:
+			{
+				ensureMsgf(false, TEXT("Actor tracking requested with invalid options. The actor will not be tracked"));
+				break;
+			}
+		}	
+	}
+}
+
 void FChaosVDPlaybackViewportClient::Draw(const FSceneView* View, FPrimitiveDrawInterface* PDI)
 {
 	FEditorViewportClient::Draw(View, PDI);
diff --git a/Engine/Plugins/Experimental/ChaosVD/Source/ChaosVD/Private/Trace/ChaosVDTraceAnalyzer.cpp b/Engine/Plugins/Experimental/ChaosVD/Source/ChaosVD/Private/Trace/ChaosVDTraceAnalyzer.cpp
--- a/Engine/Plugins/Experimental/ChaosVD/Source/ChaosVD/Private/Trace/ChaosVDTraceAnalyzer.cpp
+++ b/Engine/Plugins/Experimental/ChaosVD/Source/ChaosVD/Private/Trace/ChaosVDTraceAnalyzer.cpp
@@ -22,6 +22,9 @@ void FChaosVDTraceAnalyzer::OnAnalysisBegin(const FOnAnalysisContext& Context)
 	Builder.RouteEvent(RouteId_ChaosVDBinaryDataContent, "ChaosVDLogger", "ChaosVDBinaryDataContent");
 	Builder.RouteEvent(RouteId_ChaosVDBinaryDataEnd, "ChaosVDLogger", "ChaosVDBinaryDataEnd");
 	Builder.RouteEvent(RouteId_ChaosVDSolverSimulationSpace, "ChaosVDLogger", "ChaosVDSolverSimulationSpace");
+	
+	Builder.RouteEvent(RouteId_ChaosVDNonSolverLocation, "ChaosVDLogger", "ChaosVDNonSolverLocation");
+	Builder.RouteEvent(RouteId_ChaosVDNonSolverTransform, "ChaosVDLogger", "ChaosVDNonSolverTransform");
 
 	Builder.RouteEvent(RouteId_BeginFrame, "Misc", "BeginFrame");
 	Builder.RouteEvent(RouteId_EndFrame, "Misc", "EndFrame");
@@ -178,15 +181,10 @@ bool FChaosVDTraceAnalyzer::OnEvent(uint16 RouteId, EStyle Style, const FOnEvent
 			const int32 SolverID = EventData.GetValue<int32>("SolverID");
 
 			FVector Position;
-			Position.X = EventData.GetValue<float>("PositionX");
-			Position.Y = EventData.GetValue<float>("PositionY");
-			Position.Z = EventData.GetValue<float>("PositionZ");
+			CVD_READ_TRACE_VECTOR(Position, Position, float, EventData);
 
 			FQuat Rotation;
-			Rotation.X = EventData.GetValue<float>("RotationX");
-			Rotation.Y = EventData.GetValue<float>("RotationY");
-			Rotation.Z = EventData.GetValue<float>("RotationZ");
-			Rotation.W = EventData.GetValue<float>("RotationW");
+			CVD_READ_TRACE_QUAT(Rotation, Rotation, float, EventData);
 
 			FWriteScopeLock WriteLock(ChaosVDTraceProvider->GetDataLock());
 			// This can be null if the recording started Mid-Frame. In this case we just discard the data for now
@@ -197,6 +195,35 @@ bool FChaosVDTraceAnalyzer::OnEvent(uint16 RouteId, EStyle Style, const FOnEvent
 			}
 			break;
 		}
+	case RouteId_ChaosVDNonSolverLocation:
+		{
+			FChaosVDTrackedLocation TrackedLocation;
+			CVD_READ_TRACE_VECTOR(TrackedLocation.Location, Position, float, EventData);
+
+			EventData.GetString("DebugName", TrackedLocation.DebugName);
+			
+			FWriteScopeLock WriteLock(ChaosVDTraceProvider->GetDataLock());
+			if (FChaosVDGameFrameData* CurrentFrameData = ChaosVDTraceProvider->GetLastGameFrame_AssumesLocked())
+			{
+				CurrentFrameData->RecordedNonSolverLocationsByID.Add(FName(TrackedLocation.DebugName), MoveTemp(TrackedLocation));
+			}
+			
+			break;
+		}
+	case RouteId_ChaosVDNonSolverTransform:
+		{
+			FChaosVDTrackedTransform TrackedTransform;
+			CVD_READ_TRACE_TRANSFORM(TrackedTransform.Transform, float, EventData);
+
+			EventData.GetString("DebugName", TrackedTransform.DebugName);
+			
+			FWriteScopeLock WriteLock(ChaosVDTraceProvider->GetDataLock());
+			if (FChaosVDGameFrameData* CurrentFrameData = ChaosVDTraceProvider->GetLastGameFrame_AssumesLocked())
+			{
+				CurrentFrameData->RecordedNonSolverTransformsByID.Add(FName(TrackedTransform.DebugName), MoveTemp(TrackedTransform));
+			}
+			break;
+		}
 	default:
 		break;
 	}
diff --git a/Engine/Plugins/Experimental/ChaosVD/Source/ChaosVD/Private/Widgets/SChaosVDMainTab.cpp b/Engine/Plugins/Experimental/ChaosVD/Source/ChaosVD/Private/Widgets/SChaosVDMainTab.cpp
--- a/Engine/Plugins/Experimental/ChaosVD/Source/ChaosVD/Private/Widgets/SChaosVDMainTab.cpp
+++ b/Engine/Plugins/Experimental/ChaosVD/Source/ChaosVD/Private/Widgets/SChaosVDMainTab.cpp
@@ -2,7 +2,7 @@
 
 #include "Widgets/SChaosVDMainTab.h"
 
-#include "ChaosVDEditorSettingsTab.h"
+#include "ChaosVDEditorVisualizationSettingsTab.h"
 #include "ChaosVDEngine.h"
 #include "ChaosVDObjectDetailsTab.h"
 #include "ChaosVDOutputLogTab.h"
@@ -44,7 +44,7 @@ void SChaosVDMainTab::Construct(const FArguments& InArgs, TSharedPtr<FChaosVDEng
 	OutputLogTab = MakeShared<FChaosVDOutputLogTab>(FChaosVDTabID::OutputLog, TabManager, this);
 	PlaybackViewportTab = MakeShared<FChaosVDPlaybackViewportTab>(FChaosVDTabID::PlaybackViewport, TabManager, this);
 	SolversTracksTab = MakeShared<FChaosVDSolversTracksTab>(FChaosVDTabID::SolversTrack, TabManager, this);
-	EditorSettingsTab = MakeShared<FChaosVDEditorSettingsTab>(FChaosVDTabID::CVDEditorSettings, TabManager, this);
+	EditorSettingsTab = MakeShared<FChaosVDEditorVisualizationSettingsTab>(FChaosVDTabID::CVDEditorSettings, TabManager, this);
 
 	StatusBarID = FName(FChaosVDTabID::StatusBar.ToString() + InChaosVDEngine->GetInstanceGuid().ToString());
 	
diff --git a/Engine/Plugins/Experimental/ChaosVD/Source/ChaosVD/Private/Widgets/SChaosVDNameListPicker.cpp b/Engine/Plugins/Experimental/ChaosVD/Source/ChaosVD/Private/Widgets/SChaosVDNameListPicker.cpp
new file mode 100644
--- /dev/null
+++ b/Engine/Plugins/Experimental/ChaosVD/Source/ChaosVD/Private/Widgets/SChaosVDNameListPicker.cpp
@@ -0,0 +1,125 @@
+// Copyright Epic Games, Inc. All Rights Reserved.
+
+#include "Widgets/SChaosVDNameListPicker.h"
+#include "Widgets/Input/SComboButton.h"
+#include "Widgets/Layout/SExpandableArea.h"
+#include "Widgets/Text/STextBlock.h"
+#include "Widgets/Views/STableRow.h"
+#include "Widgets/Views/SListView.h"
+#include "Widgets/SBoxPanel.h"
+#include "Widgets/SBoxPanel.h"
+
+void SChaosVDNameListPicker::Construct(const FArguments& InArgs)
+{
+	NameSelectedDelegate = InArgs._OnNameSleceted;
+
+	// Dropdown Button
+	 SAssignNew(DropDownButtonLabelTextBox, STextBlock)
+		.Font(FAppStyle::GetFontStyle(TEXT("PropertyWindow.NormalFont")))
+		.Text(this, &SChaosVDNameListPicker::GetCurrentDropDownButtonLabel);
+
+	// Name List
+	SAssignNew(NameListWidget, SListView<TSharedPtr<FName>>)
+		.ListItemsSource(&CachedNameList)
+		.SelectionMode( ESelectionMode::Single )
+		.OnMouseButtonClick(this, &SChaosVDNameListPicker::OnClickItem)
+		.ListViewStyle(&FAppStyle::Get().GetWidgetStyle<FTableViewStyle>("SimpleListView"))
+		.OnGenerateRow(this, &SChaosVDNameListPicker::MakeNameItemForList);
+
+	// Dropdown List content
+	SAssignNew(NamesListDropdown, SListViewSelectorDropdownMenu< TSharedPtr<FName> >, nullptr, NameListWidget)
+	[
+		SNew(SBorder)
+		.BorderImage(FAppStyle::GetBrush("Menu.Background"))
+		.Padding(2)
+		[
+			SNew(SBox)
+			.WidthOverride(250)
+			[				
+				SNew(SVerticalBox)
+				+SVerticalBox::Slot()
+				.MaxHeight(200.0f)
+				[
+					NameListWidget.ToSharedRef()
+				]
+			]
+		]
+	];
+	
+	ChildSlot
+	[
+		// Combo button that summons the dropdown menu
+		SAssignNew(PickerComboButton, SComboButton)
+			.ButtonContent()
+			[
+				DropDownButtonLabelTextBox.ToSharedRef()
+			]
+			.MenuContent()
+			[
+				NamesListDropdown.ToSharedRef()
+			]
+			.IsFocusable(true)
+			.ContentPadding(2.0f)
+			.OnComboBoxOpened(this, &SChaosVDNameListPicker::OnListOpened)
+	];
+	
+}
+
+void SChaosVDNameListPicker::UpdateNameList(TArray<TSharedPtr<FName>>&& NewNameList)
+{
+	// TODO: It might be worth to check is the array is different before doing this, but it might be more expensive
+	CachedNameList = MoveTemp(NewNameList);
+	NameListWidget->RebuildList();
+}
+
+FText SChaosVDNameListPicker::GetCurrentDropDownButtonLabel() const
+{
+	if (CurrentSelectedName.IsValid())
+	{
+		return FText::FromName(*CurrentSelectedName);
+	}
+	else
+	{
+		return NSLOCTEXT("ChaosVisualDebugger", "NameList", "Pick an element...");
+	}
+}
+
+void SChaosVDNameListPicker::OnClickItem(TSharedPtr<FName> Item)
+{
+	CurrentSelectedName = Item;
+
+	if (PickerComboButton.IsValid())
+	{
+		PickerComboButton->SetIsOpen(false);
+	}
+	
+	NameSelectedDelegate.ExecuteIfBound(Item);
+}
+
+void SChaosVDNameListPicker::OnListOpened()
+{
+	if (NameListWidget.IsValid())
+	{
+		NameListWidget->SetSelection(CurrentSelectedName, ESelectInfo::OnKeyPress);
+		NameListWidget->RequestScrollIntoView(CurrentSelectedName);
+	}
+}
+
+TSharedRef<ITableRow> SChaosVDNameListPicker::MakeNameItemForList(TSharedPtr<FName> Name, const TSharedRef<STableViewBase>& OwnerTable)
+{
+	TSharedPtr<STableRow<TSharedPtr<FString>>> ListTableRow;
+
+	SAssignNew(ListTableRow, STableRow<TSharedPtr<FString>>, OwnerTable)
+		[
+			SNew(SVerticalBox)
+				+SVerticalBox::Slot()
+				.AutoHeight()
+				.Padding(20.0f, 1.0f, 20.0f, 1.0f)
+				[
+					SNew(STextBlock)
+					.Text(FText::FromName(Name.IsValid() ? *Name.Get() : FName("Invalid")))
+				]
+		];
+
+	return ListTableRow.ToSharedRef();
+}
diff --git a/Engine/Plugins/Experimental/ChaosVD/Source/ChaosVD/Private/Widgets/SChaosVDPlaybackViewport.cpp b/Engine/Plugins/Experimental/ChaosVD/Source/ChaosVD/Private/Widgets/SChaosVDPlaybackViewport.cpp
--- a/Engine/Plugins/Experimental/ChaosVD/Source/ChaosVD/Private/Widgets/SChaosVDPlaybackViewport.cpp
+++ b/Engine/Plugins/Experimental/ChaosVD/Source/ChaosVD/Private/Widgets/SChaosVDPlaybackViewport.cpp
@@ -2,6 +2,7 @@
 
 #include "Widgets/SChaosVDPlaybackViewport.h"
 
+#include "ChaosVDEditorSettings.h"
 #include "ChaosVDPlaybackController.h"
 #include "ChaosVDPlaybackViewportClient.h"
 #include "ChaosVDScene.h"
@@ -41,6 +42,7 @@ TSharedPtr<FChaosVDPlaybackViewportClient> SChaosVDPlaybackViewport::CreateViewp
 
 void SChaosVDPlaybackViewport::Construct(const FArguments& InArgs, TWeakPtr<FChaosVDScene> InScene, TWeakPtr<FChaosVDPlaybackController> InPlaybackController)
 {
+	CVDSceneWeakPtr = InScene;
 	TSharedPtr<FChaosVDScene> ScenePtr = InScene.Pin();
 	ensure(ScenePtr.IsValid());
 	ensure(InPlaybackController.IsValid());
@@ -101,6 +103,23 @@ void SChaosVDPlaybackViewport::Construct(const FArguments& InArgs, TWeakPtr<FCha
 	RegisterNewController(InPlaybackController);
 }
 
+
+void SChaosVDPlaybackViewport::TrackActor(AActor* ActorToTrack, EChaosVDActorTrackingMode TrackingMode)
+{
+	if (ensure(PlaybackViewportClient.IsValid()))
+	{
+		PlaybackViewportClient->TrackActor(ActorToTrack, TrackingMode);
+	}	
+}
+
+void SChaosVDPlaybackViewport::TrackTransform(const FTransform& TransformToTrack, EChaosVDActorTrackingMode TrackingMode)
+{
+	if (ensure(PlaybackViewportClient.IsValid()))
+	{
+		PlaybackViewportClient->TrackTransform(TransformToTrack, TrackingMode);
+	}
+}
+
 void SChaosVDPlaybackViewport::HandlePlaybackControllerDataUpdated(TWeakPtr<FChaosVDPlaybackController> InController)
 {
 	if (PlaybackController != InController)
@@ -150,7 +169,7 @@ void SChaosVDPlaybackViewport::HandlePostSelectionChange(const UTypedElementSele
 
 void SChaosVDPlaybackViewport::OnPlaybackSceneUpdated()
 {
-	PlaybackViewportClient->bNeedsRedraw = true;
+	PlaybackViewportClient->bNeedsRedraw = true;	
 }
 
 void SChaosVDPlaybackViewport::RegisterNewController(TWeakPtr<FChaosVDPlaybackController> NewController)
diff --git a/Engine/Plugins/Experimental/ChaosVD/Source/ChaosVD/Private/Widgets/SChaosVDVisualizationControls.cpp b/Engine/Plugins/Experimental/ChaosVD/Source/ChaosVD/Private/Widgets/SChaosVDVisualizationControls.cpp
new file mode 100644
--- /dev/null
+++ b/Engine/Plugins/Experimental/ChaosVD/Source/ChaosVD/Private/Widgets/SChaosVDVisualizationControls.cpp
@@ -0,0 +1,233 @@
+// Copyright Epic Games, Inc. All Rights Reserved.
+
+#include "Widgets/SChaosVDVisualizationControls.h"
+
+#include "ChaosVDEditorSettings.h"
+#include "ChaosVDModule.h"
+#include "ChaosVDPlaybackController.h"
+#include "ChaosVDPlaybackViewportClient.h"
+#include "ChaosVDRecording.h"
+#include "ChaosVDScene.h"
+#include "Elements/Framework/TypedElementSelectionSet.h"
+#include "PropertyEditorModule.h"
+#include "Modules/ModuleManager.h"
+#include "Widgets/Layout/SBox.h"
+#include "Widgets/Layout/SExpandableArea.h"
+#include "Widgets/SChaosVDNameListPicker.h"
+#include "Widgets/SChaosVDPlaybackViewport.h"
+#include "Widgets/Text/STextBlock.h"
+
+void SChaosVDVisualizationControls::Construct(const FArguments& InArgs, const TWeakPtr<FChaosVDPlaybackController>& InPlaybackController,const TWeakPtr<SChaosVDPlaybackViewport>& InPlaybackViewport)
+{
+	RegisterNewController(InPlaybackController);
+
+	PlaybackViewport = InPlaybackViewport;
+	
+	FPropertyEditorModule& PropertyEditorModule = FModuleManager::GetModuleChecked<FPropertyEditorModule>("PropertyEditor");
+	FDetailsViewArgs DetailsViewArgs;
+	DetailsViewArgs.bAllowSearch = false;
+	DetailsViewArgs.NameAreaSettings = FDetailsViewArgs::HideNameArea;
+
+	TSharedRef<IDetailsView> DetailsPanel = PropertyEditorModule.CreateDetailView(DetailsViewArgs);
+
+	DetailsPanel->SetObject(GetMutableDefault<UChaosVDEditorSettings>());
+
+	ChildSlot
+	[
+		SNew(SVerticalBox)
+		+SVerticalBox::Slot()
+		.Padding(2,5,2,5)
+		.AutoHeight()
+		[
+			DetailsPanel
+		]
+		+SVerticalBox::Slot()
+		.AutoHeight()
+		.Padding(2,5,2,5)
+		[
+			SNew(SExpandableArea)
+				.InitiallyCollapsed(false)
+				.BorderBackgroundColor(FLinearColor::White)
+				.Padding(FMargin(8.f))
+				.Visibility_Raw(this, &SChaosVDVisualizationControls::GetTrackedDataOptionsVisibility)
+				.HeaderContent()
+				[
+					SNew(SHorizontalBox)
+					+SHorizontalBox::Slot()
+					.HAlign(HAlign_Left)
+					.VAlign(VAlign_Center)
+					.AutoWidth()
+					.Padding(0.f, 0.f, 0.f, 0.f)
+					[
+						SNew(STextBlock)
+						.Text(FText::FromString("Trackable Data"))
+						.Font(FCoreStyle::Get().GetFontStyle("ExpandableArea.TitleFont"))
+					]
+				]
+				.BodyContent()
+				[
+					SNew(SVerticalBox)
+					+SVerticalBox::Slot()
+					[
+						SNew(SHorizontalBox)
+						.Visibility_Raw(this, &SChaosVDVisualizationControls::GetTransformsPickerVisibility)
+						+SHorizontalBox::Slot()
+						.VAlign(VAlign_Center)
+						.HAlign(HAlign_Left)
+						[
+							SNew(STextBlock)
+							.Text(FText::FromString("Transform to Track"))
+							.Font(FCoreStyle::Get().GetFontStyle("ExpandableArea.TitleFont"))
+						]
+						+SHorizontalBox::Slot()
+						.VAlign(VAlign_Center)
+						.HAlign(HAlign_Right)
+						[
+							SAssignNew(TransformNamePickerWidget, SChaosVDNameListPicker)
+							.OnNameSleceted_Raw(this, &SChaosVDVisualizationControls::HandleTransformNameSelected)
+						]	
+					]
+					+SVerticalBox::Slot()
+					.Padding(2,5,2,5)
+					[
+						SNew(SHorizontalBox)
+						.Visibility_Raw(this, &SChaosVDVisualizationControls::GetLocationsPickerVisibility)
+						+SHorizontalBox::Slot()
+						.VAlign(VAlign_Center)
+						.HAlign(HAlign_Left)
+						[
+							SNew(STextBlock)
+							.Text(FText::FromString("Location to Track"))
+							.Font(FCoreStyle::Get().GetFontStyle("ExpandableArea.TitleFont"))
+						]
+						+SHorizontalBox::Slot()
+						.VAlign(VAlign_Center)
+						.HAlign(HAlign_Right)
+						[
+							SAssignNew(LocationsNamePickerWidget, SChaosVDNameListPicker)
+							.OnNameSleceted_Raw(this, &SChaosVDVisualizationControls::HandleLocationNameSelected)
+						]
+					]
+				]
+		]
+	];
+}
+
+void SChaosVDVisualizationControls::HandleControllerTrackFrameUpdated(TWeakPtr<FChaosVDPlaybackController> InController, const FChaosVDTrackInfo* UpdatedTrackInfo, FGuid InstigatorGuid)
+{
+	if (UpdatedTrackInfo == nullptr)
+	{
+		return;
+	}
+
+	if (UpdatedTrackInfo->TrackType != EChaosVDTrackType::Game)
+	{
+		return;
+	}
+
+	const TSharedPtr<SChaosVDPlaybackViewport> PlaybackViewportSharedPtr = PlaybackViewport.Pin();
+	const TSharedPtr<FChaosVDPlaybackController> ControllerSharedPtr = InController.Pin();
+		
+	if (!PlaybackViewportSharedPtr.IsValid() || !ControllerSharedPtr.IsValid())
+	{
+		return;
+	}
+
+	if (const TSharedPtr<FChaosVDRecording> RecordingSharedPtr  = ControllerSharedPtr->GetCurrentRecording().Pin())
+	{
+		if (FChaosVDGameFrameData* FrameData = RecordingSharedPtr->GetGameFrameData_AssumesLocked(UpdatedTrackInfo->CurrentFrame))
+		{
+			UpdateNameListFromKeys(FrameData->RecordedNonSolverTransformsByID, *TransformNamePickerWidget.Get());
+			UpdateNameListFromKeys(FrameData->RecordedNonSolverLocationsByID, *LocationsNamePickerWidget.Get());
+
+			if (const TSharedPtr<FChaosVDScene> CVDSceneSharedPtr = ControllerSharedPtr->GetControllerScene().Pin())
+			{
+				if (const UChaosVDEditorSettings* CVDEditorSettings = GetDefault<UChaosVDEditorSettings>())
+				{
+					switch (CVDEditorSettings->TrackingTarget)
+					{
+						case EChaosVDActorTrackingTarget::SelectedObject:
+							{
+								TArray<AActor*> SelectedActors = CVDSceneSharedPtr->GetElementSelectionSet()->GetSelectedObjects<AActor>();
+								if (SelectedActors.Num() > 0)
+								{
+									//TODO: Update this if we add multi selection support
+									if (AActor* SelectedActor = SelectedActors[0])
+									{
+										PlaybackViewportSharedPtr->TrackActor(SelectedActor, CVDEditorSettings->TrackingOptions);
+									}
+								}
+								break;
+							}
+						case EChaosVDActorTrackingTarget::RecordedTransform:
+							{
+								if (SelectedTrackedTransformName.IsValid())
+								{
+									if (const FChaosVDTrackedTransform* TrackedTransform = FrameData->RecordedNonSolverTransformsByID.Find(*SelectedTrackedTransformName))
+									{
+										PlaybackViewportSharedPtr->TrackTransform(TrackedTransform->Transform, CVDEditorSettings->TrackingOptions);
+									}
+								}
+								
+								break;
+							}
+						case EChaosVDActorTrackingTarget::RecordedLocation:
+							{
+								if (SelectedTrackedLocationName.IsValid())
+								{
+									if (const FChaosVDTrackedLocation* TrackedTransform = FrameData->RecordedNonSolverLocationsByID.Find(*SelectedTrackedLocationName))
+									{
+										FTransform LocationTransform;
+										LocationTransform.SetLocation(TrackedTransform->Location);
+										PlaybackViewportSharedPtr->TrackTransform(LocationTransform, CVDEditorSettings->TrackingOptions);
+									}
+								}
+								
+								break;
+							}
+						default:
+							break;
+						}			
+				}
+			}
+		}
+	}
+}
+
+void SChaosVDVisualizationControls::HandleLocationNameSelected(TSharedPtr<FName> SelectedName)
+{
+	SelectedTrackedLocationName = SelectedName;
+}
+
+void SChaosVDVisualizationControls::HandleTransformNameSelected(TSharedPtr<FName> SelectedName)
+{
+	SelectedTrackedTransformName = SelectedName;
+}
+
+EVisibility SChaosVDVisualizationControls::GetLocationsPickerVisibility() const
+{
+	if (const UChaosVDEditorSettings* Settings = GetDefault<UChaosVDEditorSettings>())
+	{
+		return Settings->TrackingTarget == EChaosVDActorTrackingTarget::RecordedLocation ? EVisibility::Visible : EVisibility::Collapsed;
+	}
+
+	return EVisibility::Collapsed;
+}
+
+EVisibility SChaosVDVisualizationControls::GetTransformsPickerVisibility() const
+{
+	if (const UChaosVDEditorSettings* Settings = GetDefault<UChaosVDEditorSettings>())
+	{
+		return Settings->TrackingTarget == EChaosVDActorTrackingTarget::RecordedTransform ? EVisibility::Visible : EVisibility::Collapsed;
+	}
+
+	return EVisibility::Collapsed;
+}
+
+
+EVisibility SChaosVDVisualizationControls::GetTrackedDataOptionsVisibility() const
+{
+	const bool bIsTrackedDataVisible = GetLocationsPickerVisibility() == EVisibility::Visible ||  GetTransformsPickerVisibility() == EVisibility::Visible;
+
+	return bIsTrackedDataVisible ? EVisibility::Visible : EVisibility::Collapsed;
+}
diff --git a/Engine/Plugins/Experimental/ChaosVD/Source/ChaosVD/Public/ChaosVDEditorSettings.h b/Engine/Plugins/Experimental/ChaosVD/Source/ChaosVD/Public/ChaosVDEditorSettings.h
--- a/Engine/Plugins/Experimental/ChaosVD/Source/ChaosVD/Public/ChaosVDEditorSettings.h
+++ b/Engine/Plugins/Experimental/ChaosVD/Source/ChaosVD/Public/ChaosVDEditorSettings.h
@@ -2,6 +2,7 @@
 
 #pragma once
 
+#include "ChaosVDRecording.h"
 #include "Engine/World.h"
 #include "Engine/StaticMesh.h"
 #include "UObject/Object.h"
@@ -9,6 +10,24 @@
 
 #include "ChaosVDEditorSettings.generated.h"
 
+
+UENUM()
+enum class EChaosVDActorTrackingMode
+{
+	ByDistanceOffset,
+	ByBoundingBox,
+	MatchTransform
+};
+
+UENUM()
+enum class EChaosVDActorTrackingTarget
+{
+	Disabled,
+	SelectedObject,
+	RecordedTransform,
+	RecordedLocation,
+};
+
 UCLASS(config = Engine)
 class UChaosVDEditorSettings : public UObject
 {
@@ -24,6 +43,18 @@ public:
 	UPROPERTY(EditAnywhere, Category = "Viewport Visualization Flags", meta = (Bitmask, BitmaskEnum = EChaosVDCollisionVisualizationFlags))
 	bool bShowDebugText = false;
 
-	UPROPERTY(Config, EditInstanceOnly, Category = "Editor Options")
+	UPROPERTY(EditAnywhere, Category = "Viewport Tracking")
+	EChaosVDActorTrackingTarget TrackingTarget;
+
+	UPROPERTY(EditAnywhere, Category = "Viewport Tracking", meta=(EditCondition = "TrackingTarget != EChaosVDActorTrackingTarget::Disabled", EditConditionHides))
+	EChaosVDActorTrackingMode TrackingOptions;
+
+	UPROPERTY(EditAnywhere, Category = "Viewport Tracking", meta=(EditCondition = "TrackingOptions == EChaosVDActorTrackingMode::ByDistanceOffset && TrackingTarget != EChaosVDActorTrackingTarget::Disabled", EditConditionHides))
+	float TrackingDistanceOffset = 1500.0f;
+
+	UPROPERTY(EditAnywhere, Category = "Viewport Tracking", meta=(EditCondition = "TrackingOptions == EChaosVDActorTrackingMode::ByBoundingBox && TrackingTarget != EChaosVDActorTrackingTarget::Disabled", EditConditionHides))
+	float ExpandViewTrackingBy = 60.0f;
+
+	UPROPERTY(Config)
 	TSoftObjectPtr<UWorld> BasePhysicsVDWorld;
-};
\ No newline at end of file
+};
diff --git a/Engine/Plugins/Experimental/ChaosVD/Source/ChaosVD/Public/ChaosVDEditorSettingsTab.h b/Engine/Plugins/Experimental/ChaosVD/Source/ChaosVD/Public/ChaosVDEditorSettingsTab.h
deleted file mode 100644
diff --git a/Engine/Plugins/Experimental/ChaosVD/Source/ChaosVD/Public/ChaosVDEditorVisualizationSettingsTab.h b/Engine/Plugins/Experimental/ChaosVD/Source/ChaosVD/Public/ChaosVDEditorVisualizationSettingsTab.h
new file mode 100644
--- /dev/null
+++ b/Engine/Plugins/Experimental/ChaosVD/Source/ChaosVD/Public/ChaosVDEditorVisualizationSettingsTab.h
@@ -0,0 +1,18 @@
+﻿// Copyright Epic Games, Inc. All Rights Reserved.
+#pragma once
+
+#include "ChaosVDTabSpawnerBase.h"
+#include "Widgets/SChaosVDMainTab.h"
+
+class SChaosVDNameListPicker;
+
+class FChaosVDEditorVisualizationSettingsTab : public FChaosVDTabSpawnerBase
+{
+public:
+	FChaosVDEditorVisualizationSettingsTab(const FName& InTabID, const TSharedPtr<FTabManager>& InTabManager, SChaosVDMainTab* InOwningTabWidget)
+		: FChaosVDTabSpawnerBase(InTabID, InTabManager, InOwningTabWidget)
+	{
+	}
+
+	virtual TSharedRef<SDockTab> HandleTabSpawned(const FSpawnTabArgs& Args) override;
+};
diff --git a/Engine/Plugins/Experimental/ChaosVD/Source/ChaosVD/Public/ChaosVDPlaybackViewportClient.h b/Engine/Plugins/Experimental/ChaosVD/Source/ChaosVD/Public/ChaosVDPlaybackViewportClient.h
--- a/Engine/Plugins/Experimental/ChaosVD/Source/ChaosVD/Public/ChaosVDPlaybackViewportClient.h
+++ b/Engine/Plugins/Experimental/ChaosVD/Source/ChaosVD/Public/ChaosVDPlaybackViewportClient.h
@@ -4,6 +4,7 @@
 #include "LevelEditorViewport.h"
 
 class FChaosVDScene;
+enum class EChaosVDActorTrackingMode;
 
 /** Client viewport class used for to handle a Chaos Visual Debugger world Interaction/Rendering.
  * It re-routes interaction events to our Chaos VD scene
@@ -20,15 +21,15 @@ public:
 
 	virtual UWorld* GetWorld() const override { return CVDWorld; };
 
-private:
-
-	void HandleObjectFocused(UObject* FocusedObject);
-
-public:
 	virtual void Draw(const FSceneView* View, FPrimitiveDrawInterface* PDI) override;
 	virtual void DrawCanvas(FViewport& InViewport, FSceneView& View, FCanvas& Canvas) override;
+	
+	void TrackActor(AActor* ActorToTrack, EChaosVDActorTrackingMode TrackingMode);
+	void TrackTransform(const FTransform& TransformToTrack, EChaosVDActorTrackingMode TrackingMode);
 
 private:
+	void HandleObjectFocused(UObject* FocusedObject);
+	
 	FDelegateHandle ObjectFocusedDelegateHandle;
 	UWorld* CVDWorld;
 	TWeakPtr<FChaosVDScene> CVDScene;
diff --git a/Engine/Plugins/Experimental/ChaosVD/Source/ChaosVD/Public/ChaosVDPlaybackViewportTab.h b/Engine/Plugins/Experimental/ChaosVD/Source/ChaosVD/Public/ChaosVDPlaybackViewportTab.h
--- a/Engine/Plugins/Experimental/ChaosVD/Source/ChaosVD/Public/ChaosVDPlaybackViewportTab.h
+++ b/Engine/Plugins/Experimental/ChaosVD/Source/ChaosVD/Public/ChaosVDPlaybackViewportTab.h
@@ -20,6 +20,8 @@ public:
 	
 	virtual TSharedRef<SDockTab> HandleTabSpawned(const FSpawnTabArgs& Args) override;
 
+	TWeakPtr<SChaosVDPlaybackViewport> GetPlaybackViewportWidget() { return PlaybackViewportWidget; }
+
 protected:
 	TSharedPtr<SChaosVDPlaybackViewport> PlaybackViewportWidget;
 };
diff --git a/Engine/Plugins/Experimental/ChaosVD/Source/ChaosVD/Public/Trace/ChaosVDTraceAnalyzer.h b/Engine/Plugins/Experimental/ChaosVD/Source/ChaosVD/Public/Trace/ChaosVDTraceAnalyzer.h
--- a/Engine/Plugins/Experimental/ChaosVD/Source/ChaosVD/Public/Trace/ChaosVDTraceAnalyzer.h
+++ b/Engine/Plugins/Experimental/ChaosVD/Source/ChaosVD/Public/Trace/ChaosVDTraceAnalyzer.h
@@ -47,6 +47,8 @@ private:
 		RouteId_BeginFrame,
 		RouteId_EndFrame,
 		RouteId_ChaosVDParticleDestroyed,
+		RouteId_ChaosVDNonSolverLocation,
+		RouteId_ChaosVDNonSolverTransform,
 	};
 
 	TraceServices::IAnalysisSession& Session;
@@ -55,3 +57,37 @@ private:
 	
 	FChaosVDTraceAnalysisComplete ChaosVDTraceAnalysisCompleteDelegate;
 };
+
+#ifndef CVD_STRINGIZE
+	#define CVD_STRINGIZE(x) #x
+#endif
+
+#ifndef CVD_READ_TRACE_VECTOR
+	#define CVD_READ_TRACE_VECTOR(Vector, VectorName, ValueType, EventData) \
+		Vector.X = EventData.GetValue<ValueType>(CVD_STRINGIZE(VectorName##X)); \
+		Vector.Y = EventData.GetValue<ValueType>(CVD_STRINGIZE(VectorName##Y)); \
+		Vector.Z = EventData.GetValue<ValueType>(CVD_STRINGIZE(VectorName##Z));
+#endif
+
+#ifndef CVD_READ_TRACE_QUAT
+	#define CVD_READ_TRACE_QUAT(Rotator, RotatorName, ValueType, EventData) \
+		Rotator.X = EventData.GetValue<ValueType>(CVD_STRINGIZE(RotatorName##X)); \
+		Rotator.Y = EventData.GetValue<ValueType>(CVD_STRINGIZE(RotatorName##Y)); \
+		Rotator.Z = EventData.GetValue<ValueType>(CVD_STRINGIZE(RotatorName##Z)); \
+		Rotator.W = EventData.GetValue<ValueType>(CVD_STRINGIZE(RotatorName##W));
+#endif
+
+#ifndef CVD_READ_TRACE_TRANSFORM
+	#define CVD_READ_TRACE_TRANSFORM(Transform, ValueType, EventData) \
+		{ \
+			FVector Location; \
+			CVD_READ_TRACE_VECTOR(Location, Position, float, EventData); \
+			FVector Scale; \
+			CVD_READ_TRACE_VECTOR(Location, Position, float, EventData); \
+			FQuat Rotation; \
+			CVD_READ_TRACE_QUAT(Rotation, Rotation, float, EventData); \
+			Transform.SetLocation(Location); \
+			Transform.SetScale3D(Scale); \
+			Transform.SetRotation(Rotation); \
+		}
+#endif
diff --git a/Engine/Plugins/Experimental/ChaosVD/Source/ChaosVD/Public/Widgets/SChaosVDMainTab.h b/Engine/Plugins/Experimental/ChaosVD/Source/ChaosVD/Public/Widgets/SChaosVDMainTab.h
--- a/Engine/Plugins/Experimental/ChaosVD/Source/ChaosVD/Public/Widgets/SChaosVDMainTab.h
+++ b/Engine/Plugins/Experimental/ChaosVD/Source/ChaosVD/Public/Widgets/SChaosVDMainTab.h
@@ -6,7 +6,8 @@
 #include "Templates/SharedPointer.h"
 #include "Widgets/SCompoundWidget.h"
 
-class FChaosVDEditorSettingsTab;
+class FChaosVDTabSpawnerBase;
+class FChaosVDEditorVisualizationSettingsTab;
 class FChaosVDSolversTracksTab;
 class FChaosVDEngine;
 class FChaosVDOutputLogTab;
@@ -29,6 +30,8 @@ public:
 	TSharedRef<FChaosVDEngine> GetChaosVDEngineInstance() const { return ChaosVDEngine.ToSharedRef(); };
 	const FName& GetStatusBarID() const { return StatusBarID; };
 
+	TWeakPtr<FChaosVDPlaybackViewportTab> GetPlaybackViewportTab() { return PlaybackViewportTab;}
+
 private:
 
 	TSharedRef<FTabManager::FLayout> GenerateMainLayout();
@@ -45,7 +48,7 @@ private:
 	TSharedPtr<FChaosVDPlaybackViewportTab> PlaybackViewportTab;
 	TSharedPtr<FChaosVDSolversTracksTab> SolversTracksTab;
 	TSharedPtr<FChaosVDOutputLogTab> OutputLogTab;
-	TSharedPtr<FChaosVDEditorSettingsTab> EditorSettingsTab;
+	TSharedPtr<FChaosVDEditorVisualizationSettingsTab> EditorSettingsTab;
 
 	FName StatusBarID;
 
diff --git a/Engine/Plugins/Experimental/ChaosVD/Source/ChaosVD/Public/Widgets/SChaosVDNameListPicker.h b/Engine/Plugins/Experimental/ChaosVD/Source/ChaosVD/Public/Widgets/SChaosVDNameListPicker.h
new file mode 100644
--- /dev/null
+++ b/Engine/Plugins/Experimental/ChaosVD/Source/ChaosVD/Public/Widgets/SChaosVDNameListPicker.h
@@ -0,0 +1,54 @@
+// Copyright Epic Games, Inc. All Rights Reserved.
+
+#pragma once
+#include "SListViewSelectorDropdownMenu.h"
+#include "Widgets/SCompoundWidget.h"
+#include "Containers/Array.h"
+#include "Templates/SharedPointer.h"
+
+class SComboButton;
+class STextBlock;
+
+DECLARE_DELEGATE_OneParam(FChaosVDNameListSelected, TSharedPtr<FName>)
+
+/** Widget that takes a list of String Views and creates a drop down menu for it for the Chaos Visual Debugger Tool.
+ */
+class SChaosVDNameListPicker : public SCompoundWidget
+{
+public:
+	SLATE_BEGIN_ARGS(SChaosVDNameListPicker){}
+	SLATE_EVENT(FChaosVDNameListSelected, OnNameSleceted)
+	SLATE_END_ARGS()
+	
+	void Construct(const FArguments& InArgs);
+
+	/**
+	 * Replaces the names in the list view
+	 * @param NewNameList Array of Name to represent.
+	 */
+	void UpdateNameList(TArray<TSharedPtr<FName>>&& NewNameList);
+	
+protected:
+
+	FText GetCurrentDropDownButtonLabel() const;
+	
+	void OnClickItem(TSharedPtr<FName> Item);
+
+	void OnListOpened();
+
+	TSharedRef<ITableRow> MakeNameItemForList(TSharedPtr<FName> Name, const TSharedRef<STableViewBase>& OwnerTable);
+
+	TSharedPtr<SComboButton> PickerComboButton;
+
+	TSharedPtr<STextBlock> DropDownButtonLabelTextBox;
+	
+	TSharedPtr<SListView<TSharedPtr<FName>>> NameListWidget;
+
+	TSharedPtr<SListViewSelectorDropdownMenu<TSharedPtr<FName>>> NamesListDropdown;
+
+	TArray<TSharedPtr<FName>> CachedNameList;
+
+	TSharedPtr<FName> CurrentSelectedName;
+
+	FChaosVDNameListSelected NameSelectedDelegate;
+};
diff --git a/Engine/Plugins/Experimental/ChaosVD/Source/ChaosVD/Public/Widgets/SChaosVDPlaybackViewport.h b/Engine/Plugins/Experimental/ChaosVD/Source/ChaosVD/Public/Widgets/SChaosVDPlaybackViewport.h
--- a/Engine/Plugins/Experimental/ChaosVD/Source/ChaosVD/Public/Widgets/SChaosVDPlaybackViewport.h
+++ b/Engine/Plugins/Experimental/ChaosVD/Source/ChaosVD/Public/Widgets/SChaosVDPlaybackViewport.h
@@ -7,6 +7,7 @@
 #include "Templates/SharedPointer.h"
 #include "Widgets/SCompoundWidget.h"
 
+enum class EChaosVDActorTrackingMode;
 class FChaosVDPlaybackViewportClient;
 class SChaosVDSolverPlaybackControls;
 class FChaosVDPlaybackController;
@@ -29,6 +30,9 @@ public:
 
 	void Construct(const FArguments& InArgs, TWeakPtr<FChaosVDScene> InScene, TWeakPtr<FChaosVDPlaybackController> InPlaybackController);
 
+	void TrackActor(AActor* ActorToTrack, EChaosVDActorTrackingMode TrackingMode);
+	void TrackTransform(const FTransform& TransformToTrack, EChaosVDActorTrackingMode TrackingMode);
+
 protected:
 
 	TSharedPtr<FChaosVDPlaybackViewportClient> CreateViewportClient() const;
@@ -47,4 +51,6 @@ protected:
 	TSharedPtr<FChaosVDPlaybackViewportClient> PlaybackViewportClient;
 	TSharedPtr<SViewport> ViewportWidget;
 	TSharedPtr<FSceneViewport> SceneViewport;
+	
+	TWeakPtr<FChaosVDScene> CVDSceneWeakPtr;
 };
diff --git a/Engine/Plugins/Experimental/ChaosVD/Source/ChaosVD/Public/Widgets/SChaosVDVisualizationControls.h b/Engine/Plugins/Experimental/ChaosVD/Source/ChaosVD/Public/Widgets/SChaosVDVisualizationControls.h
new file mode 100644
--- /dev/null
+++ b/Engine/Plugins/Experimental/ChaosVD/Source/ChaosVD/Public/Widgets/SChaosVDVisualizationControls.h
@@ -0,0 +1,57 @@
+// Copyright Epic Games, Inc. All Rights Reserved.
+
+#pragma once
+#include "ChaosVDPlaybackControllerObserver.h"
+#include "Widgets/SChaosVDNameListPicker.h"
+#include "Templates/SharedPointer.h"
+
+class SChaosVDPlaybackViewport;
+class FChaosVDPlaybackViewportClient;
+struct FChaosVDTrackInfo;
+struct FGuid;
+
+/**
+ * Widget containing the visualization controls for the Chaos Visual debugger tool to be shown in a Tab
+ */
+class SChaosVDVisualizationControls : public SCompoundWidget, public FChaosVDPlaybackControllerObserver
+{
+public:
+	SLATE_BEGIN_ARGS( SChaosVDVisualizationControls ){}
+	SLATE_END_ARGS()
+
+	void Construct(const FArguments& InArgs, const TWeakPtr<FChaosVDPlaybackController>& InPlaybackController, const TWeakPtr<SChaosVDPlaybackViewport>& InPlaybackViewport);
+
+	virtual void HandleControllerTrackFrameUpdated(TWeakPtr<FChaosVDPlaybackController> InController, const FChaosVDTrackInfo* UpdatedTrackInfo, FGuid InstigatorGuid) override;
+
+protected:
+
+	void HandleLocationNameSelected(TSharedPtr<FName> SelectedName);
+	void HandleTransformNameSelected(TSharedPtr<FName> SelectedName);
+
+	template<typename ValueType>
+	void UpdateNameListFromKeys(const TMap<FName, ValueType>& MapToEvaluate, SChaosVDNameListPicker& NameList);
+	
+	EVisibility GetLocationsPickerVisibility() const;
+	EVisibility GetTransformsPickerVisibility() const;
+	EVisibility GetTrackedDataOptionsVisibility() const;
+
+	TSharedPtr<SChaosVDNameListPicker> TransformNamePickerWidget;
+	
+	TSharedPtr<SChaosVDNameListPicker> LocationsNamePickerWidget;
+
+	TWeakPtr<SChaosVDPlaybackViewport> PlaybackViewport;
+
+	TSharedPtr<FName> SelectedTrackedTransformName;
+	TSharedPtr<FName> SelectedTrackedLocationName;
+};
+
+template <typename ValueType>
+void SChaosVDVisualizationControls::UpdateNameListFromKeys(const TMap<FName, ValueType>& MapToEvaluate, SChaosVDNameListPicker& NameList)
+{
+	TArray<TSharedPtr<FName>> NewNameList;
+	NewNameList.Reserve(MapToEvaluate.Num());
+
+	Algo::Transform(MapToEvaluate, NewNameList, [](const TPair<FName, ValueType>& TransformData){ return MakeShared<FName>(TransformData.Key);});
+	
+	NameList.UpdateNameList(MoveTemp(NewNameList));
+}
diff --git a/Engine/Source/Runtime/Experimental/Chaos/Private/ChaosVisualDebugger/ChaosVisualDebuggerTrace.cpp b/Engine/Source/Runtime/Experimental/Chaos/Private/ChaosVisualDebugger/ChaosVisualDebuggerTrace.cpp
--- a/Engine/Source/Runtime/Experimental/Chaos/Private/ChaosVisualDebugger/ChaosVisualDebuggerTrace.cpp
+++ b/Engine/Source/Runtime/Experimental/Chaos/Private/ChaosVisualDebugger/ChaosVisualDebuggerTrace.cpp
@@ -32,6 +32,8 @@ UE_TRACE_EVENT_DEFINE(ChaosVDLogger, ChaosVDBinaryDataContent)
 UE_TRACE_EVENT_DEFINE(ChaosVDLogger, ChaosVDBinaryDataEnd)
 UE_TRACE_EVENT_DEFINE(ChaosVDLogger, ChaosVDSolverSimulationSpace)
 UE_TRACE_EVENT_DEFINE(ChaosVDLogger, ChaosVDDummyEvent)
+UE_TRACE_EVENT_DEFINE(ChaosVDLogger, ChaosVDNonSolverLocation)
+UE_TRACE_EVENT_DEFINE(ChaosVDLogger, ChaosVDNonSolverTransform)
 
 static FAutoConsoleVariable CVarChaosVDCompressBinaryData(
 	TEXT("p.Chaos.VD.CompressBinaryData"),
@@ -536,7 +538,35 @@ void FChaosVisualDebuggerTrace::TraceImplicitObject(FChaosVDImplicitObjectWrappe
 	WrappedGeometryData.Serialize(Ar);
 
 	TraceBinaryData(TLSDataBuffer.BufferRef, FChaosVDImplicitObjectWrapper::WrapperTypeName);
-}	
+}
+
+void FChaosVisualDebuggerTrace::TraceNonSolverLocation(const FVector& InLocation, FStringView DebugNameID)
+{
+	if (!IsTracing())
+	{
+		return;
+	}
+
+	UE_TRACE_LOG(ChaosVDLogger, ChaosVDNonSolverLocation, ChaosVDChannel)
+			<< ChaosVDNonSolverLocation.Cycle(FPlatformTime::Cycles64())
+			<< CVD_TRACE_VECTOR_ON_EVENT(ChaosVDNonSolverLocation, Position, InLocation)
+			<< ChaosVDNonSolverLocation.DebugName(DebugNameID.GetData(), DebugNameID.Len());
+}
+
+void FChaosVisualDebuggerTrace::TraceNonSolverTransform(const FTransform& InTransform, FStringView DebugNameID)
+{
+	if (!IsTracing())
+	{
+		return;
+	}
+
+	UE_TRACE_LOG(ChaosVDLogger, ChaosVDNonSolverTransform, ChaosVDChannel)
+		<< ChaosVDNonSolverTransform.Cycle(FPlatformTime::Cycles64())
+		<< CVD_TRACE_VECTOR_ON_EVENT(ChaosVDNonSolverTransform, Position, InTransform.GetLocation())
+		<< CVD_TRACE_VECTOR_ON_EVENT(ChaosVDNonSolverTransform, Scale, InTransform.GetScale3D())
+		<< CVD_TRACE_ROTATOR_ON_EVENT(ChaosVDNonSolverTransform, Rotation, InTransform.GetRotation())
+		<< ChaosVDNonSolverTransform.DebugName(DebugNameID.GetData(), DebugNameID.Len());
+}
 
 void FChaosVisualDebuggerTrace::RegisterEventHandlers()
 {
diff --git a/Engine/Source/Runtime/Experimental/Chaos/Public/ChaosVisualDebugger/ChaosVisualDebuggerTrace.h b/Engine/Source/Runtime/Experimental/Chaos/Public/ChaosVisualDebugger/ChaosVisualDebuggerTrace.h
--- a/Engine/Source/Runtime/Experimental/Chaos/Public/ChaosVisualDebugger/ChaosVisualDebuggerTrace.h
+++ b/Engine/Source/Runtime/Experimental/Chaos/Public/ChaosVisualDebugger/ChaosVisualDebuggerTrace.h
@@ -83,6 +83,14 @@
 	#define CVD_TRACE_MID_PHASES_FROM_COLLISION_CONSTRAINTS(CollisionConstraints)
 #endif
 
+#ifndef CVD_TRACE_NON_SOLVER_LOCATION
+      #define CVD_TRACE_NON_SOLVER_LOCATION(Location, DebugName)
+#endif
+
+#ifndef CVD_TRACE_NON_SOLVER_TRANSFORM
+	#define CVD_TRACE_NON_SOLVER_TRANSFORM(Transform, DebugName)
+#endif
+
 #else
 
 #include "ChaosVDRuntimeModule.h"
@@ -185,6 +193,20 @@ UE_TRACE_EVENT_BEGIN_EXTERN(ChaosVDLogger, ChaosVDSolverSimulationSpace)
 	CVD_DEFINE_TRACE_ROTATOR(Chaos::FReal, Rotation)
 UE_TRACE_EVENT_END()
 
+UE_TRACE_EVENT_BEGIN_EXTERN(ChaosVDLogger, ChaosVDNonSolverLocation)
+	UE_TRACE_EVENT_FIELD(uint64, Cycle)
+	CVD_DEFINE_TRACE_VECTOR(Chaos::FReal, Position)
+	UE_TRACE_EVENT_FIELD(UE::Trace::WideString, DebugName)
+UE_TRACE_EVENT_END()
+
+UE_TRACE_EVENT_BEGIN_EXTERN(ChaosVDLogger, ChaosVDNonSolverTransform)
+	UE_TRACE_EVENT_FIELD(uint64, Cycle)
+	CVD_DEFINE_TRACE_VECTOR(Chaos::FReal, Position)
+	CVD_DEFINE_TRACE_VECTOR(Chaos::FReal, Scale)
+	CVD_DEFINE_TRACE_ROTATOR(Chaos::FReal, Rotation)
+	UE_TRACE_EVENT_FIELD(UE::Trace::WideString, DebugName)
+UE_TRACE_EVENT_END()
+
 UE_TRACE_EVENT_BEGIN_EXTERN(ChaosVDLogger, ChaosVDDummyEvent)
 	UE_TRACE_EVENT_FIELD(int32, SolverID)
 UE_TRACE_EVENT_END()
@@ -278,6 +300,16 @@ UE_TRACE_EVENT_END()
 	FChaosVisualDebuggerTrace::TraceMidPhasesFromCollisionConstraints(CollisionConstraints);
 #endif
 
+#ifndef CVD_TRACE_NON_SOLVER_LOCATION
+	#define CVD_TRACE_NON_SOLVER_LOCATION(Location, DebugName) \
+	FChaosVisualDebuggerTrace::TraceNonSolverLocation(Location, DebugName);
+#endif
+
+#ifndef CVD_TRACE_NON_SOLVER_TRANSFORM
+	#define CVD_TRACE_NON_SOLVER_TRANSFORM(Transform, DebugName) \
+	FChaosVisualDebuggerTrace::TraceNonSolverTransform(Transform, DebugName);
+#endif
+
 struct FChaosVDContext;
 
 namespace Chaos
@@ -379,6 +411,22 @@ public:
 	 */
 	static CHAOS_API void TraceImplicitObject(FChaosVDImplicitObjectWrapper WrappedGeometryData);
 
+	/**
+	 * Traces the provided location using with the provided ID -
+	 * These are not tied to any solver step so they will be recorded as part of the current game frame data
+	 * @param InLocation Location to Trace
+	 * @param DebugNameID Name to use as ID
+	 */
+	static CHAOS_API void TraceNonSolverLocation(const FVector& InLocation, FStringView DebugNameID);
+
+	/**
+	 * Traces the provided transform using with the provided ID -
+	 * These are not tied to any solver step so they will be recorded as part of the current game frame data
+	 * @param InTransform Transform to Trace
+	 * @param DebugNameID Name to use as ID
+	 */
+	static CHAOS_API void TraceNonSolverTransform(const FTransform& InTransform, FStringView DebugNameID);
+
 	/** Returns true if the provided solver ID needs a Full Capture */
 	static bool ShouldPerformFullCapture(int32 SolverID);
 
diff --git a/Engine/Source/Runtime/Experimental/ChaosVDData/Public/ChaosVDRecording.h b/Engine/Source/Runtime/Experimental/ChaosVDData/Public/ChaosVDRecording.h
--- a/Engine/Source/Runtime/Experimental/ChaosVDData/Public/ChaosVDRecording.h
+++ b/Engine/Source/Runtime/Experimental/ChaosVDData/Public/ChaosVDRecording.h
@@ -27,6 +27,18 @@ struct FChaosVDStepData
 	TSet<int32> ParticlesDestroyedIDs;
 };
 
+struct FChaosVDTrackedLocation
+{
+	FString DebugName;
+	FVector Location;
+};
+
+struct FChaosVDTrackedTransform
+{
+	FString DebugName;
+	FTransform Transform;
+};
+
 typedef TArray<FChaosVDStepData, TInlineAllocator<16>> FChaosVDStepsContainer;
 
 struct CHAOSVDDATA_API FChaosVDSolverFrameData
@@ -44,6 +56,9 @@ struct FChaosVDGameFrameData
 {
 	uint64 FirstCycle;
 	uint64 LastCycle;
+
+	TMap<FName, FChaosVDTrackedLocation> RecordedNonSolverLocationsByID;
+	TMap<FName, FChaosVDTrackedTransform> RecordedNonSolverTransformsByID;
 };
 
 /**
