[ChaosVD] Added recording support for ShapeData and made use that in the CVD editor to offer more visualization options.

- Created a new DetailsView for CVD, that mimics the ActorDetails view with sub object editing. This allow us to show the geometry components in the details panel when a Particle actor is selected.
- Created a custom view for Geometry components, to only show CVD data, and expose the material so it can be changed if desired.
- Created derived classes of StaticMeshComponent and Instanced Static Mesh Components, so we can add CVD data to them and be visible/editable when a geometry component is selected.
- Added a new global visualization option that allows you to hide specific shape types (Query Only, Complex, etc)
- Query only shapes now by default are render with a translucent material

diff --git a/Engine/Plugins/Experimental/ChaosVD/Config/Engine.ini b/Engine/Plugins/Experimental/ChaosVD/Config/Engine.ini
--- a/Engine/Plugins/Experimental/ChaosVD/Config/Engine.ini
+++ b/Engine/Plugins/Experimental/ChaosVD/Config/Engine.ini
@@ -1,3 +1,5 @@
 ﻿[/Script/ChaosVD.ChaosVDEditorSettings]
 DebugMesh=/ChaosVD/SM_1M_Cube.SM_1M_Cube
-BasePhysicsVDWorld=/ChaosVD/ChaosVDBaseWorld.ChaosVDBaseWorld
\ No newline at end of file
+BasePhysicsVDWorld=/ChaosVD/ChaosVDBaseWorld.ChaosVDBaseWorld
+QueryOnlyMeshesMaterial=/ChaosVD/Materials/ChaosVDQueryOnlyMesh.ChaosVDQueryOnlyMesh
+SimMeshesMaterial=/ChaosVD/Materials/ChasoVDSimMesh.ChasoVDSimMesh
\ No newline at end of file
diff --git a/Engine/Plugins/Experimental/ChaosVD/Source/ChaosVD/ChaosVD.Build.cs b/Engine/Plugins/Experimental/ChaosVD/Source/ChaosVD/ChaosVD.Build.cs
--- a/Engine/Plugins/Experimental/ChaosVD/Source/ChaosVD/ChaosVD.Build.cs
+++ b/Engine/Plugins/Experimental/ChaosVD/Source/ChaosVD/ChaosVD.Build.cs
@@ -43,7 +43,9 @@ public class ChaosVD : ModuleRules
 				"SceneOutliner",
 				"TypedElementRuntime",
 				"TypedElementFramework",
-				"StatusBar"
+				"StatusBar",
+				"SubobjectEditor",
+				"SubobjectDataInterface"
 			}
 			);
 
diff --git a/Engine/Plugins/Experimental/ChaosVD/Source/ChaosVD/Private/ChaosVDEditorSettings.cpp b/Engine/Plugins/Experimental/ChaosVD/Source/ChaosVD/Private/ChaosVDEditorSettings.cpp
new file mode 100644
--- /dev/null
+++ b/Engine/Plugins/Experimental/ChaosVD/Source/ChaosVD/Private/ChaosVDEditorSettings.cpp
@@ -0,0 +1,14 @@
+// Copyright Epic Games, Inc. All Rights Reserved.
+
+#include "ChaosVDEditorSettings.h"
+
+void UChaosVDEditorSettings::PostEditChangeProperty(FPropertyChangedEvent& PropertyChangedEvent)
+{
+	Super::PostEditChangeProperty(PropertyChangedEvent);
+
+	const FName PropertyName = PropertyChangedEvent.GetPropertyName();
+	if (PropertyName == GET_MEMBER_NAME_CHECKED(UChaosVDEditorSettings, GeometryVisibilityFlags))
+	{
+		VisibilitySettingsChangedDelegate.Broadcast(this);
+	}
+}
diff --git a/Engine/Plugins/Experimental/ChaosVD/Source/ChaosVD/Private/ChaosVDGeometryBuilder.cpp b/Engine/Plugins/Experimental/ChaosVD/Source/ChaosVD/Private/ChaosVDGeometryBuilder.cpp
--- a/Engine/Plugins/Experimental/ChaosVD/Source/ChaosVD/Private/ChaosVDGeometryBuilder.cpp
+++ b/Engine/Plugins/Experimental/ChaosVD/Source/ChaosVD/Private/ChaosVDGeometryBuilder.cpp
@@ -2,6 +2,7 @@
 
 #include "ChaosVDGeometryBuilder.h"
 
+#include "ChaosVDGeometryDataComponent.h"
 #include "DynamicMeshToMeshDescription.h"
 #include "MeshConstraints.h"
 #include "MeshConstraintsUtil.h"
@@ -191,6 +192,12 @@ void FChaosVDGeometryBuilder::ApplyMeshToComponentFromKey(TWeakObjectPtr<UMeshCo
 			StaticMeshComponent->SetStaticMesh(*StaticMesh);
 		}
 	}
+
+	if (IChaosVDGeometryDataComponent* DataComponent = Cast<IChaosVDGeometryDataComponent>(MeshComponent.Get()))
+	{
+		DataComponent->SetIsMeshReady(true);
+		DataComponent->OnMeshReady()->ExecuteIfBound(*DataComponent);
+	}
 }
 
 bool FChaosVDGeometryBuilder::GameThreadTick(float DeltaTime)
diff --git a/Engine/Plugins/Experimental/ChaosVD/Source/ChaosVD/Private/ChaosVDGeometryDataComponent.cpp b/Engine/Plugins/Experimental/ChaosVD/Source/ChaosVD/Private/ChaosVDGeometryDataComponent.cpp
new file mode 100644
--- /dev/null
+++ b/Engine/Plugins/Experimental/ChaosVD/Source/ChaosVD/Private/ChaosVDGeometryDataComponent.cpp
@@ -0,0 +1,87 @@
+// Copyright Epic Games, Inc. All Rights Reserved.
+
+#include "ChaosVDGeometryDataComponent.h"
+
+#include "ChaosVDEditorSettings.h"
+#include "ChaosVDGeometryBuilder.h"
+#include "ChaosVDModule.h"
+#include "Components/MeshComponent.h"
+#include "MaterialDomain.h"
+#include "Materials/Material.h"
+
+void FChaosVDGeometryDataComponentBase::UpdateVisibility_Internal(const FChaosVDShapeCollisionData& InCollisionData, UMeshComponent* MeshComponent)
+{
+	if (!InCollisionData.bIsValid)
+	{
+		return;
+	}
+
+	if (MeshComponent)
+	{
+		if (const UChaosVDEditorSettings* EditorSettings = GetDefault<UChaosVDEditorSettings>())
+		{
+			if (UMaterialInterface* QueryOnlyMaterial = EditorSettings->QueryOnlyMeshesMaterial.Get())
+			{
+				const bool bIsQueryOnly = InCollisionData.bQueryCollision && !InCollisionData.bSimCollision;
+
+				if (bIsQueryOnly)
+				{
+					MeshComponent->SetMaterial(0, QueryOnlyMaterial);
+				}
+				else
+				{
+					MeshComponent->SetMaterial(0, UMaterial::GetDefaultMaterial(EMaterialDomain::MD_Surface));
+				}
+			}
+			else
+			{
+				ensure(false);
+				UE_LOG(LogChaosVDEditor, Error, TEXT("[%hs] Failed to get Query Only material for, applying the default mesh to all geometry"), ANSI_TO_TCHAR(__FUNCTION__));
+				
+				MeshComponent->SetMaterial(0, UMaterial::GetDefaultMaterial(EMaterialDomain::MD_Surface));
+			}
+
+			bool bShouldGeometryBeVisible = false;
+
+			// Complex vs Simple takes priority although this is subject to change
+			const bool bShouldBeVisibleIfComplex = InCollisionData.bIsComplex && EnumHasAnyFlags(static_cast<EChaosVDGeometryVisibilityFlags>(EditorSettings->GeometryVisibilityFlags), EChaosVDGeometryVisibilityFlags::Complex);
+			const bool bShouldBeVisibleIfSimple = !InCollisionData.bIsComplex && EnumHasAnyFlags(static_cast<EChaosVDGeometryVisibilityFlags>(EditorSettings->GeometryVisibilityFlags), EChaosVDGeometryVisibilityFlags::Simple);
+			if (bShouldBeVisibleIfComplex || bShouldBeVisibleIfSimple)
+			{
+				bShouldGeometryBeVisible = (InCollisionData.bSimCollision && EnumHasAnyFlags(static_cast<EChaosVDGeometryVisibilityFlags>(EditorSettings->GeometryVisibilityFlags), EChaosVDGeometryVisibilityFlags::Simulated))
+				|| (InCollisionData.bQueryCollision && EnumHasAnyFlags(static_cast<EChaosVDGeometryVisibilityFlags>(EditorSettings->GeometryVisibilityFlags), EChaosVDGeometryVisibilityFlags::Query));
+			}
+
+			MeshComponent->SetVisibility(bShouldGeometryBeVisible);
+		}
+	}
+}
+
+UE_DISABLE_OPTIMIZATION
+void FChaosVDGeometryDataComponentBase::UpdateDataFromShapeArray_Internal(const TArray<FChaosVDShapeCollisionData>& InShapeArray, FChaosVDShapeCollisionData& CollisionDataToUpdate)
+{
+	if (!ensureMsgf(RootImplicitObject.IsValid(), TEXT("Tried to Update Collision Data without a valid Implicit Object")))
+	{
+		return;
+	}
+
+	RootImplicitObject.GetReference()->VisitObjects([this, &CollisionDataToUpdate, &InShapeArray] (const Chaos::FImplicitObject* ImplicitA, const Chaos::FRigidTransform3& RelativeTransformA, const int32 RootObjectIndexA, const int32 ObjectIndex, const int32 LeafObjectIndexA)
+	{
+		if (!InShapeArray.IsValidIndex(RootObjectIndexA))
+		{
+			return true;
+		}
+
+		if (ImplicitA->GetTypeHash() == GeometryID)
+		{
+			CollisionDataToUpdate = InShapeArray[RootObjectIndexA];
+			CollisionDataToUpdate.bIsComplex = FChaosVDGeometryBuilder::DoesImplicitContainType(ImplicitA, Chaos::ImplicitObjectType::HeightField) || FChaosVDGeometryBuilder::DoesImplicitContainType(ImplicitA, Chaos::ImplicitObjectType::TriangleMesh);
+			CollisionDataToUpdate.bIsValid = true;
+
+			return false;
+		}
+
+		return true;
+	});
+}
+UE_ENABLE_OPTIMIZATION
\ No newline at end of file
diff --git a/Engine/Plugins/Experimental/ChaosVD/Source/ChaosVD/Private/ChaosVDModule.cpp b/Engine/Plugins/Experimental/ChaosVD/Source/ChaosVD/Private/ChaosVDModule.cpp
--- a/Engine/Plugins/Experimental/ChaosVD/Source/ChaosVD/Private/ChaosVDModule.cpp
+++ b/Engine/Plugins/Experimental/ChaosVD/Source/ChaosVD/Private/ChaosVDModule.cpp
@@ -10,6 +10,7 @@
 #include "ChaosVDEngine.h"
 #include "WorkspaceMenuStructure.h"
 #include "WorkspaceMenuStructureModule.h"
+#include "DetailsCustomizations/ChaosVDGeometryComponentCustomization.h"
 #include "DetailsCustomizations/ChaosVDParticleDataWrapperCustomization.h"
 #include "Trace/ChaosVDTraceManager.h"
 #include "Misc/Guid.h"
@@ -51,6 +52,8 @@ void FChaosVDModule::RegisterClassesCustomDetails() const
 {
 	FPropertyEditorModule& PropertyModule = FModuleManager::GetModuleChecked<FPropertyEditorModule>("PropertyEditor");
 	PropertyModule.RegisterCustomClassLayout("ChaosVDParticleActor", FOnGetDetailCustomizationInstance::CreateStatic(&FChaosVDParticleActorCustomization::MakeInstance));
+	PropertyModule.RegisterCustomClassLayout("ChaosVDInstancedStaticMeshComponent", FOnGetDetailCustomizationInstance::CreateStatic(&FChaosVDGeometryComponentCustomization::MakeInstance));
+	PropertyModule.RegisterCustomClassLayout("ChaosVDStaticMeshComponent", FOnGetDetailCustomizationInstance::CreateStatic(&FChaosVDGeometryComponentCustomization::MakeInstance));
 	PropertyModule.RegisterCustomPropertyTypeLayout("ChaosVDParticleDataWrapper", FOnGetPropertyTypeCustomizationInstance::CreateStatic(&FChaosVDParticleDataWrapperCustomization::MakeInstance));
 }
 
diff --git a/Engine/Plugins/Experimental/ChaosVD/Source/ChaosVD/Private/ChaosVDObjectDetailsTab.cpp b/Engine/Plugins/Experimental/ChaosVD/Source/ChaosVD/Private/ChaosVDObjectDetailsTab.cpp
--- a/Engine/Plugins/Experimental/ChaosVD/Source/ChaosVD/Private/ChaosVDObjectDetailsTab.cpp
+++ b/Engine/Plugins/Experimental/ChaosVD/Source/ChaosVD/Private/ChaosVDObjectDetailsTab.cpp
@@ -4,27 +4,24 @@
 
 #include "ChaosVDScene.h"
 #include "ChaosVDStyle.h"
+#include "Editor.h"
+#include "Elements/Framework/TypedElementSelectionSet.h"
 #include "Modules/ModuleManager.h"
 #include "PropertyEditorModule.h"
 #include "Templates/SharedPointer.h"
-#include "Editor.h"
 #include "Widgets/Docking/SDockTab.h"
-#include "Elements/Framework/TypedElementSelectionSet.h"
+#include "Widgets/SChaosVDDetailsView.h"
 
 #define LOCTEXT_NAMESPACE "ChaosVisualDebugger"
 
+class SSubobjectEditor;
+
+
 TSharedRef<SDockTab> FChaosVDObjectDetailsTab::HandleTabSpawned(const FSpawnTabArgs& Args)
 {
-	FPropertyEditorModule& PropertyEditorModule = FModuleManager::GetModuleChecked<FPropertyEditorModule>("PropertyEditor");
-	FDetailsViewArgs DetailsViewArgs;
-	DetailsViewArgs.bAllowSearch = false;
-	DetailsViewArgs.NameAreaSettings = FDetailsViewArgs::HideNameArea;
-
 	TSharedPtr<FChaosVDScene> ScenePtr = GetChaosVDScene().Pin();
 	check(ScenePtr);
 
-	DetailsPanel = PropertyEditorModule.CreateDetailView(DetailsViewArgs);
-
 	RegisterSelectionSetObject(ScenePtr->GetElementSelectionSet());
 
 	TSharedRef<SDockTab> DetailsPanelTab =
@@ -35,7 +32,7 @@ TSharedRef<SDockTab> FChaosVDObjectDetailsTab::HandleTabSpawned(const FSpawnTabA
 
 	DetailsPanelTab->SetContent
 	(
-		DetailsPanel.ToSharedRef()
+		SAssignNew(DetailsPanelView, SChaosVDDetailsView)
 	);
 
 	DetailsPanelTab->SetTabIcon(FChaosVDStyle::Get().GetBrush("TabIconDetailsPanel"));
@@ -52,7 +49,7 @@ void FChaosVDObjectDetailsTab::HandlePostSelectionChange(const UTypedElementSele
 		// We don't support multi selection yet
 		ensure(SelectedActors.Num() == 1);
 
-		DetailsPanel->SetObject(SelectedActors[0], true);
+		DetailsPanelView->SetSelectedObject(SelectedActors[0]);
 	}
 }
 
diff --git a/Engine/Plugins/Experimental/ChaosVD/Source/ChaosVD/Private/ChaosVDParticleActor.cpp b/Engine/Plugins/Experimental/ChaosVD/Source/ChaosVD/Private/ChaosVDParticleActor.cpp
--- a/Engine/Plugins/Experimental/ChaosVD/Source/ChaosVD/Private/ChaosVDParticleActor.cpp
+++ b/Engine/Plugins/Experimental/ChaosVD/Source/ChaosVD/Private/ChaosVDParticleActor.cpp
@@ -4,6 +4,8 @@
 
 #include "ChaosVDGeometryBuilder.h"
 #include "ChaosVDScene.h"
+#include "Components/ChaosVDInstancedStaticMeshComponent.h"
+#include "Components/ChaosVDStaticMeshComponent.h"
 #include "Components/MeshComponent.h"
 #include "Components/SceneComponent.h"
 #include "Components/StaticMeshComponent.h"
@@ -12,11 +14,6 @@
 #include "Engine/StaticMesh.h"
 #include "Visualizers/ChaosVDParticleDataVisualizer.h"
 
-static FAutoConsoleVariable CVarChaosVDHideVolumeAndBrushesHack(
-	TEXT("p.Chaos.VD.Tool.HideVolumeAndBrushesHack"),
-	true,
-	TEXT("If true, it will hide any geometry if its name contains Volume or Brush"));
-
 AChaosVDParticleActor::AChaosVDParticleActor(const FObjectInitializer& ObjectInitializer) : Super(ObjectInitializer)
 {
 	RootComponent = CreateDefaultSubobject<USceneComponent>(TEXT("SceneComponent0"));
@@ -42,6 +39,10 @@ void AChaosVDParticleActor::UpdateFromRecordedParticleData(const FChaosVDParticl
 
 	// TODO: We should store a ptr to the data and in our custom details panel draw it
 	ParticleDataViewer = InRecordedData;
+
+	// Now that we have updated particle data, update the Shape data and visibility as needed
+	UpdateShapeDataComponents();
+	UpdateGeometryComponentsVisibility();
 }
 
 void AChaosVDParticleActor::UpdateCollisionData(const TArray<TSharedPtr<FChaosVDParticlePairMidPhase>>& InRecordedMidPhases)
@@ -68,8 +69,13 @@ void AChaosVDParticleActor::UpdateCollisionData(const TArray<FChaosVDConstraint>
 	}
 }
 
-void AChaosVDParticleActor::UpdateGeometry(const Chaos::FImplicitObject* ImplicitObject, EChaosVDActorGeometryUpdateFlags OptionsFlags)
+void AChaosVDParticleActor::UpdateGeometry(const Chaos::FConstImplicitObjectPtr& InImplicitObject, EChaosVDActorGeometryUpdateFlags OptionsFlags)
 {
+	if (!InImplicitObject.IsValid())
+	{
+		return;
+	}
+	
 	if (EnumHasAnyFlags(OptionsFlags, EChaosVDActorGeometryUpdateFlags::ForceUpdate))
 	{
 		bIsGeometryDataGenerationStarted = false;
@@ -98,37 +104,32 @@ void AChaosVDParticleActor::UpdateGeometry(const Chaos::FImplicitObject* Implici
 			Chaos::FRigidTransform3 Transform;
 
 			// Heightfields need to be created as Static meshes and use normal Static Mesh components because we need LODs for them due to their high triangle count
-			if (FChaosVDGeometryBuilder::DoesImplicitContainType(ImplicitObject, Chaos::ImplicitObjectType::HeightField))
+			if (FChaosVDGeometryBuilder::DoesImplicitContainType(InImplicitObject, Chaos::ImplicitObjectType::HeightField))
 			{
 				constexpr int32 LODsToGenerateNum = 3;
 				constexpr int32 StartingMeshComponentIndex = 0;
-				GeometryGenerator->CreateMeshComponentsFromImplicit<UStaticMesh, UStaticMeshComponent>(ImplicitObject, this, OutGeneratedMeshComponents, Transform, StartingMeshComponentIndex, LODsToGenerateNum);
+				GeometryGenerator->CreateMeshComponentsFromImplicit<UStaticMesh, UChaosVDStaticMeshComponent>(InImplicitObject, this, OutGeneratedMeshComponents, Transform, StartingMeshComponentIndex, LODsToGenerateNum);
 			}
 			else
 			{
-				GeometryGenerator->CreateMeshComponentsFromImplicit<UStaticMesh, UInstancedStaticMeshComponent>(ImplicitObject, this, OutGeneratedMeshComponents, Transform);
+				GeometryGenerator->CreateMeshComponentsFromImplicit<UStaticMesh, UChaosVDInstancedStaticMeshComponent>(InImplicitObject, this, OutGeneratedMeshComponents, Transform);
 			}
 
 			if (OutGeneratedMeshComponents.Num() > 0)
 			{
 				MeshComponents.Append(OutGeneratedMeshComponents);
 
-				if (CVarChaosVDHideVolumeAndBrushesHack->GetBool())
+				for (TWeakObjectPtr<UMeshComponent> MeshComponent : MeshComponents)
 				{
-					// This is a temp hack (and is not performant) we need until we have a proper way to filer out trigger volumes/brushes at will
-					// Without this most maps will be covered in boxes
-					if (ParticleDataViewer.DebugName.Contains(TEXT("Brush")) || ParticleDataViewer.DebugName.Contains(TEXT("Volume")))
+					if (IChaosVDGeometryDataComponent* DataComponent = Cast<IChaosVDGeometryDataComponent>(MeshComponent.Get()))
 					{
-						for (TWeakObjectPtr<UMeshComponent> MeshComponent : OutGeneratedMeshComponents)
-						{
-							if (MeshComponent.IsValid())
-							{
-								MeshComponent->SetVisibility(false);
-							}
-						}
+						DataComponent->SetRootImplicitObject(InImplicitObject);
 					}
 				}
 
+				UpdateShapeDataComponents();
+				UpdateGeometryComponentsVisibility();
+
 				bIsGeometryDataGenerationStarted = true;
 			}
 		}
@@ -139,9 +140,9 @@ void AChaosVDParticleActor::UpdateGeometry(uint32 NewGeometryHash, EChaosVDActor
 {
 	if (const TSharedPtr<FChaosVDScene>& ScenePtr = OwningScene.Pin())
 	{
-		if (const Chaos::FConstImplicitObjectPtr* Geometry = ScenePtr->GetUpdatedGeometry(NewGeometryHash))
+		if (const Chaos::FConstImplicitObjectPtr& Geometry = ScenePtr->GetUpdatedGeometry(NewGeometryHash))
 		{
-			UpdateGeometry(*Geometry, OptionsFlags);
+			UpdateGeometry(Geometry, OptionsFlags);
 		}
 	}
 }
@@ -230,6 +231,39 @@ void AChaosVDParticleActor::PostEditChangeProperty(FPropertyChangedEvent& Proper
 	}
 }
 
+void AChaosVDParticleActor::UpdateShapeDataComponents()
+{
+	for (TWeakObjectPtr<UMeshComponent> MeshComponent : MeshComponents)
+	{
+		if (IChaosVDGeometryDataComponent* DataComponent = Cast<IChaosVDGeometryDataComponent>(MeshComponent.Get()))
+		{
+			DataComponent->UpdateDataFromShapeArray(ParticleDataViewer.CollisionDataPerShape);
+		}
+	}			
+}
+
+void AChaosVDParticleActor::UpdateGeometryComponentsVisibility()
+{
+	for (TWeakObjectPtr<UMeshComponent> MeshComponent : MeshComponents)
+	{
+		if (IChaosVDGeometryDataComponent* DataComponent = Cast<IChaosVDGeometryDataComponent>(MeshComponent.Get()))
+		{
+			// We need wait until we have a valid mesh
+			if (DataComponent->IsMeshReady())
+			{
+				DataComponent->UpdateVisibility();
+			}
+			else if (!DataComponent->OnMeshReady()->IsBound())
+			{
+				DataComponent->OnMeshReady()->BindWeakLambda(this, [](IChaosVDGeometryDataComponent& GeometryDataComponent)
+				{
+					GeometryDataComponent.UpdateVisibility();
+				});
+			}
+		}
+	}
+}
+
 void AChaosVDParticleActor::DrawVisualization(const FSceneView* View, FPrimitiveDrawInterface* PDI)
 {
 	for (const TPair<FStringView, TUniquePtr<FChaosVDDataVisualizerBase>>& VisualizerWithID : CVDVisualizers)
diff --git a/Engine/Plugins/Experimental/ChaosVD/Source/ChaosVD/Private/ChaosVDParticleActorCustomization.cpp b/Engine/Plugins/Experimental/ChaosVD/Source/ChaosVD/Private/ChaosVDParticleActorCustomization.cpp
--- a/Engine/Plugins/Experimental/ChaosVD/Source/ChaosVD/Private/ChaosVDParticleActorCustomization.cpp
+++ b/Engine/Plugins/Experimental/ChaosVD/Source/ChaosVD/Private/ChaosVDParticleActorCustomization.cpp
@@ -4,6 +4,8 @@
 
 #include "DetailCategoryBuilder.h"
 #include "DetailLayoutBuilder.h"
+#include "DetailsCustomizations/ChaosVDDetailsCustomizationUtils.h"
+
 
 FChaosVDParticleActorCustomization::FChaosVDParticleActorCustomization()
 {
@@ -18,18 +20,9 @@ TSharedRef<IDetailCustomization> FChaosVDParticleActorCustomization::MakeInstanc
 
 void FChaosVDParticleActorCustomization::CustomizeDetails(IDetailLayoutBuilder& DetailBuilder)
 {
+	DetailBuilder.EditCategory(ChaosVDVisualizationCategoryName, FText::GetEmpty(), ECategoryPriority::Important);
 
-	IDetailCategoryBuilder& YourCategory = DetailBuilder.EditCategory(ChaosVDVisualizationCategoryName, FText::GetEmpty(), ECategoryPriority::Important);
-	// Hide everything as the only thing we want to show in these actors is the Recorded debug data
-	TArray<FName> CurrentCategoryNames;
-	DetailBuilder.GetCategoryNames(CurrentCategoryNames);
-	for (const FName& CategoryToHide : CurrentCategoryNames)
-	{
-		if (!AllowedCategories.Contains(CategoryToHide))
-		{
-			DetailBuilder.HideCategory(CategoryToHide);
-		}
-	}
+	FChaosVDDetailsCustomizationUtils::HideAllCategories(DetailBuilder, AllowedCategories);
 
 	DetailBuilder.EditCategory(ChaosVDCategoryName).InitiallyCollapsed(false);
 }
diff --git a/Engine/Plugins/Experimental/ChaosVD/Source/ChaosVD/Private/ChaosVDScene.cpp b/Engine/Plugins/Experimental/ChaosVD/Source/ChaosVD/Private/ChaosVDScene.cpp
--- a/Engine/Plugins/Experimental/ChaosVD/Source/ChaosVD/Private/ChaosVDScene.cpp
+++ b/Engine/Plugins/Experimental/ChaosVD/Source/ChaosVD/Private/ChaosVDScene.cpp
@@ -16,8 +16,10 @@
 #include "DataWrappers/ChaosVDParticleDataWrapper.h"
 #include "Engine/Engine.h"
 #include "Engine/LevelStreamingDynamic.h"
+#include "Engine/StreamableManager.h"
 #include "Engine/World.h"
 #include "Misc/ScopedSlowTask.h"
+#include "Materials/Material.h"
 
 #include "UObject/Package.h"
 
@@ -49,6 +51,17 @@ void FChaosVDScene::Initialize()
 	PhysicsVDWorld = CreatePhysicsVDWorld();
 
 	GeometryGenerator = MakeShared<FChaosVDGeometryBuilder>();
+	
+	StreamableManager = MakeShared<FStreamableManager>();
+
+	if (UChaosVDEditorSettings* Settings = GetMutableDefault<UChaosVDEditorSettings>())
+	{
+		// TODO: Do an async load instead, and prepare a loading screen or notification popup
+		// Jira for tracking UE-191639
+		StreamableManager->RequestSyncLoad(Settings->QueryOnlyMeshesMaterial.ToSoftObjectPath());
+		
+		Settings->OnVisibilitySettingsChanged().AddRaw(this, &FChaosVDScene::HandleVisibilitySettingsChanged);
+	}
 
 	bIsInitialized = true;
 }
@@ -60,6 +73,11 @@ void FChaosVDScene::DeInitialize()
 		return;
 	}
 
+	if (UChaosVDEditorSettings* Settings = GetMutableDefault<UChaosVDEditorSettings>())
+	{
+		Settings->OnVisibilitySettingsChanged().RemoveAll(this);
+	}
+
 	SelectionSet->OnPreChange().RemoveAll(this);
 	SelectionSet->OnChanged().RemoveAll(this);
 
@@ -222,7 +240,6 @@ void FChaosVDScene::UpdateParticlesCollisionData(const FChaosVDStepData& InRecor
 	}
 }
 
-
 void FChaosVDScene::HandleNewGeometryData(const Chaos::FConstImplicitObjectPtr& GeometryData, const uint32 GeometryID) const
 {
 	NewGeometryAvailableDelegate.Broadcast(GeometryData, GeometryID);
@@ -293,13 +310,13 @@ void FChaosVDScene::CleanUpScene()
 	ParticlesBySolverID.Reset();
 }
 
-const Chaos::FConstImplicitObjectPtr* FChaosVDScene::GetUpdatedGeometry(int32 GeometryID) const
+Chaos::FConstImplicitObjectPtr FChaosVDScene::GetUpdatedGeometry(int32 GeometryID) const
 {
 	if (ensure(LoadedRecording.IsValid()))
 	{
 		if (const Chaos::FConstImplicitObjectPtr* Geometry = LoadedRecording->GetGeometryMap().Find(GeometryID))
 		{
-			return Geometry;
+			return *Geometry;
 		}
 	}
 
@@ -335,7 +352,7 @@ AChaosVDParticleActor* FChaosVDScene::SpawnParticleFromRecordedData(const FChaos
 		{
 			if (const Chaos::FConstImplicitObjectPtr* Geometry = LoadedRecording->GetGeometryMap().Find(InParticleData.GeometryHash))
 			{
-				NewActor->UpdateGeometry(Geometry->GetReference());
+				NewActor->UpdateGeometry(*Geometry);
 			}
 		}
 
@@ -449,6 +466,17 @@ void FChaosVDScene::ClearSelectionAndNotify()
 	SelectionSet->NotifyPendingChanges();
 }
 
+void FChaosVDScene::HandleVisibilitySettingsChanged(UChaosVDEditorSettings* SettingsObject)
+{
+	for (const TPair<int32, FChaosVDParticlesByIDMap>& ParticlesBySolver : ParticlesBySolverID)
+	{
+		for (const TPair<int32, AChaosVDParticleActor*>& ParticleWithIDPair : ParticlesBySolver.Value)
+		{
+			ParticleWithIDPair.Value->UpdateGeometryComponentsVisibility();
+		}
+	}
+}
+
 void FChaosVDScene::SetSelectedObject(UObject* SelectedObject)
 {
 	if (!SelectionSet)
diff --git a/Engine/Plugins/Experimental/ChaosVD/Source/ChaosVD/Private/Components/ChaosVDInstancedStaticMeshComponent.cpp b/Engine/Plugins/Experimental/ChaosVD/Source/ChaosVD/Private/Components/ChaosVDInstancedStaticMeshComponent.cpp
new file mode 100644
--- /dev/null
+++ b/Engine/Plugins/Experimental/ChaosVD/Source/ChaosVD/Private/Components/ChaosVDInstancedStaticMeshComponent.cpp
@@ -0,0 +1,19 @@
+// Copyright Epic Games, Inc. All Rights Reserved.
+
+#include "Components/ChaosVDInstancedStaticMeshComponent.h"
+
+void UChaosVDInstancedStaticMeshComponent::UpdateVisibility()
+{
+	UpdateVisibility_Internal(CollisionData, Cast<UMeshComponent>(this));
+}
+
+void UChaosVDInstancedStaticMeshComponent::UpdateDataFromShapeArray(const TArray<FChaosVDShapeCollisionData>& InShapeArray)
+{
+	UpdateDataFromShapeArray_Internal(InShapeArray, CollisionData);
+}
+
+void UChaosVDInstancedStaticMeshComponent::SetRootImplicitObject(const Chaos::FConstImplicitObjectPtr& InImplicitObject)
+{
+	RootImplicitObject = InImplicitObject;
+}
+
diff --git a/Engine/Plugins/Experimental/ChaosVD/Source/ChaosVD/Private/Components/ChaosVDStaticMeshComponent.cpp b/Engine/Plugins/Experimental/ChaosVD/Source/ChaosVD/Private/Components/ChaosVDStaticMeshComponent.cpp
new file mode 100644
--- /dev/null
+++ b/Engine/Plugins/Experimental/ChaosVD/Source/ChaosVD/Private/Components/ChaosVDStaticMeshComponent.cpp
@@ -0,0 +1,19 @@
+// Copyright Epic Games, Inc. All Rights Reserved.
+
+#include "Components/ChaosVDStaticMeshComponent.h"
+
+void UChaosVDStaticMeshComponent::UpdateVisibility()
+{
+	UpdateVisibility_Internal(CollisionData, Cast<UMeshComponent>(this));
+}
+
+void UChaosVDStaticMeshComponent::UpdateDataFromShapeArray(const TArray<FChaosVDShapeCollisionData>& InShapeArray)
+{
+	UpdateDataFromShapeArray_Internal(InShapeArray, CollisionData);
+}
+
+void UChaosVDStaticMeshComponent::SetRootImplicitObject(const Chaos::FConstImplicitObjectPtr& InImplicitObject)
+{
+	RootImplicitObject = InImplicitObject;
+}
+
diff --git a/Engine/Plugins/Experimental/ChaosVD/Source/ChaosVD/Private/DetailsCustomizations/ChaosVDDetailsCustomizationUtils.cpp b/Engine/Plugins/Experimental/ChaosVD/Source/ChaosVD/Private/DetailsCustomizations/ChaosVDDetailsCustomizationUtils.cpp
new file mode 100644
--- /dev/null
+++ b/Engine/Plugins/Experimental/ChaosVD/Source/ChaosVD/Private/DetailsCustomizations/ChaosVDDetailsCustomizationUtils.cpp
@@ -0,0 +1,19 @@
+// Copyright Epic Games, Inc. All Rights Reserved.
+
+#include "DetailsCustomizations/ChaosVDDetailsCustomizationUtils.h"
+
+#include "DetailLayoutBuilder.h"
+
+void FChaosVDDetailsCustomizationUtils::HideAllCategories(IDetailLayoutBuilder& DetailBuilder, const TSet<FName>& AllowedCategories)
+{
+	// Hide everything as the only thing we want to show in these actors is the Recorded debug data
+	TArray<FName> CurrentCategoryNames;
+	DetailBuilder.GetCategoryNames(CurrentCategoryNames);
+	for (const FName& CategoryToHide : CurrentCategoryNames)
+	{
+		if (!AllowedCategories.Contains(CategoryToHide))
+		{
+			DetailBuilder.HideCategory(CategoryToHide);
+		}
+	}
+}
diff --git a/Engine/Plugins/Experimental/ChaosVD/Source/ChaosVD/Private/DetailsCustomizations/ChaosVDGeometryComponentCustomization.cpp b/Engine/Plugins/Experimental/ChaosVD/Source/ChaosVD/Private/DetailsCustomizations/ChaosVDGeometryComponentCustomization.cpp
new file mode 100644
--- /dev/null
+++ b/Engine/Plugins/Experimental/ChaosVD/Source/ChaosVD/Private/DetailsCustomizations/ChaosVDGeometryComponentCustomization.cpp
@@ -0,0 +1,23 @@
+// Copyright Epic Games, Inc. All Rights Reserved.
+
+#include "DetailsCustomizations/ChaosVDGeometryComponentCustomization.h"
+
+#include "DetailLayoutBuilder.h"
+#include "DetailsCustomizations/ChaosVDDetailsCustomizationUtils.h"
+
+FChaosVDGeometryComponentCustomization::FChaosVDGeometryComponentCustomization()
+{
+	AllowedCategories.Add(FName("StaticMesh"));
+	AllowedCategories.Add(FName("Materials"));
+	AllowedCategories.Add(FName("Geometry Data"));
+}
+
+TSharedRef<IDetailCustomization> FChaosVDGeometryComponentCustomization::MakeInstance()
+{
+	return MakeShareable( new FChaosVDGeometryComponentCustomization );
+}
+
+void FChaosVDGeometryComponentCustomization::CustomizeDetails(IDetailLayoutBuilder& DetailBuilder)
+{
+	FChaosVDDetailsCustomizationUtils::HideAllCategories(DetailBuilder, AllowedCategories);
+}
diff --git a/Engine/Plugins/Experimental/ChaosVD/Source/ChaosVD/Private/Widgets/SChaosVDDetailsView.cpp b/Engine/Plugins/Experimental/ChaosVD/Source/ChaosVD/Private/Widgets/SChaosVDDetailsView.cpp
new file mode 100644
--- /dev/null
+++ b/Engine/Plugins/Experimental/ChaosVD/Source/ChaosVD/Private/Widgets/SChaosVDDetailsView.cpp
@@ -0,0 +1,138 @@
+// Copyright Epic Games, Inc. All Rights Reserved.
+
+#include "Widgets/SChaosVDDetailsView.h"
+
+#include "GameFramework/Actor.h"
+#include "PropertyEditorModule.h"
+#include "SSubobjectInstanceEditor.h"
+#include "SSubobjectEditor.h"
+#include "SSubobjectEditorModule.h"
+
+void SChaosVDDetailsView::Construct(const FArguments& InArgs)
+{
+	FPropertyEditorModule& PropertyEditorModule = FModuleManager::GetModuleChecked<FPropertyEditorModule>("PropertyEditor");
+	FDetailsViewArgs DetailsViewArgs;
+	DetailsViewArgs.bUpdatesFromSelection = false;
+	DetailsViewArgs.bLockable = true;
+	DetailsViewArgs.bAllowFavoriteSystem = false;
+	DetailsViewArgs.NameAreaSettings = FDetailsViewArgs::ObjectsUseNameArea | FDetailsViewArgs::ComponentsAndActorsUseNameArea;
+	DetailsViewArgs.bCustomNameAreaLocation = true;
+	DetailsViewArgs.bCustomFilterAreaLocation = false;
+	DetailsViewArgs.bShowSectionSelector = false;
+
+	DetailsView = PropertyEditorModule.CreateDetailView(DetailsViewArgs);
+
+	FModuleManager::LoadModuleChecked<FSubobjectEditorModule>("SubobjectEditor");
+
+	SubobjectEditor = SNew(SSubobjectInstanceEditor)
+		.ObjectContext(this, &SChaosVDDetailsView::GetRootContextObject)
+		.AllowEditing(false)
+		.OnSelectionUpdated(this, &SChaosVDDetailsView::OnSelectedSubobjectsChanged);
+
+	ChildSlot
+	[
+		SNew(SVerticalBox)
+		+SVerticalBox::Slot()
+		.Padding(10.f, 4.f, 0.f, 0.f)
+		.AutoHeight()
+		[
+			DetailsView->GetNameAreaWidget().ToSharedRef()
+		]
+		+SVerticalBox::Slot()
+		[
+			SNew(SSplitter)
+			.MinimumSlotHeight(80.0f)
+			.Orientation(Orient_Vertical)
+			.Style(FAppStyle::Get(), "SplitterDark")
+			.PhysicalSplitterHandleSize(2.0f)
+			+SSplitter::Slot()
+			.Value(0.2f)
+			[
+				SNew(SBox)
+				.Padding(FMargin(2.0f, 0.0f, 2.0f, 0.0f))
+				[
+					SubobjectEditor.ToSharedRef()
+				]
+			]
+			+SSplitter::Slot()
+			[
+				SNew( SVerticalBox )
+				+ SVerticalBox::Slot()
+				[
+					DetailsView.ToSharedRef()
+				]
+			]
+		]
+	];
+}
+
+void SChaosVDDetailsView::SetSelectedObject(UObject* NewObject)
+{
+	if (DetailsView->IsLocked())
+	{
+		return;
+	}
+
+	CurrentObjectInView = NewObject;
+	DetailsView->SetObject(NewObject, true);
+	SubobjectEditor->UpdateTree();
+}
+
+void SChaosVDDetailsView::OnSelectedSubobjectsChanged(const TArray<TSharedPtr<FSubobjectEditorTreeNode>>& SelectedNodes)
+{
+	if (SelectedNodes.Num() == 0)
+	{
+		return;
+	}
+
+	if (!DetailsView->IsLocked())
+	{
+		const bool bContainsRootActor = SelectedNodes.ContainsByPredicate([](const FSubobjectEditorTreeNodePtrType& Node)
+		{
+			if (Node.IsValid())
+			{
+				if (FSubobjectData* Data = Node->GetDataSource())
+				{
+					return Data->IsRootActor();
+				}
+			}
+
+			return false;
+		});
+
+		if (bContainsRootActor)
+		{
+			DetailsView->SetObject(GetRootContextObject());
+		}
+		else
+		{
+			TArray<UObject*> Components;
+		
+			if (AActor* ContextAsActor = Cast<AActor>(GetRootContextObject()))
+			{
+				Components.Reserve(SelectedNodes.Num());
+                for (const FSubobjectEditorTreeNodePtrType& Node : SelectedNodes)
+                {
+                	if (Node.IsValid())
+                	{
+                		if (FSubobjectData* Data = Node->GetDataSource())
+                		{
+                			if (Data->IsComponent())
+                			{
+                				if (const UActorComponent* Component = Data->FindComponentInstanceInActor(ContextAsActor))
+                				{
+                					if (Component != ContextAsActor->GetRootComponent())
+                					{
+                						Components.Add(const_cast<UActorComponent*>(Component));
+                					}
+                				}
+                			}
+                		}
+                	}
+                }
+			}
+
+			DetailsView->SetObjects(Components);
+		}
+	}
+}
diff --git a/Engine/Plugins/Experimental/ChaosVD/Source/ChaosVD/Public/ChaosVDEditorSettings.h b/Engine/Plugins/Experimental/ChaosVD/Source/ChaosVD/Public/ChaosVDEditorSettings.h
--- a/Engine/Plugins/Experimental/ChaosVD/Source/ChaosVD/Public/ChaosVDEditorSettings.h
+++ b/Engine/Plugins/Experimental/ChaosVD/Source/ChaosVD/Public/ChaosVDEditorSettings.h
@@ -10,6 +10,10 @@
 
 #include "ChaosVDEditorSettings.generated.h"
 
+class UChaosVDEditorSettings;
+class UMaterial;
+
+DECLARE_MULTICAST_DELEGATE_OneParam(FChaosVDVisibilitySettingsChaged, UChaosVDEditorSettings* CVDEditorSettingsObject)
 
 UENUM()
 enum class EChaosVDActorTrackingMode
@@ -28,6 +32,16 @@ enum class EChaosVDActorTrackingTarget
 	RecordedLocation,
 };
 
+UENUM(meta = (Bitflags, UseEnumValuesAsMaskValuesInEditor = "true"))
+enum class EChaosVDGeometryVisibilityFlags : uint8
+{
+	Query = 1 << 1,
+	Simulated = 1 << 2,
+	Simple = 1 << 3,
+	Complex = 1 << 4,
+};
+ENUM_CLASS_FLAGS(EChaosVDGeometryVisibilityFlags)
+
 UCLASS(config = Engine)
 class UChaosVDEditorSettings : public UObject
 {
@@ -57,4 +71,17 @@ public:
 
 	UPROPERTY(Config)
 	TSoftObjectPtr<UWorld> BasePhysicsVDWorld;
+
+	UPROPERTY(EditAnywhere, Category = "Geometry Visibility", meta = (Bitmask, BitmaskEnum = "/Script/ChaosVD.EChaosVDGeometryVisibilityFlags"))
+	uint8 GeometryVisibilityFlags = static_cast<uint8>(EChaosVDGeometryVisibilityFlags::Simulated | EChaosVDGeometryVisibilityFlags::Simple);
+
+	UPROPERTY(Config)
+	TSoftObjectPtr<UMaterial> QueryOnlyMeshesMaterial;
+
+	virtual void PostEditChangeProperty(FPropertyChangedEvent& PropertyChangedEvent) override;
+
+	FChaosVDVisibilitySettingsChaged& OnVisibilitySettingsChanged() { return VisibilitySettingsChangedDelegate; };
+
+protected:
+	FChaosVDVisibilitySettingsChaged VisibilitySettingsChangedDelegate;
 };
diff --git a/Engine/Plugins/Experimental/ChaosVD/Source/ChaosVD/Public/ChaosVDGeometryBuilder.h b/Engine/Plugins/Experimental/ChaosVD/Source/ChaosVD/Public/ChaosVDGeometryBuilder.h
--- a/Engine/Plugins/Experimental/ChaosVD/Source/ChaosVD/Public/ChaosVDGeometryBuilder.h
+++ b/Engine/Plugins/Experimental/ChaosVD/Source/ChaosVD/Public/ChaosVDGeometryBuilder.h
@@ -4,6 +4,7 @@
 
 #include "AsyncCompilationHelpers.h"
 #include "ChaosVDConvexMeshGenerator.h"
+#include "ChaosVDGeometryDataComponent.h"
 #include "ChaosVDHeightfieldMeshGenerator.h"
 #include "ChaosVDTriMeshGenerator.h"
 #include "Chaos/HeightField.h"
@@ -124,10 +125,11 @@ private:
 	 * Creates an empty DynamicMeshComponent/Static Mesh Component or Instanced mesh component and adds it to the actor
 	 * @param Owner Actor who will own the component
 	 * @param Name Name of the component. It has to be unique within the components in the owner actor
-	 * @param Transform (Optional) Transform to apply as Relative Transform in the component after its creating and attachment to the provided actor
+	 * @param Transform Transform to apply as Relative Transform in the component after its creating and attachment to the provided actor
+	 * @param DataComponentKey Key to match this component data to a specific geometry. Usually the type hash of the implicit object
 	 * */
 	template<typename ComponentType>
-	ComponentType* CreateMeshComponent(AActor* Owner, const FString& Name, const Chaos::FRigidTransform3& Transform) const;
+	ComponentType* CreateMeshComponent(AActor* Owner, const FString& Name, const Chaos::FRigidTransform3& Transform, uint32 DataComponentKey) const;
 
 	/**
 	 * Applies a mesh to a mesh component based on its type
@@ -222,11 +224,11 @@ template <typename MeshType, typename ComponentType>
 void FChaosVDGeometryBuilder::CreateMeshComponentsFromImplicit(const Chaos::FImplicitObject* InImplicitObject, AActor* Owner, TArray<TWeakObjectPtr<UMeshComponent>>& OutMeshComponents, Chaos::FRigidTransform3& Transform, const int32 Index, const int32 DesiredLODCount)
 {
 	static_assert(std::is_same_v<MeshType, UStaticMesh> || std::is_same_v<MeshType, UDynamicMesh>, "CreateMeshComponentsFromImplicit Only supports DynamicMesh and Static Mesh");
-	static_assert(std::is_same_v<ComponentType, UStaticMeshComponent> || std::is_same_v<ComponentType, UInstancedStaticMeshComponent> || std::is_same_v<MeshType, UDynamicMeshComponent>, "CreateMeshComponentsFromImplicit Only supports DynamicMeshComponent, Static MeshComponent and Instanced Static Mesh Component");
+	static_assert(std::is_base_of_v<UStaticMeshComponent, ComponentType> || std::is_base_of_v<UInstancedStaticMeshComponent, ComponentType> /*|| std::is_same_v<MeshType, UDynamicMeshComponent>*/, "CreateMeshComponentsFromImplicit Only supports DynamicMeshComponent, Static MeshComponent and Instanced Static Mesh Component");
 
 	// We could have you make the Mesh type as template and infer the component type based on that, but we also want to be able to create Static Meshes with either Instanced or normal static mesh components
-	constexpr bool bHasValidCombinationForStaticMesh =  std::is_same_v<MeshType, UStaticMesh> && (std::is_same_v<ComponentType, UStaticMeshComponent> || std::is_same_v<ComponentType, UInstancedStaticMeshComponent>);
-	constexpr bool bHasValidCombinationForDynamicMesh =  std::is_same_v<MeshType, UDynamicMesh> && std::is_same_v<ComponentType, UDynamicMeshComponent>;
+	constexpr bool bHasValidCombinationForStaticMesh =  std::is_same_v<MeshType, UStaticMesh> && (std::is_base_of_v<UStaticMeshComponent, ComponentType> || std::is_base_of_v<UInstancedStaticMeshComponent, ComponentType>);
+	constexpr bool bHasValidCombinationForDynamicMesh =  std::is_same_v<MeshType, UDynamicMesh> && std::is_base_of_v<UDynamicMeshComponent, ComponentType>;
 	static_assert(bHasValidCombinationForStaticMesh || bHasValidCombinationForDynamicMesh , "Incorrect Component type for Mesh type. Did you use a Dynamic Mesh with a Static Mesh component type?.");
 
 	using namespace Chaos;
@@ -266,9 +268,10 @@ void FChaosVDGeometryBuilder::CreateMeshComponentsFromImplicit(const Chaos::FImp
 
 			const FString Name = FString::Format(TEXT("{0} - {1}"), {TEXT("Sphere"), FString::FromInt(Index)});
 
-			MeshComponent = CreateMeshComponent<ComponentType>(Owner, Name, Transform);
-
 			const uint32 GeometryKey = Sphere->GetTypeHash();
+				
+			MeshComponent = CreateMeshComponent<ComponentType>(Owner, Name, Transform, GeometryKey);
+
 			Mesh = GetCachedMeshForImplicit<MeshType>(GeometryKey);
 
 			if (Mesh)
@@ -294,9 +297,11 @@ void FChaosVDGeometryBuilder::CreateMeshComponentsFromImplicit(const Chaos::FImp
 			const Chaos::TBox<FReal, 3>* Box = InImplicitObject->template GetObject<Chaos::TBox<FReal, 3>>();
 
 			const FString Name = FString::Format(TEXT("{0} - {1}"), {TEXT("Box"), FString::FromInt(Index)});
-			MeshComponent = CreateMeshComponent<ComponentType>(Owner, Name, Transform);
-
+			
 			const uint32 GeometryKey = Box->GetTypeHash();
+
+			MeshComponent = CreateMeshComponent<ComponentType>(Owner, Name, Transform, GeometryKey);
+
 			Mesh = GetCachedMeshForImplicit<MeshType>(GeometryKey);
 
 			if (Mesh)
@@ -325,7 +330,10 @@ void FChaosVDGeometryBuilder::CreateMeshComponentsFromImplicit(const Chaos::FImp
 
 			const FString Name = FString::Format(TEXT("{0} - {1}"), {TEXT("Capsule"), FString::FromInt(Index)});
 			const FRigidTransform3 StartingTransform;
-			MeshComponent = CreateMeshComponent<ComponentType>(Owner, Name, StartingTransform);
+
+			const uint32 GeometryKey = Capsule->GetTypeHash();
+				
+			MeshComponent = CreateMeshComponent<ComponentType>(Owner, Name, StartingTransform, GeometryKey);
 
 			// Re-adjust the location so the pivot is not the center of the capsule, and transform it based on the provided transform
 			const FVector FinalLocation = Transform.TransformPosition(Capsule->GetCenter() - Capsule->GetAxis() * Capsule->GetSegment().GetLength() * 0.5f);
@@ -335,7 +343,6 @@ void FChaosVDGeometryBuilder::CreateMeshComponentsFromImplicit(const Chaos::FImp
 			MeshComponent->SetRelativeLocation(FinalLocation);
 			MeshComponent->SetRelativeScale3D(Transform.GetScale3D());
 
-			const uint32 GeometryKey = Capsule->GetTypeHash();
 			Mesh = GetCachedMeshForImplicit<MeshType>(GeometryKey);
 
 			if (Mesh)
@@ -367,8 +374,13 @@ void FChaosVDGeometryBuilder::CreateMeshComponentsFromImplicit(const Chaos::FImp
 			if (const FConvex* Convex = GetGeometryBasedOnPackedType<FConvex>(InImplicitObject, Transform, PackedType))
 			{
 				const FString Name = FString::Format(TEXT("{0} - {1}"), {TEXT("Convex"), FString::FromInt(Index)});
-				MeshComponent = CreateMeshComponent<ComponentType>(Owner, Name, Transform);
+
 				
+				uint32 DataComponentKey = InImplicitObject->GetTypeHash();
+				MeshComponent = CreateMeshComponent<ComponentType>(Owner, Name, Transform, DataComponentKey);
+
+				// For the Cache key, we need the hash of the geometry itself without scale or transformations, because these will be applied to the
+				// mesh component. For example, we don't want to generate two meshes for a box, because in one instance was scaled. We only one mesh for one box, and scale the component as needed
 				const uint32 GeometryKey = Convex->GetTypeHash();
 				Mesh = GetCachedMeshForImplicit<MeshType>(GeometryKey);
 	
@@ -397,8 +409,13 @@ void FChaosVDGeometryBuilder::CreateMeshComponentsFromImplicit(const Chaos::FImp
 			if (const FTriangleMeshImplicitObject* TriangleMesh = GetGeometryBasedOnPackedType<FTriangleMeshImplicitObject>(InImplicitObject, Transform, PackedType))
 			{
 				const FString Name = FString::Format(TEXT("{0} - {1}"), {TEXT("Trimesh"), FString::FromInt(Index)});
-				MeshComponent = CreateMeshComponent<ComponentType>(Owner, Name, Transform);
 
+				// For the Component data key, we need the hash of the implicit as it is as we will need to match it when looking for shape data
+				uint32 DataComponentKey = InImplicitObject->GetTypeHash();
+				MeshComponent = CreateMeshComponent<ComponentType>(Owner, Name, Transform, DataComponentKey);
+
+				// For the Cache key, we need the hash of the geometry itself without scale or transformations, because these will be applied to the
+				// mesh component. For example, we don't want to generate two meshes for a box, because in one instance was scaled. We only one mesh for one box, and scale the component as needed
 				const uint32 GeometryKey = TriangleMesh->GetTypeHash();
 				Mesh = GetCachedMeshForImplicit<MeshType>(GeometryKey);
 	
@@ -424,8 +441,13 @@ void FChaosVDGeometryBuilder::CreateMeshComponentsFromImplicit(const Chaos::FImp
 			if (const FHeightField* HeightField = GetGeometryBasedOnPackedType<FHeightField>(InImplicitObject, Transform, PackedType))
 			{
 				const FString Name = FString::Format(TEXT("{0} - {1}"), {TEXT("HeightField"), FString::FromInt(Index)});
-				MeshComponent = CreateMeshComponent<ComponentType>(Owner, Name, Transform);
 
+				// For the Component data key, we need the hash of the implicit as it is as we will need to match it when looking for shape data
+				uint32 DataComponentKey = InImplicitObject->GetTypeHash();
+				MeshComponent = CreateMeshComponent<ComponentType>(Owner, Name, Transform, DataComponentKey);
+
+				// For the Cache key, we need the hash of the geometry itself without scale or transformations, because these will be applied to the
+				// mesh component. For example, we don't want to generate two meshes for a box, because in one instance was scaled. We only one mesh for one box, and scale the component as needed
 				const uint32 GeometryKey = HeightField->GetTypeHash();
 				Mesh = GetCachedMeshForImplicit<MeshType>(GeometryKey);
 	
@@ -480,7 +502,7 @@ MeshType* FChaosVDGeometryBuilder::GetCachedMeshForImplicit(const uint32 Geometr
 }
 
 template <typename ComponentType>
-ComponentType* FChaosVDGeometryBuilder::CreateMeshComponent(AActor* Owner, const FString& Name, const Chaos::FRigidTransform3& Transform) const
+ComponentType* FChaosVDGeometryBuilder::CreateMeshComponent(AActor* Owner, const FString& Name, const Chaos::FRigidTransform3& Transform, uint32 DataComponentKey) const
 {
 	ComponentType* MeshComponent = NewObject<ComponentType>(Owner, *Name);
 
@@ -493,17 +515,24 @@ ComponentType* FChaosVDGeometryBuilder::CreateMeshComponent(AActor* Owner, const
 
 	MeshComponent->bSelectable = true;
 
-	if constexpr (std::is_same_v<ComponentType, UDynamicMeshComponent> || std::is_same_v<ComponentType, UStaticMeshComponent>)
+	constexpr bool bIsStaticMeshComponent = std::is_base_of_v<UStaticMeshComponent, ComponentType> && !std::is_base_of_v<UInstancedStaticMeshComponent, ComponentType>;
+
+	if constexpr (std::is_base_of_v<UDynamicMeshComponent, ComponentType> || bIsStaticMeshComponent)
 	{
 		MeshComponent->SetRelativeTransform(Transform);
 	}
-	else if constexpr (std::is_same_v<ComponentType, UInstancedStaticMeshComponent>)
+	else if constexpr (std::is_base_of_v<UInstancedStaticMeshComponent, ComponentType>)
 	{
 		// If we have negative scale we need to force reverse the culling mode in this component otherwise the faces will be inverted
 		MeshComponent->SetReverseCulling(HasNegativeScale(Transform));
 		MeshComponent->AddInstance(Transform);
 	}
 
+	if (IChaosVDGeometryDataComponent* DataComponent = Cast<IChaosVDGeometryDataComponent>(MeshComponent))
+	{
+		DataComponent->SetGeometryID(DataComponentKey);
+	}
+
 	return MeshComponent;
 }
 
diff --git a/Engine/Plugins/Experimental/ChaosVD/Source/ChaosVD/Public/ChaosVDGeometryDataComponent.h b/Engine/Plugins/Experimental/ChaosVD/Source/ChaosVD/Public/ChaosVDGeometryDataComponent.h
new file mode 100644
--- /dev/null
+++ b/Engine/Plugins/Experimental/ChaosVD/Source/ChaosVD/Public/ChaosVDGeometryDataComponent.h
@@ -0,0 +1,83 @@
+// Copyright Epic Games, Inc. All Rights Reserved.
+
+#pragma once
+#include "Containers/Array.h"
+#include "Chaos/ImplicitFwd.h"
+#include "DataWrappers/ChaosVDCollisionDataWrappers.h"
+#include "HAL/Platform.h"
+#include "UObject/Interface.h"
+
+#include "Chaos/ImplicitObject.h"
+
+#include "ChaosVDGeometryDataComponent.generated.h"
+
+class IChaosVDGeometryDataComponent;
+class UMeshComponent;
+
+namespace Chaos
+{
+	class FImplicitObject;
+}
+
+DECLARE_DELEGATE_OneParam(FChaosVDMeshReadyDelegate, IChaosVDGeometryDataComponent&)
+
+UINTERFACE()
+class UChaosVDGeometryDataComponent : public UInterface
+{
+	GENERATED_BODY()
+};
+
+/**
+ * Interface with a default implementation for any Geometry component that
+ * contains CVD data
+ */
+class IChaosVDGeometryDataComponent
+{
+	GENERATED_BODY()
+
+public:
+
+	/** Returns the Geometry ID used to identify the geometry data this component represents */
+	virtual uint32 GetGeometryID() const { PURE_VIRTUAL(IChaosVDGeometryDataComponent::GetGeometryID, return 0;); }
+
+	/** Sets the Geometry ID used to identify the geometry data this component represents */
+	virtual void SetGeometryID(uint32 ID){ PURE_VIRTUAL(IChaosVDGeometryDataComponent::SetGeometryID, ); }
+
+	/** True if the mesh this component represents is ready for use */
+	virtual bool IsMeshReady() const { PURE_VIRTUAL(IChaosVDGeometryDataComponent::IsMeshReady, return false;); }
+
+	/** Sets if the mesh this component represents is ready for use or not */
+	virtual void SetIsMeshReady(bool bIsReady) { PURE_VIRTUAL(IChaosVDGeometryDataComponent::SetIsMeshReady, ); }
+	
+	/** Triggers when the mesh this component represents is ready */
+	virtual FChaosVDMeshReadyDelegate* OnMeshReady() { PURE_VIRTUAL(IChaosVDGeometryDataComponent::OnMeshReady, return nullptr;); }
+
+	/** Stores a shader ptr to the root implicit object of the implicit object this component represents  */
+	virtual void SetRootImplicitObject(const Chaos::FConstImplicitObjectPtr& InImplicitObject) { PURE_VIRTUAL(IChaosVDGeometryDataComponent::SetRootImplicitObject, ); } 
+
+	/** Returns a ptr to the CVD Collision Data */
+	virtual FChaosVDShapeCollisionData* GetCollisionData() { PURE_VIRTUAL( IChaosVDGeometryDataComponent::GetCollisionData, return nullptr; ); }
+
+	/** Updates the visibility of this component based on the stored CVD data*/
+	virtual void UpdateVisibility() { PURE_VIRTUAL(IChaosVDGeometryDataComponent::UpdateVisibility, ); }
+
+	/** Finds and updates the Shape data using the provided array as source*/
+	virtual void UpdateDataFromShapeArray(const TArray<FChaosVDShapeCollisionData>& InShapeArray) { PURE_VIRTUAL(IChaosVDGeometryDataComponent::UpdateDataFromShapeArray, ); }
+};
+
+/** Base Implementation for a GeometryDataComponent */
+class FChaosVDGeometryDataComponentBase
+{
+	
+protected:
+
+	void UpdateVisibility_Internal(const FChaosVDShapeCollisionData& InCollisionData, UMeshComponent* MeshComponent);
+	
+	void UpdateDataFromShapeArray_Internal(const TArray<FChaosVDShapeCollisionData>& InShapeArray, FChaosVDShapeCollisionData& CollisionDataToUpdate);
+	
+	uint32 GeometryID = 0;
+	bool bIsMeshReady = false;
+
+	FChaosVDMeshReadyDelegate MeshReadyDelegate;
+	Chaos::FConstImplicitObjectPtr RootImplicitObject = nullptr;
+};
diff --git a/Engine/Plugins/Experimental/ChaosVD/Source/ChaosVD/Public/ChaosVDObjectDetailsTab.h b/Engine/Plugins/Experimental/ChaosVD/Source/ChaosVD/Public/ChaosVDObjectDetailsTab.h
--- a/Engine/Plugins/Experimental/ChaosVD/Source/ChaosVD/Public/ChaosVDObjectDetailsTab.h
+++ b/Engine/Plugins/Experimental/ChaosVD/Source/ChaosVD/Public/ChaosVDObjectDetailsTab.h
@@ -10,13 +10,12 @@
 struct FChaosVDParticleDebugData;
 
 class AActor;
-class IDetailsView;
 class FName;
 class FSpawnTabArgs;
 class FTabManager;
+class SChaosVDDetailsView;
 class SDockTab;
 
-
 /** Spawns and handles and instance for the visual debugger details panel */
 class FChaosVDObjectDetailsTab : public FChaosVDTabSpawnerBase, public FChaosVDSceneSelectionObserver, public TSharedFromThis<FChaosVDObjectDetailsTab>
 {
@@ -33,5 +32,5 @@ protected:
 	virtual void HandlePostSelectionChange(const UTypedElementSelectionSet* ChangedSelectionSet) override;
 
 	FDelegateHandle SelectionDelegateHandle;
-	TSharedPtr<IDetailsView> DetailsPanel;
+	TSharedPtr<SChaosVDDetailsView> DetailsPanelView;
 };
diff --git a/Engine/Plugins/Experimental/ChaosVD/Source/ChaosVD/Public/ChaosVDParticleActor.h b/Engine/Plugins/Experimental/ChaosVD/Source/ChaosVD/Public/ChaosVDParticleActor.h
--- a/Engine/Plugins/Experimental/ChaosVD/Source/ChaosVD/Public/ChaosVDParticleActor.h
+++ b/Engine/Plugins/Experimental/ChaosVD/Source/ChaosVD/Public/ChaosVDParticleActor.h
@@ -1,6 +1,7 @@
 // Copyright Epic Games, Inc. All Rights Reserved.
 
 #pragma once
+#include "ChaosVDGeometryDataComponent.h"
 #include "Chaos/Core.h"
 #include "DataWrappers/ChaosVDParticleDataWrapper.h"
 #include "DataWrappers/ChaosVDCollisionDataWrappers.h"
@@ -43,13 +44,10 @@ public:
 
 	void UpdateCollisionData(const TArray<TSharedPtr<FChaosVDParticlePairMidPhase>>& InRecordedMidPhases);
 	void UpdateCollisionData(const TArray<FChaosVDConstraint>& InRecordedConstraints);
-	void UpdateGeometry(const Chaos::FImplicitObject* ImplicitObject, EChaosVDActorGeometryUpdateFlags OptionsFlags = EChaosVDActorGeometryUpdateFlags::None);
+	void UpdateGeometry(const Chaos::FConstImplicitObjectPtr& InImplicitObject, EChaosVDActorGeometryUpdateFlags OptionsFlags = EChaosVDActorGeometryUpdateFlags::None);
 
 	void UpdateGeometry(uint32 NewGeometryHash, EChaosVDActorGeometryUpdateFlags OptionsFlags = EChaosVDActorGeometryUpdateFlags::None);
 
-	UE_DEPRECATED(5.4, "Use UpdateGeometry with FImplicitObject* to update collision implicit object.")
-	void UpdateGeometry(const TSharedPtr<Chaos::FImplicitObject>& ImplicitObject, EChaosVDActorGeometryUpdateFlags OptionsFlags = EChaosVDActorGeometryUpdateFlags::None) {check(false);}
-
 	void SetScene(const TSharedPtr<FChaosVDScene>& InScene);
 
 	virtual void BeginDestroy() override;
@@ -66,7 +64,11 @@ public:
 	virtual void PostEditChangeProperty(FPropertyChangedEvent& PropertyChangedEvent) override;
 #endif
 
+	void UpdateGeometryComponentsVisibility();
+
 protected:
+
+	void UpdateShapeDataComponents();
 	
 	UPROPERTY(EditAnywhere, Category = "Viewport Visualization Flags", meta = (Bitmask, BitmaskEnum = "/Script/ChaosVD.EChaosVDParticleDataVisualizationFlags"))
 	uint8 LocalParticleDataVisualizationFlags;
diff --git a/Engine/Plugins/Experimental/ChaosVD/Source/ChaosVD/Public/ChaosVDScene.h b/Engine/Plugins/Experimental/ChaosVD/Source/ChaosVD/Public/ChaosVDScene.h
--- a/Engine/Plugins/Experimental/ChaosVD/Source/ChaosVD/Public/ChaosVDScene.h
+++ b/Engine/Plugins/Experimental/ChaosVD/Source/ChaosVD/Public/ChaosVDScene.h
@@ -5,13 +5,16 @@
 #include "ChaosVDRecording.h"
 #include "Containers/UnrealString.h"
 #include "Containers/Map.h"
-#include "UObject/ObjectMacros.h"
 #include "UObject/GCObject.h"
+#include "UObject/ObjectMacros.h"
+#include "UObject/ObjectPtr.h"
 
+class UChaosVDEditorSettings;
 class FChaosVDGeometryBuilder;
 class AChaosVDParticleActor;
 class FReferenceCollector;
 class UObject;
+class UMaterial;
 class UTypedElementSelectionSet;
 class UWorld;
 
@@ -64,7 +67,7 @@ public:
 	FChaosVDGeometryDataLoaded& OnNewGeometryAvailable(){ return NewGeometryAvailableDelegate; }
 
 	// No need to deprecate the old version since it is not a public API nor inline 
-	const Chaos::FConstImplicitObjectPtr* GetUpdatedGeometry(int32 GeometryID) const;
+	Chaos::FConstImplicitObjectPtr GetUpdatedGeometry(int32 GeometryID) const;
 	
 	/** Adds an object to the selection set if it was not selected already, making it selected in practice */
 	void SetSelectedObject(UObject* SelectedObject);
@@ -111,6 +114,8 @@ private:
 
 	void ClearSelectionAndNotify();
 
+	void HandleVisibilitySettingsChanged(UChaosVDEditorSettings* SettingsObject);
+
 	/** UWorld instance used to represent the recorded debug data */
 	TObjectPtr<UWorld> PhysicsVDWorld = nullptr;
 
@@ -128,5 +133,8 @@ private:
 	/** Array of actors with hit proxies that need to be updated */
 	TArray<AActor*> PendingActorsToUpdateSelectionProxy;
 
+	/** Scene Streamable manager that we'll use to async load any assets we depend on */
+	TSharedPtr<struct FStreamableManager> StreamableManager;
+
 	bool bIsInitialized = false;
 };
diff --git a/Engine/Plugins/Experimental/ChaosVD/Source/ChaosVD/Public/Components/ChaosVDInstancedStaticMeshComponent.h b/Engine/Plugins/Experimental/ChaosVD/Source/ChaosVD/Public/Components/ChaosVDInstancedStaticMeshComponent.h
new file mode 100644
--- /dev/null
+++ b/Engine/Plugins/Experimental/ChaosVD/Source/ChaosVD/Public/Components/ChaosVDInstancedStaticMeshComponent.h
@@ -0,0 +1,44 @@
+// Copyright Epic Games, Inc. All Rights Reserved.
+
+#pragma once
+
+#include "Components/InstancedStaticMeshComponent.h"
+#include "ChaosVDGeometryDataComponent.h"
+#include "ChaosVDInstancedStaticMeshComponent.generated.h"
+
+/** CVD version of a Instance Static Mesh Component that holds additional CVD data */
+UCLASS(HideCategories=("Transform"), MinimalAPI)
+class UChaosVDInstancedStaticMeshComponent : public UInstancedStaticMeshComponent, public IChaosVDGeometryDataComponent, public FChaosVDGeometryDataComponentBase
+{
+	GENERATED_BODY()
+	
+	UChaosVDInstancedStaticMeshComponent(const FObjectInitializer& ObjectInitializer)
+	: Super(ObjectInitializer)
+	{
+	}
+	
+	// BEGIN IChaosVDGeometryDataComponent Interface
+	virtual FChaosVDShapeCollisionData* GetCollisionData() override { return &CollisionData; }
+	
+	virtual uint32 GetGeometryID() const override { return GeometryID; }
+	
+	virtual void  SetGeometryID(uint32 ID) override { GeometryID = ID; }
+
+	virtual bool IsMeshReady() const override { return bIsMeshReady; }
+	
+	virtual void SetIsMeshReady(bool bIsReady) override { bIsMeshReady = bIsReady; }
+
+	virtual FChaosVDMeshReadyDelegate* OnMeshReady() override { return &MeshReadyDelegate; }
+	
+	virtual void SetRootImplicitObject(const Chaos::FConstImplicitObjectPtr& InImplicitObject) override;
+	
+	virtual void UpdateVisibility() override;
+	
+	virtual void UpdateDataFromShapeArray(const TArray<FChaosVDShapeCollisionData>& InShapeArray) override;
+	// END IChaosVDGeometryDataComponent Interface
+
+protected:
+	
+	UPROPERTY(VisibleAnywhere, Category="Geometry Data", meta=(FullyExpand = true))
+	FChaosVDShapeCollisionData CollisionData;
+};
diff --git a/Engine/Plugins/Experimental/ChaosVD/Source/ChaosVD/Public/Components/ChaosVDStaticMeshComponent.h b/Engine/Plugins/Experimental/ChaosVD/Source/ChaosVD/Public/Components/ChaosVDStaticMeshComponent.h
new file mode 100644
--- /dev/null
+++ b/Engine/Plugins/Experimental/ChaosVD/Source/ChaosVD/Public/Components/ChaosVDStaticMeshComponent.h
@@ -0,0 +1,44 @@
+// Copyright Epic Games, Inc. All Rights Reserved.
+
+#pragma once
+
+#include "Components/StaticMeshComponent.h"
+#include "ChaosVDGeometryDataComponent.h"
+#include "ChaosVDStaticMeshComponent.generated.h"
+
+/** CVD version of a Static Mesh Component that holds additional CVD data */
+UCLASS(HideCategories=("Transform"), MinimalAPI)
+class UChaosVDStaticMeshComponent : public UStaticMeshComponent, public IChaosVDGeometryDataComponent, public FChaosVDGeometryDataComponentBase
+{
+	GENERATED_BODY()
+public:
+	UChaosVDStaticMeshComponent(const FObjectInitializer& ObjectInitializer)
+		: Super(ObjectInitializer)
+	{		
+	}
+
+	// BEGIN IChaosVDGeometryDataComponent Interface
+	virtual FChaosVDShapeCollisionData* GetCollisionData() override { return &CollisionData; }
+	
+	virtual uint32 GetGeometryID() const override { return GeometryID; }
+	
+	virtual void  SetGeometryID(uint32 ID) override { GeometryID = ID; }
+
+	virtual bool IsMeshReady() const override { return bIsMeshReady; }
+	
+	virtual void SetIsMeshReady(bool bIsReady) override { bIsMeshReady = bIsReady; }
+
+	virtual FChaosVDMeshReadyDelegate* OnMeshReady() override { return &MeshReadyDelegate; }
+	
+	virtual void SetRootImplicitObject(const Chaos::FConstImplicitObjectPtr& InImplicitObject) override;
+	
+	virtual void UpdateVisibility() override;
+	
+	virtual void UpdateDataFromShapeArray(const TArray<FChaosVDShapeCollisionData>& InShapeArray) override;;
+	// END IChaosVDGeometryDataComponent Interface
+
+protected:
+	
+	UPROPERTY(VisibleAnywhere, Category="Geometry Data", meta=(FullyExpand = true))
+	FChaosVDShapeCollisionData CollisionData;
+};
diff --git a/Engine/Plugins/Experimental/ChaosVD/Source/ChaosVD/Public/DetailsCustomizations/ChaosVDDetailsCustomizationUtils.h b/Engine/Plugins/Experimental/ChaosVD/Source/ChaosVD/Public/DetailsCustomizations/ChaosVDDetailsCustomizationUtils.h
new file mode 100644
--- /dev/null
+++ b/Engine/Plugins/Experimental/ChaosVD/Source/ChaosVD/Public/DetailsCustomizations/ChaosVDDetailsCustomizationUtils.h
@@ -0,0 +1,21 @@
+// Copyright Epic Games, Inc. All Rights Reserved.
+
+#pragma once
+
+#include "Containers/Set.h"
+#include "UObject/NameTypes.h"
+
+class IDetailLayoutBuilder;
+class FName;
+
+/** Helper Class for CVD Custom Details view */
+class FChaosVDDetailsCustomizationUtils
+{
+public:
+	/**
+	 * Hides all categories of this view, except the ones provided in the Allowed Categories set
+	 * @param DetailBuilder Layout builder of the class we are customizing
+	 * @param AllowedCategories Set of category names we do not want to hide 
+	 */
+	static void HideAllCategories(IDetailLayoutBuilder& DetailBuilder, const TSet<FName>& AllowedCategories = TSet<FName>());
+};
diff --git a/Engine/Plugins/Experimental/ChaosVD/Source/ChaosVD/Public/DetailsCustomizations/ChaosVDGeometryComponentCustomization.h b/Engine/Plugins/Experimental/ChaosVD/Source/ChaosVD/Public/DetailsCustomizations/ChaosVDGeometryComponentCustomization.h
new file mode 100644
--- /dev/null
+++ b/Engine/Plugins/Experimental/ChaosVD/Source/ChaosVD/Public/DetailsCustomizations/ChaosVDGeometryComponentCustomization.h
@@ -0,0 +1,20 @@
+// Copyright Epic Games, Inc. All Rights Reserved.
+
+#pragma once
+#include "IDetailCustomization.h"
+
+/**
+ * Customization view for CVD Geometry component classes 
+ */
+class FChaosVDGeometryComponentCustomization : public IDetailCustomization
+{
+public:
+	FChaosVDGeometryComponentCustomization();
+
+	static TSharedRef<IDetailCustomization> MakeInstance();
+
+	virtual void CustomizeDetails(IDetailLayoutBuilder& DetailBuilder) override;
+
+private:
+	TSet<FName> AllowedCategories;
+};
diff --git a/Engine/Plugins/Experimental/ChaosVD/Source/ChaosVD/Public/Widgets/SChaosVDDetailsView.h b/Engine/Plugins/Experimental/ChaosVD/Source/ChaosVD/Public/Widgets/SChaosVDDetailsView.h
new file mode 100644
--- /dev/null
+++ b/Engine/Plugins/Experimental/ChaosVD/Source/ChaosVD/Public/Widgets/SChaosVDDetailsView.h
@@ -0,0 +1,45 @@
+// Copyright Epic Games, Inc. All Rights Reserved.
+
+#pragma once
+
+#include "Templates/SharedPointer.h"
+#include "UObject/WeakObjectPtrTemplates.h"
+#include "Widgets/SCompoundWidget.h"
+
+class FSubobjectEditorTreeNode;
+class AActor;
+class SBox;
+class SSplitter;
+class IDetailsView;
+class SSubobjectEditor;
+class UObject;
+
+/**
+ * Simple details view with SubObject Editor
+ */
+class CHAOSVD_API SChaosVDDetailsView : public SCompoundWidget
+{
+public:
+	SLATE_BEGIN_ARGS(SChaosVDDetailsView)
+		{
+		}
+	SLATE_END_ARGS()
+
+	/** Constructs this widget with InArgs */
+	void Construct(const FArguments& InArgs);
+
+	/** Updates the current object this view details is viewing */
+	void SetSelectedObject(UObject* NewObject);
+
+protected:
+	UObject* GetRootContextObject() const { return CurrentObjectInView.IsValid() ? CurrentObjectInView.Get() : nullptr; };
+	
+	void OnSelectedSubobjectsChanged(const TArray<TSharedPtr<FSubobjectEditorTreeNode>>& SelectedNodes);
+
+	TWeakObjectPtr<UObject> CurrentObjectInView;
+	
+	TSharedPtr<IDetailsView> DetailsView;
+
+	// The subobject editor provides a tree widget that allows for editing of subobjects
+	TSharedPtr<SSubobjectEditor> SubobjectEditor;
+};
diff --git a/Engine/Source/Runtime/Experimental/Chaos/Private/ChaosVisualDebugger/ChaosVisualDebuggerTrace.cpp b/Engine/Source/Runtime/Experimental/Chaos/Private/ChaosVisualDebugger/ChaosVisualDebuggerTrace.cpp
--- a/Engine/Source/Runtime/Experimental/Chaos/Private/ChaosVisualDebugger/ChaosVisualDebuggerTrace.cpp
+++ b/Engine/Source/Runtime/Experimental/Chaos/Private/ChaosVisualDebugger/ChaosVisualDebuggerTrace.cpp
@@ -141,6 +141,21 @@ void FChaosVisualDebuggerTrace::TraceParticle(Chaos::FGeometryParticleHandle* Pa
 		ParticleDataWrapper.GeometryHash = GeometryHash;
 		ParticleDataWrapper.SolverID = ContextData.Id;
 		
+		const Chaos::FShapeInstanceArray& ShapesInstancesArray = ParticleHandle->ShapeInstances();
+		ParticleDataWrapper.CollisionDataPerShape.Reserve(ShapesInstancesArray.Num());
+		
+		for (const Chaos::FShapeInstancePtr& ShapeData : ShapesInstancesArray)
+		{
+			const Chaos::FCollisionData& CollisionData = ShapeData->GetCollisionData();
+			FChaosVDShapeCollisionData CVDCollisionData;
+			CVDCollisionData.bQueryCollision = CollisionData.bQueryCollision;
+			CVDCollisionData.bIsProbe = CollisionData.bIsProbe;
+			CVDCollisionData.bSimCollision = CollisionData.bSimCollision;
+			CVDCollisionData.CollisionTraceType = static_cast<EChaosVDCollisionTraceFlag>(CollisionData.CollisionTraceType);
+
+			ParticleDataWrapper.CollisionDataPerShape.Add(MoveTemp(CVDCollisionData));
+		}
+		
 		FChaosVDScopedTLSBufferAccessor TLSDataBuffer;
 
 		FMemoryWriter MemWriterAr(TLSDataBuffer.BufferRef);
diff --git a/Engine/Source/Runtime/Experimental/ChaosVisualDebugger/Private/DataWrappers/ChaosVDCollisionDataWrappers.cpp b/Engine/Source/Runtime/Experimental/ChaosVisualDebugger/Private/DataWrappers/ChaosVDCollisionDataWrappers.cpp
--- a/Engine/Source/Runtime/Experimental/ChaosVisualDebugger/Private/DataWrappers/ChaosVDCollisionDataWrappers.cpp
+++ b/Engine/Source/Runtime/Experimental/ChaosVisualDebugger/Private/DataWrappers/ChaosVDCollisionDataWrappers.cpp
@@ -94,3 +94,14 @@ bool FChaosVDParticlePairMidPhase::Serialize(FArchive& Ar)
 
 	return true;
 }
+
+bool FChaosVDShapeCollisionData::Serialize(FArchive& Ar)
+{
+	Ar << CollisionTraceType;
+	
+	FArchive_Serialize_BitfieldBool(Ar, bSimCollision);
+	FArchive_Serialize_BitfieldBool(Ar, bQueryCollision);
+	FArchive_Serialize_BitfieldBool(Ar, bIsProbe);
+
+	return true;
+}
diff --git a/Engine/Source/Runtime/Experimental/ChaosVisualDebugger/Private/DataWrappers/ChaosVDParticleDataWrapper.cpp b/Engine/Source/Runtime/Experimental/ChaosVisualDebugger/Private/DataWrappers/ChaosVDParticleDataWrapper.cpp
--- a/Engine/Source/Runtime/Experimental/ChaosVisualDebugger/Private/DataWrappers/ChaosVDParticleDataWrapper.cpp
+++ b/Engine/Source/Runtime/Experimental/ChaosVisualDebugger/Private/DataWrappers/ChaosVDParticleDataWrapper.cpp
@@ -108,5 +108,7 @@ bool FChaosVDParticleDataWrapper::Serialize(FArchive& Ar)
 
 	Ar << bHasValidData;
 
+	Ar << CollisionDataPerShape;
+
 	return true;
 }
diff --git a/Engine/Source/Runtime/Experimental/ChaosVisualDebugger/Public/DataWrappers/ChaosVDCollisionDataWrappers.h b/Engine/Source/Runtime/Experimental/ChaosVisualDebugger/Public/DataWrappers/ChaosVDCollisionDataWrappers.h
--- a/Engine/Source/Runtime/Experimental/ChaosVisualDebugger/Public/DataWrappers/ChaosVDCollisionDataWrappers.h
+++ b/Engine/Source/Runtime/Experimental/ChaosVisualDebugger/Public/DataWrappers/ChaosVDCollisionDataWrappers.h
@@ -305,3 +305,55 @@ inline FArchive& operator<<(FArchive& Ar, FChaosVDParticlePairMidPhase& Data)
 	Data.Serialize(Ar);
 	return Ar;
 }
+
+UENUM()
+enum class EChaosVDCollisionTraceFlag
+{
+	/** Use project physics settings (DefaultShapeComplexity) */
+	UseDefault,
+	/** Create both simple and complex shapes. Simple shapes are used for regular scene queries and collision tests. Complex shape (per poly) is used for complex scene queries.*/
+	UseSimpleAndComplex,
+	/** Create only simple shapes. Use simple shapes for all scene queries and collision tests.*/
+	UseSimpleAsComplex,
+	/** Create only complex shapes (per poly). Use complex shapes for all scene queries and collision tests. Can be used in simulation for static shapes only (i.e can be collided against but not moved through forces or velocity.) */
+	UseComplexAsSimple,
+	/** */
+	MAX,
+};
+
+USTRUCT()
+struct CHAOSVDRUNTIME_API FChaosVDShapeCollisionData
+{
+	GENERATED_BODY()
+
+	UPROPERTY(VisibleAnywhere, Category=CollisionData)
+	EChaosVDCollisionTraceFlag CollisionTraceType = EChaosVDCollisionTraceFlag::UseDefault;
+
+	UPROPERTY(VisibleAnywhere, Category=CollisionData)
+	uint8 bSimCollision : 1 = false;
+	UPROPERTY(VisibleAnywhere, Category=CollisionData)
+	uint8 bQueryCollision : 1 = false;
+	UPROPERTY(VisibleAnywhere, Category=CollisionData)
+	uint8 bIsProbe : 1 = false;
+
+	bool bIsComplex = false;
+
+	bool bIsValid = false;
+
+	bool Serialize(FArchive& Ar);
+};
+
+template<>
+struct TStructOpsTypeTraits<FChaosVDShapeCollisionData> : public TStructOpsTypeTraitsBase2<FChaosVDShapeCollisionData>
+{
+	enum
+	{
+		WithSerializer = true,
+	};
+};
+
+inline FArchive& operator<<(FArchive& Ar, FChaosVDShapeCollisionData& Data)
+{
+	Data.Serialize(Ar);
+	return Ar;
+}
diff --git a/Engine/Source/Runtime/Experimental/ChaosVisualDebugger/Public/DataWrappers/ChaosVDParticleDataWrapper.h b/Engine/Source/Runtime/Experimental/ChaosVisualDebugger/Public/DataWrappers/ChaosVDParticleDataWrapper.h
--- a/Engine/Source/Runtime/Experimental/ChaosVisualDebugger/Public/DataWrappers/ChaosVDParticleDataWrapper.h
+++ b/Engine/Source/Runtime/Experimental/ChaosVisualDebugger/Public/DataWrappers/ChaosVDParticleDataWrapper.h
@@ -444,6 +444,9 @@ struct FChaosVDParticleDataWrapper : public FChaosVDParticleDataBase
 
 	UPROPERTY(EditAnywhere, Category= "Particle Collision")
 	TArray<FChaosVDConstraint> ParticleConstraints;
+
+	UPROPERTY(EditAnywhere, Category= "Particle Collision")
+	TArray<FChaosVDShapeCollisionData> CollisionDataPerShape;
 	
 	/** Only used during recording */
 	TSharedPtr<FString> DebugNamePtr;
