[ChaosVD] Removing Disable Optimization Macro in ChaosVDGeometryDataComponent

diff --git a/Engine/Plugins/Experimental/ChaosVD/Source/ChaosVD/Private/ChaosVDGeometryDataComponent.cpp b/Engine/Plugins/Experimental/ChaosVD/Source/ChaosVD/Private/ChaosVDGeometryDataComponent.cpp
--- a/Engine/Plugins/Experimental/ChaosVD/Source/ChaosVD/Private/ChaosVDGeometryDataComponent.cpp
+++ b/Engine/Plugins/Experimental/ChaosVD/Source/ChaosVD/Private/ChaosVDGeometryDataComponent.cpp
@@ -57,7 +57,6 @@ void FChaosVDGeometryDataComponentBase::UpdateVisibility_Internal(const FChaosVD
 	}
 }
 
-UE_DISABLE_OPTIMIZATION
 void FChaosVDGeometryDataComponentBase::UpdateDataFromShapeArray_Internal(const TArray<FChaosVDShapeCollisionData>& InShapeArray, FChaosVDShapeCollisionData& CollisionDataToUpdate)
 {
 	if (!ensureMsgf(RootImplicitObject.IsValid(), TEXT("Tried to Update Collision Data without a valid Implicit Object")))
@@ -84,4 +83,3 @@ void FChaosVDGeometryDataComponentBase::UpdateDataFromShapeArray_Internal(const
 		return true;
 	});
 }
-UE_ENABLE_OPTIMIZATION
\ No newline at end of file
