[ChaosVD] Fix Linux CIS issue on ChaosVDGeometryDataComponent

diff --git a/Engine/Plugins/Experimental/ChaosVD/Source/ChaosVD/Public/ChaosVDGeometryDataComponent.h b/Engine/Plugins/Experimental/ChaosVD/Source/ChaosVD/Public/ChaosVDGeometryDataComponent.h
--- a/Engine/Plugins/Experimental/ChaosVD/Source/ChaosVD/Public/ChaosVDGeometryDataComponent.h
+++ b/Engine/Plugins/Experimental/ChaosVD/Source/ChaosVD/Public/ChaosVDGeometryDataComponent.h
@@ -38,31 +38,31 @@ class IChaosVDGeometryDataComponent
 public:
 
 	/** Returns the Geometry ID used to identify the geometry data this component represents */
-	virtual uint32 GetGeometryID() const { PURE_VIRTUAL(IChaosVDGeometryDataComponent::GetGeometryID, return 0;); }
+	virtual uint32 GetGeometryID() const PURE_VIRTUAL(IChaosVDGeometryDataComponent::GetGeometryID, return 0;);
 
 	/** Sets the Geometry ID used to identify the geometry data this component represents */
-	virtual void SetGeometryID(uint32 ID){ PURE_VIRTUAL(IChaosVDGeometryDataComponent::SetGeometryID, ); }
+	virtual void SetGeometryID(uint32 ID) PURE_VIRTUAL(IChaosVDGeometryDataComponent::SetGeometryID);
 
 	/** True if the mesh this component represents is ready for use */
-	virtual bool IsMeshReady() const { PURE_VIRTUAL(IChaosVDGeometryDataComponent::IsMeshReady, return false;); }
+	virtual bool IsMeshReady() const  PURE_VIRTUAL(IChaosVDGeometryDataComponent::IsMeshReady, return false;);
 
 	/** Sets if the mesh this component represents is ready for use or not */
-	virtual void SetIsMeshReady(bool bIsReady) { PURE_VIRTUAL(IChaosVDGeometryDataComponent::SetIsMeshReady, ); }
+	virtual void SetIsMeshReady(bool bIsReady) PURE_VIRTUAL(IChaosVDGeometryDataComponent::SetIsMeshReady);
 	
 	/** Triggers when the mesh this component represents is ready */
-	virtual FChaosVDMeshReadyDelegate* OnMeshReady() { PURE_VIRTUAL(IChaosVDGeometryDataComponent::OnMeshReady, return nullptr;); }
+	virtual FChaosVDMeshReadyDelegate* OnMeshReady() PURE_VIRTUAL(IChaosVDGeometryDataComponent::OnMeshReady, return nullptr;);
 
 	/** Stores a shader ptr to the root implicit object of the implicit object this component represents  */
-	virtual void SetRootImplicitObject(const Chaos::FConstImplicitObjectPtr& InImplicitObject) { PURE_VIRTUAL(IChaosVDGeometryDataComponent::SetRootImplicitObject, ); } 
+	virtual void SetRootImplicitObject(const Chaos::FConstImplicitObjectPtr& InImplicitObject) PURE_VIRTUAL(IChaosVDGeometryDataComponent::SetRootImplicitObject);
 
 	/** Returns a ptr to the CVD Collision Data */
-	virtual FChaosVDShapeCollisionData* GetCollisionData() { PURE_VIRTUAL( IChaosVDGeometryDataComponent::GetCollisionData, return nullptr; ); }
+	virtual FChaosVDShapeCollisionData* GetCollisionData() PURE_VIRTUAL(IChaosVDGeometryDataComponent::GetCollisionData, return nullptr;);
 
 	/** Updates the visibility of this component based on the stored CVD data*/
-	virtual void UpdateVisibility() { PURE_VIRTUAL(IChaosVDGeometryDataComponent::UpdateVisibility, ); }
+	virtual void UpdateVisibility() PURE_VIRTUAL(IChaosVDGeometryDataComponent::UpdateVisibility);
 
 	/** Finds and updates the Shape data using the provided array as source*/
-	virtual void UpdateDataFromShapeArray(const TArray<FChaosVDShapeCollisionData>& InShapeArray) { PURE_VIRTUAL(IChaosVDGeometryDataComponent::UpdateDataFromShapeArray, ); }
+	virtual void UpdateDataFromShapeArray(const TArray<FChaosVDShapeCollisionData>& InShapeArray) PURE_VIRTUAL(IChaosVDGeometryDataComponent::UpdateDataFromShapeArray);
 };
 
 /** Base Implementation for a GeometryDataComponent */
