[ChaosVD] Added custom viewport menu overlay and lighting controls

- Now the viewport will have on top a version of the normal Level viewport menu overlay, but with only the supported options. This allow us to change between different view modes, and have an easier quick access to the CVD Geometry Visualization options.
- Made the Directional light and Skysphere part of the transient map created for CVD instead of being a locked sublevel. This allows changing their properties like in the normal Level Editor.
- Created an Editor Mode and Editor mode tools in CVD, to allow re-use the existing transform toolbar and widgets. This also opens the door to bring other elements of the Level editor, and extension features.
- In Viewport, when the directional light is selected, the transform gizmo and toolbar will be shown (It is explicitly disabled for Particle Actors).
- Fixed an issue with the object tracking code. Part of the logic was on a widget so when that widget tab was closed, object tracking stopped working.

diff --git a/Engine/Plugins/Experimental/ChaosVD/Config/Engine.ini b/Engine/Plugins/Experimental/ChaosVD/Config/Engine.ini
--- a/Engine/Plugins/Experimental/ChaosVD/Config/Engine.ini
+++ b/Engine/Plugins/Experimental/ChaosVD/Config/Engine.ini
@@ -1,5 +1,4 @@
 ﻿[/Script/ChaosVD.ChaosVDEditorSettings]
-DebugMesh=/ChaosVD/SM_1M_Cube.SM_1M_Cube
-BasePhysicsVDWorld=/ChaosVD/ChaosVDBaseWorld.ChaosVDBaseWorld
 QueryOnlyMeshesMaterial=/ChaosVD/Materials/ChaosVDQueryOnlyMesh.ChaosVDQueryOnlyMesh
-SimMeshesMaterial=/ChaosVD/Materials/ChasoVDSimMesh.ChasoVDSimMesh
\ No newline at end of file
+SimMeshesMaterial=/ChaosVD/Materials/ChasoVDSimMesh.ChasoVDSimMesh
+SkySphereActorClass=/ChaosVD/BP_CVD_SkySphere.BP_CVD_SkySphere_C
\ No newline at end of file
diff --git a/Engine/Plugins/Experimental/ChaosVD/Source/ChaosVD/Private/ChaosVDEditorMode.cpp b/Engine/Plugins/Experimental/ChaosVD/Source/ChaosVD/Private/ChaosVDEditorMode.cpp
new file mode 100644
--- /dev/null
+++ b/Engine/Plugins/Experimental/ChaosVD/Source/ChaosVD/Private/ChaosVDEditorMode.cpp
@@ -0,0 +1,108 @@
+// Copyright Epic Games, Inc. All Rights Reserved.
+
+#include "ChaosVDEditorMode.h"
+
+#include "ChaosVDParticleActor.h"
+#include "EditorModeManager.h"
+#include "EditorViewportClient.h"
+#include "Elements/Actor/ActorElementEditorViewportInteractionCustomization.h"
+#include "Selection.h"
+
+const FEditorModeID UChaosVDEditorMode::EM_ChaosVisualDebugger(TEXT("EM_ChaosVisualDebugger"));
+
+FVector FChaosVDEdModeWidgetHelper::GetWidgetLocation() const
+{
+	if (const AActor* SelectedActor = GetCurrentSelectedActor())
+	{
+		return SelectedActor->GetActorLocation();
+	}
+	else
+	{
+		return FLegacyEdModeWidgetHelper::GetWidgetLocation();
+	}
+}
+
+AActor* FChaosVDEdModeWidgetHelper::GetCurrentSelectedActor() const
+{
+	if (!ensure(Owner))
+	{
+		return nullptr;
+	}
+	
+	if (USelection* CurrentSelection = Owner->GetSelectedActors())
+	{
+		ensureMsgf(CurrentSelection->Num() < 2, TEXT("CVD does not support multi selection yet"));
+
+		return CurrentSelection->GetTop<AActor>();
+	}
+
+	return nullptr;
+}
+
+bool FChaosVDEdModeWidgetHelper::InputDelta(FEditorViewportClient* InViewportClient, FViewport* InViewport, FVector& InDrag, FRotator& InRot, FVector& InScale)
+{
+	bool bHandled = false;
+	if (InViewportClient->bWidgetAxisControlledByDrag)
+	{
+		if (AActor* SelectedActor = GetCurrentSelectedActor())
+		{
+			
+			FActorElementEditorViewportInteractionCustomization::ApplyDeltaToActor(
+							SelectedActor, true, &InDrag,
+							&InRot, &InScale, SelectedActor->GetActorLocation(), FInputDeviceState());
+
+			bHandled = true;
+		}
+	}
+	
+	if (bHandled)
+	{
+		return true;
+	}
+	else
+	{
+		return FLegacyEdModeWidgetHelper::InputDelta(InViewportClient, InViewport, InDrag, InRot, InScale);
+	}
+}
+
+bool FChaosVDEdModeWidgetHelper::UsesTransformWidget(UE::Widget::EWidgetMode CheckMode) const
+{
+	// Currently all modes follow the same rule. ParticleActors cannot be moved
+	return UsesTransformWidget();
+}
+
+bool FChaosVDEdModeWidgetHelper::UsesTransformWidget() const
+{
+	if (const AActor* SelectedActor = GetCurrentSelectedActor())
+	{
+		// Particle Actors cannot be moved
+		return !SelectedActor->IsA(AChaosVDParticleActor::StaticClass());
+	}
+
+	return false;
+}
+
+UChaosVDEditorMode::UChaosVDEditorMode() : UBaseLegacyWidgetEdMode()
+{
+	Info = FEditorModeInfo(UChaosVDEditorMode::EM_ChaosVisualDebugger,
+		NSLOCTEXT("ChaosVisualDebugger", "ChaosVisualDebuggerEditorModeName", "Chaos Visual Debugger Editor Mode"),
+		FSlateIcon(),
+		false);
+
+	Toolkit = nullptr;
+}
+
+bool UChaosVDEditorMode::CanAutoSave() const
+{
+	return false;
+}
+
+UWorld* UChaosVDEditorMode::GetWorld() const
+{
+	return CVDWorldPtr.IsValid() ? CVDWorldPtr.Get() : nullptr;
+}
+
+TSharedRef<FLegacyEdModeWidgetHelper> UChaosVDEditorMode::CreateWidgetHelper()
+{
+	return MakeShared<FChaosVDEdModeWidgetHelper>();
+}
diff --git a/Engine/Plugins/Experimental/ChaosVD/Source/ChaosVD/Private/ChaosVDEditorModeTools.cpp b/Engine/Plugins/Experimental/ChaosVD/Source/ChaosVD/Private/ChaosVDEditorModeTools.cpp
new file mode 100644
--- /dev/null
+++ b/Engine/Plugins/Experimental/ChaosVD/Source/ChaosVD/Private/ChaosVDEditorModeTools.cpp
@@ -0,0 +1,58 @@
+// Copyright Epic Games, Inc. All Rights Reserved.
+
+#include "ChaosVDEditorModeTools.h"
+
+#include "ChaosVDEditorMode.h"
+#include "ChaosVDScene.h"
+#include "Selection.h"
+
+
+FChaosVDEditorModeTools::FChaosVDEditorModeTools(const TWeakPtr<FChaosVDScene>& InScenePtr) : FEditorModeTools(), ScenePtr(InScenePtr)
+{
+	USelection::SelectNoneEvent.RemoveAll(this);
+}
+
+FString FChaosVDEditorModeTools::GetReferencerName() const
+{
+	return TEXT("FChaosVDEditorModeTools");
+}
+
+USelection* FChaosVDEditorModeTools::GetSelectedActors() const
+{
+	if (const TSharedPtr<FChaosVDScene> SceneSharedPtr = ScenePtr.Pin())
+	{
+		return SceneSharedPtr->GetActorSelectionObject();
+	}
+	
+	return nullptr;
+}
+
+USelection* FChaosVDEditorModeTools::GetSelectedObjects() const
+{
+	if (const TSharedPtr<FChaosVDScene> SceneSharedPtr = ScenePtr.Pin())
+	{
+		return SceneSharedPtr->GetObjectsSelectionObject();
+	}
+
+	return nullptr;
+}
+
+USelection* FChaosVDEditorModeTools::GetSelectedComponents() const
+{
+	if (const TSharedPtr<FChaosVDScene> SceneSharedPtr = ScenePtr.Pin())
+	{
+		return SceneSharedPtr->GetComponentsSelectionObject();
+	}
+	
+	return nullptr;
+}
+
+UWorld* FChaosVDEditorModeTools::GetWorld() const
+{
+	if (const TSharedPtr<FChaosVDScene> SceneSharedPtr = ScenePtr.Pin())
+	{
+		return SceneSharedPtr->GetUnderlyingWorld();
+	}
+
+	return nullptr;
+}
diff --git a/Engine/Plugins/Experimental/ChaosVD/Source/ChaosVD/Private/ChaosVDPlaybackControllerObserver.cpp b/Engine/Plugins/Experimental/ChaosVD/Source/ChaosVD/Private/ChaosVDPlaybackControllerObserver.cpp
--- a/Engine/Plugins/Experimental/ChaosVD/Source/ChaosVD/Private/ChaosVDPlaybackControllerObserver.cpp
+++ b/Engine/Plugins/Experimental/ChaosVD/Source/ChaosVD/Private/ChaosVDPlaybackControllerObserver.cpp
@@ -14,6 +14,11 @@ FChaosVDPlaybackControllerObserver::~FChaosVDPlaybackControllerObserver()
 	}
 }
 
+TWeakPtr<FChaosVDPlaybackController> FChaosVDPlaybackControllerObserver::GetObservedController()
+{
+	return PlaybackController;
+}
+
 void FChaosVDPlaybackControllerObserver::RegisterNewController(TWeakPtr<FChaosVDPlaybackController> NewController)
 {
 	if (PlaybackController != NewController)
diff --git a/Engine/Plugins/Experimental/ChaosVD/Source/ChaosVD/Private/ChaosVDPlaybackViewportClient.cpp b/Engine/Plugins/Experimental/ChaosVD/Source/ChaosVD/Private/ChaosVDPlaybackViewportClient.cpp
--- a/Engine/Plugins/Experimental/ChaosVD/Source/ChaosVD/Private/ChaosVDPlaybackViewportClient.cpp
+++ b/Engine/Plugins/Experimental/ChaosVD/Source/ChaosVD/Private/ChaosVDPlaybackViewportClient.cpp
@@ -2,18 +2,28 @@
 
 #include "ChaosVDPlaybackViewportClient.h"
 
-#include "CameraController.h"
 #include "ChaosVDEditorSettings.h"
 #include "ChaosVDParticleActor.h"
+#include "ChaosVDPlaybackController.h"
 #include "ChaosVDScene.h"
-#include "EngineUtils.h"
+#include "ChaosVDSkySphereInterface.h"
+#include "EditorModeManager.h"
 #include "Elements/Framework/TypedElementSelectionSet.h"
-#include "Visualizers/ChaosVDDebugDrawUtils.h"
+#include "Engine/DirectionalLight.h"
+#include "EngineUtils.h"
 #include "SEditorViewport.h"
+#include "Selection.h"
+#include "UnrealWidget.h"
+#include "Visualizers/ChaosVDDebugDrawUtils.h"
 
-
-FChaosVDPlaybackViewportClient::FChaosVDPlaybackViewportClient() : FEditorViewportClient(nullptr), CVDWorld(nullptr)
+FChaosVDPlaybackViewportClient::FChaosVDPlaybackViewportClient(const TSharedPtr<FEditorModeTools>& InModeTools) : FEditorViewportClient(InModeTools.Get()), CVDWorld(nullptr)
 {
+	Widget->SetUsesEditorModeTools(InModeTools.Get());
+
+	if (GEngine)
+	{
+		GEngine->OnActorMoving().AddRaw(this, &FChaosVDPlaybackViewportClient::HandleActorMoving);
+	}
 }
 
 FChaosVDPlaybackViewportClient::~FChaosVDPlaybackViewportClient()
@@ -25,10 +35,17 @@ FChaosVDPlaybackViewportClient::~FChaosVDPlaybackViewportClient()
 			ScenePtr->OnObjectFocused().Remove(ObjectFocusedDelegateHandle);
 		}
 	}
+
+	if (GEngine)
+	{
+		GEngine->OnActorMoving().RemoveAll(this);
+	}
 }
 
 void FChaosVDPlaybackViewportClient::ProcessClick(FSceneView& View, HHitProxy* HitProxy, FKey Key, EInputEvent Event, uint32 HitX, uint32 HitY)
 {
+	FEditorViewportClient::ProcessClick(View, HitProxy, Key, Event, HitX, HitY);
+	
 	if (HitProxy == nullptr)
 	{
 		return;
@@ -66,6 +83,20 @@ void FChaosVDPlaybackViewportClient::HandleObjectFocused(UObject* FocusedObject)
 	}
 }
 
+void FChaosVDPlaybackViewportClient::HandleActorMoving(AActor* MovedActor) const
+{
+	if (Cast<ADirectionalLight>(MovedActor))
+	{
+		if (const TSharedPtr<FChaosVDScene> SceneSharedPtr = CVDScene.Pin())
+		{
+			if (SceneSharedPtr->GetSkySphereActor()->Implements<UChaosVDSkySphereInterface>())
+			{
+				IChaosVDSkySphereInterface::Execute_Refresh(SceneSharedPtr->GetSkySphereActor());
+			}
+		}
+	}
+}
+
 void FChaosVDPlaybackViewportClient::TrackActor(AActor* ActorToTrack, EChaosVDActorTrackingMode TrackingMode)
 {
 	if (!ActorToTrack)
@@ -153,6 +184,63 @@ void FChaosVDPlaybackViewportClient::TrackTransform(const FTransform& TransformT
 	}
 }
 
+void FChaosVDPlaybackViewportClient::PerformSelectedTrackingForFrame(FChaosVDGameFrameData* FrameData)
+{
+	if (const TSharedPtr<FChaosVDScene> CVDSceneSharedPtr = CVDScene.Pin())
+	{
+		if (const UChaosVDEditorSettings* CVDEditorSettings = GetDefault<UChaosVDEditorSettings>())
+		{
+			switch (CVDEditorSettings->TrackingTarget)
+			{
+			case EChaosVDActorTrackingTarget::SelectedObject:
+				{
+					if (ModeTools.IsValid())
+					{
+						USelection* CurrentSelection = ModeTools->GetSelectedActors();
+		
+						//TODO: Update this if we add multi selection support
+						if (AActor* SelectedActor = CurrentSelection ? CurrentSelection->GetTop<AActor>() : nullptr)
+						{
+							TrackActor(SelectedActor, CVDEditorSettings->TrackingOptions);
+						}
+					}
+					break;
+				}
+			case EChaosVDActorTrackingTarget::RecordedTransform:
+				{
+					// TODO: Find a better place to store the current selected Transform name, it should not be the editor settings object
+					if (const TSharedPtr<FName>& TransformName = CVDEditorSettings->SelectedTrackedTransformName)
+					{
+						if (const FChaosVDTrackedTransform* TrackedTransform = FrameData->RecordedNonSolverTransformsByID.Find(*TransformName))
+						{
+							TrackTransform(TrackedTransform->Transform, CVDEditorSettings->TrackingOptions);
+						}
+					}
+								
+					break;
+				}
+			case EChaosVDActorTrackingTarget::RecordedLocation:
+				{
+					// TODO: Find a better place to store the current selected Location name, it should not be the editor settings object
+					if (const TSharedPtr<FName>& LocationName = CVDEditorSettings->SelectedTrackedLocationName)
+					{
+						if (const FChaosVDTrackedLocation* TrackedTransform = FrameData->RecordedNonSolverLocationsByID.Find(*LocationName))
+						{
+							FTransform LocationTransform;
+							LocationTransform.SetLocation(TrackedTransform->Location);
+							TrackTransform(LocationTransform, CVDEditorSettings->TrackingOptions);
+						}
+					}
+								
+					break;
+				}
+			default:
+				break;
+			}			
+		}
+	}
+}
+
 void FChaosVDPlaybackViewportClient::Draw(const FSceneView* View, FPrimitiveDrawInterface* PDI)
 {
 	FEditorViewportClient::Draw(View, PDI);
diff --git a/Engine/Plugins/Experimental/ChaosVD/Source/ChaosVD/Private/ChaosVDScene.cpp b/Engine/Plugins/Experimental/ChaosVD/Source/ChaosVD/Private/ChaosVDScene.cpp
--- a/Engine/Plugins/Experimental/ChaosVD/Source/ChaosVD/Private/ChaosVDScene.cpp
+++ b/Engine/Plugins/Experimental/ChaosVD/Source/ChaosVD/Private/ChaosVDScene.cpp
@@ -7,21 +7,21 @@
 #include "ChaosVDParticleActor.h"
 #include "Chaos/ImplicitObject.h"
 #include "ChaosVDRecording.h"
-#include "EditorActorFolders.h"
-#include "WorldPersistentFolders.h"
-#include "Elements/Framework/EngineElementsLibrary.h"
-#include "Elements/Framework/TypedElementSelectionSet.h"
-#include "EditorActorFolders.h"
-#include "WorldPersistentFolders.h"
+#include "ChaosVDSkySphereInterface.h"
 #include "DataWrappers/ChaosVDParticleDataWrapper.h"
+#include "EditorActorFolders.h"
+#include "EditorLevelUtils.h"
+#include "Engine/DirectionalLight.h"
 #include "Engine/Engine.h"
-#include "Engine/LevelStreamingDynamic.h"
 #include "Engine/StreamableManager.h"
 #include "Engine/World.h"
+#include "Elements/Framework/EngineElementsLibrary.h"
+#include "Elements/Framework/TypedElementSelectionSet.h"
 #include "Misc/ScopedSlowTask.h"
 #include "Materials/Material.h"
-
+#include "Selection.h"
 #include "UObject/Package.h"
+#include "WorldPersistentFolders.h"
 
 #define LOCTEXT_NAMESPACE "ChaosVisualDebugger"
 
@@ -43,10 +43,7 @@ void FChaosVDScene::Initialize()
 		return;
 	}
 
-	SelectionSet = NewObject<UTypedElementSelectionSet>();
-
-	SelectionSet->OnPreChange().AddRaw(this, &FChaosVDScene::HandlePreSelectionChange);
-	SelectionSet->OnChanged().AddRaw(this, &FChaosVDScene::HandlePostSelectionChange);
+	InitializeSelectionSets();
 	
 	PhysicsVDWorld = CreatePhysicsVDWorld();
 
@@ -78,8 +75,7 @@ void FChaosVDScene::DeInitialize()
 		Settings->OnVisibilitySettingsChanged().RemoveAll(this);
 	}
 
-	SelectionSet->OnPreChange().RemoveAll(this);
-	SelectionSet->OnChanged().RemoveAll(this);
+	DeInitializeSelectionSets();
 
 	GeometryGenerator.Reset();
 
@@ -103,6 +99,9 @@ void FChaosVDScene::AddReferencedObjects(FReferenceCollector& Collector)
 {
 	Collector.AddReferencedObject(PhysicsVDWorld);
 	Collector.AddReferencedObject(SelectionSet);
+	Collector.AddReferencedObject(ObjectSelection);
+	Collector.AddReferencedObject(ActorSelection);
+	Collector.AddReferencedObject(ComponentSelection);
 }
 
 void FChaosVDScene::UpdateFromRecordedStepData(const int32 SolverID, const FString& SolverName, const FChaosVDStepData& InRecordedStepData, const FChaosVDSolverFrameData& InFrameData)
@@ -367,6 +366,42 @@ int32 FChaosVDScene::GetIDForRecordedParticleData(const FChaosVDParticleDataWrap
 	return InParticleData.ParticleIndex;
 }
 
+void FChaosVDScene::CreateBaseLights(UWorld* TargetWorld) const
+{
+	if (!TargetWorld)
+	{
+		return;
+	}
+
+	const FName LightingFolderPath("ChaosVisualDebugger/Lighting");
+
+	const FVector SpawnPosition(0.0, 0.0, 2000.0);
+	
+	if (const UChaosVDEditorSettings* Settings = GetDefault<UChaosVDEditorSettings>())
+	{
+		if (ADirectionalLight* DirectionalLightActor = TargetWorld->SpawnActor<ADirectionalLight>())
+		{
+			DirectionalLightActor->SetCastShadows(false);
+			DirectionalLightActor->SetMobility(EComponentMobility::Movable);
+			DirectionalLightActor->SetActorLocation(SpawnPosition);
+
+			DirectionalLightActor->SetFolderPath(LightingFolderPath);
+
+			TSubclassOf<AActor> SkySphereClass = Settings->SkySphereActorClass.TryLoadClass<AActor>();
+			SkySphere = TargetWorld->SpawnActor(SkySphereClass.Get());
+			if (SkySphere)
+			{
+				SkySphere->SetActorLocation(SpawnPosition);
+				SkySphere->SetFolderPath(LightingFolderPath);
+				if (SkySphere->Implements<UChaosVDSkySphereInterface>())
+				{
+					IChaosVDSkySphereInterface::Execute_SetDirectionalLightSource(SkySphere, DirectionalLightActor);
+				}
+			}
+		}
+	}
+}
+
 UWorld* FChaosVDScene::CreatePhysicsVDWorld() const
 {
 	const FName UniqueWorldName = FName(FGuid::NewGuid().ToString());
@@ -387,26 +422,7 @@ UWorld* FChaosVDScene::CreatePhysicsVDWorld() const
 										  .SetTransactional( false )
 	);
 
-	// Add the base content as a sublevel
-	const UChaosVDEditorSettings* Settings = GetDefault<UChaosVDEditorSettings>();
-
-	ULevelStreamingDynamic* StreamedInLevel = NewObject<ULevelStreamingDynamic>(NewWorld);
-	StreamedInLevel->SetWorldAssetByPackageName(FName(Settings->BasePhysicsVDWorld.GetLongPackageName()));
-
-	StreamedInLevel->PackageNameToLoad = FName(Settings->BasePhysicsVDWorld.GetLongPackageName());
-
-	StreamedInLevel->SetShouldBeLoaded(true);
-	StreamedInLevel->bShouldBlockOnLoad = true;
-	StreamedInLevel->bInitiallyLoaded = true;
-
-	StreamedInLevel->SetShouldBeVisible(true);
-	StreamedInLevel->bInitiallyVisible = true;
-	StreamedInLevel->bLocked = true;
-
-	NewWorld->AddStreamingLevel(StreamedInLevel);
-
-	NewWorld->FlushLevelStreaming(EFlushLevelStreamingType::Full);
-
+	CreateBaseLights(NewWorld);
 	return NewWorld;
 }
 
@@ -477,6 +493,34 @@ void FChaosVDScene::HandleVisibilitySettingsChanged(UChaosVDEditorSettings* Sett
 	}
 }
 
+void FChaosVDScene::InitializeSelectionSets()
+{
+	SelectionSet = NewObject<UTypedElementSelectionSet>(GetTransientPackage(), NAME_None, RF_Transactional);
+	SelectionSet->AddToRoot();
+
+	ActorSelection = USelection::CreateActorSelection(GetTransientPackage(), TEXT("CVDSelectedActors"), RF_Transactional);
+	ActorSelection->SetElementSelectionSet(SelectionSet);
+
+	ComponentSelection = USelection::CreateComponentSelection(GetTransientPackage(), TEXT("CVDSelectedComponents"), RF_Transactional);
+	ComponentSelection->SetElementSelectionSet(SelectionSet);
+
+	ObjectSelection = USelection::CreateObjectSelection(GetTransientPackage(), TEXT("CVDSelectedObjects"), RF_Transactional);
+	ObjectSelection->SetElementSelectionSet(SelectionSet);
+
+	SelectionSet->OnPreChange().AddRaw(this, &FChaosVDScene::HandlePreSelectionChange);
+	SelectionSet->OnChanged().AddRaw(this, &FChaosVDScene::HandlePostSelectionChange);
+}
+
+void FChaosVDScene::DeInitializeSelectionSets()
+{
+	ActorSelection->SetElementSelectionSet(nullptr);
+	ComponentSelection->SetElementSelectionSet(nullptr);
+	ObjectSelection->SetElementSelectionSet(nullptr);
+
+	SelectionSet->OnPreChange().RemoveAll(this);
+	SelectionSet->OnChanged().RemoveAll(this);
+}
+
 void FChaosVDScene::SetSelectedObject(UObject* SelectedObject)
 {
 	if (!SelectionSet)
diff --git a/Engine/Plugins/Experimental/ChaosVD/Source/ChaosVD/Private/Widgets/SChaosVDEditorViewportViewMenu.cpp b/Engine/Plugins/Experimental/ChaosVD/Source/ChaosVD/Private/Widgets/SChaosVDEditorViewportViewMenu.cpp
new file mode 100644
--- /dev/null
+++ b/Engine/Plugins/Experimental/ChaosVD/Source/ChaosVD/Private/Widgets/SChaosVDEditorViewportViewMenu.cpp
@@ -0,0 +1,41 @@
+// Copyright Epic Games, Inc. All Rights Reserved.
+
+#include "Widgets/SChaosVDEditorViewportViewMenu.h"
+
+#include "EditorViewportCommands.h"
+#include "Engine/EngineBaseTypes.h"
+#include "SViewportToolBar.h"
+#include "ToolMenu.h"
+#include "ToolMenuDelegates.h"
+#include "ToolMenus.h"
+#include "ToolMenuSection.h"
+#include "UObject/NameTypes.h"
+
+void SChaosVDEditorViewportViewMenu::Construct(const FArguments& InArgs, TSharedRef<SEditorViewport> InViewport, TSharedRef<SViewportToolBar> InParentToolBar)
+{
+	SEditorViewportViewMenu::Construct(SEditorViewportViewMenu::FArguments(), InViewport, InParentToolBar);
+	MenuName = FName("ChaosVD.ViewportToolbar.View");
+}
+
+void SChaosVDEditorViewportViewMenu::RegisterMenus() const
+{
+	if (!UToolMenus::Get()->IsMenuRegistered(MenuName))
+	{
+		if (UToolMenu* Menu = UToolMenus::Get()->RegisterMenu(MenuName))
+		{
+			Menu->AddDynamicSection("BaseSection", FNewToolMenuDelegate::CreateLambda([](UToolMenu* InMenu)
+			{
+				if (InMenu)
+				{
+					const FEditorViewportCommands& BaseViewportActions = FEditorViewportCommands::Get();
+					FToolMenuSection& Section = InMenu->AddSection("ViewMode", NSLOCTEXT("ChaosVisualDebugger","ViewModeHeader", "View Mode"));
+					{
+						Section.AddMenuEntry(BaseViewportActions.LitMode, UViewModeUtils::GetViewModeDisplayName(VMI_Lit));
+						Section.AddMenuEntry(BaseViewportActions.UnlitMode, UViewModeUtils::GetViewModeDisplayName(VMI_Unlit));
+						Section.AddMenuEntry(BaseViewportActions.WireframeMode, UViewModeUtils::GetViewModeDisplayName(VMI_BrushWireframe));
+					}
+				}
+			}));
+		}
+	}
+}
diff --git a/Engine/Plugins/Experimental/ChaosVD/Source/ChaosVD/Private/Widgets/SChaosVDPlaybackViewport.cpp b/Engine/Plugins/Experimental/ChaosVD/Source/ChaosVD/Private/Widgets/SChaosVDPlaybackViewport.cpp
--- a/Engine/Plugins/Experimental/ChaosVD/Source/ChaosVD/Private/Widgets/SChaosVDPlaybackViewport.cpp
+++ b/Engine/Plugins/Experimental/ChaosVD/Source/ChaosVD/Private/Widgets/SChaosVDPlaybackViewport.cpp
@@ -2,16 +2,17 @@
 
 #include "Widgets/SChaosVDPlaybackViewport.h"
 
+#include "ChaosVDEditorMode.h"
+#include "ChaosVDEditorModeTools.h"
 #include "ChaosVDEditorSettings.h"
 #include "ChaosVDPlaybackController.h"
 #include "ChaosVDPlaybackViewportClient.h"
 #include "ChaosVDScene.h"
+#include "EditorModeManager.h"
 #include "Framework/Application/SlateApplication.h"
-#include "LevelEditorViewport.h"
 #include "Elements/Framework/TypedElementSelectionSet.h"
-#include "Slate/SceneViewport.h"
 #include "Widgets/SChaosVDTimelineWidget.h"
-#include "Widgets/SViewport.h"
+#include "Widgets/SChaosVDViewportToolbar.h"
 #include "Widgets/Text/STextBlock.h"
 
 #define LOCTEXT_NAMESPACE "ChaosVisualDebugger"
@@ -22,47 +23,36 @@ SChaosVDPlaybackViewport::~SChaosVDPlaybackViewport()
 	PlaybackViewportClient.Reset();
 }
 
-TSharedPtr<FChaosVDPlaybackViewportClient> SChaosVDPlaybackViewport::CreateViewportClient() const
+void SChaosVDPlaybackViewport::Construct(const FArguments& InArgs, TWeakPtr<FChaosVDScene> InScene, TWeakPtr<FChaosVDPlaybackController> InPlaybackController)
 {
-	TSharedPtr<FChaosVDPlaybackViewportClient> NewViewport = MakeShared<FChaosVDPlaybackViewportClient>();
+	Extender = MakeShared<FExtender>();
 
-	NewViewport->SetAllowCinematicControl(false);
-	
-	NewViewport->bSetListenerPosition = false;
-	NewViewport->EngineShowFlags = FEngineShowFlags(ESFIM_Editor);
-	NewViewport->LastEngineShowFlags = FEngineShowFlags(ESFIM_Editor);
-	NewViewport->ViewportType = LVT_Perspective;
-	NewViewport->bDrawAxes = true;
-	NewViewport->bDisableInput = false;
-	NewViewport->VisibilityDelegate.BindLambda([] {return true; });
-	NewViewport->EngineShowFlags.SetSelectionOutline(true);
+	EditorModeTools = MakeShared<FChaosVDEditorModeTools>(InScene);
+	EditorModeTools->SetWidgetMode(UE::Widget::WM_Translate);
+	EditorModeTools->SetDefaultMode(UChaosVDEditorMode::EM_ChaosVisualDebugger);
+	EditorModeTools->ActivateDefaultMode();
 
-	return NewViewport;
-}
+	SEditorViewport::Construct(SEditorViewport::FArguments());
 
-void SChaosVDPlaybackViewport::Construct(const FArguments& InArgs, TWeakPtr<FChaosVDScene> InScene, TWeakPtr<FChaosVDPlaybackController> InPlaybackController)
-{
 	CVDSceneWeakPtr = InScene;
 	TSharedPtr<FChaosVDScene> ScenePtr = InScene.Pin();
 	ensure(ScenePtr.IsValid());
 	ensure(InPlaybackController.IsValid());
 
-	PlaybackViewportClient = CreateViewportClient();
-
-	ViewportWidget = SNew(SViewport)
-		.RenderDirectlyToWindow(false)
-		.IsEnabled(FSlateApplication::Get().GetNormalExecutionAttribute())
-		.EnableGammaCorrection(false)
-		.EnableBlending(false);
+	PlaybackViewportClient = StaticCastSharedPtr<FChaosVDPlaybackViewportClient>(GetViewportClient());
 
-	SceneViewport = MakeShareable(new FSceneViewport(PlaybackViewportClient.Get(), ViewportWidget));
-
-	PlaybackViewportClient->Viewport = SceneViewport.Get();
-
-	ViewportWidget->SetViewportInterface(SceneViewport.ToSharedRef());
+	// TODO: Add a way to gracefully shutdown (close) the tool when a no recoverable situation like this happens (UE-191876)
+	check(PlaybackViewportClient.IsValid());
 	
-	// Default to the base map
 	PlaybackViewportClient->SetScene(InScene);
+	
+	if (UChaosVDEditorMode* CVDEdMode = Cast<UChaosVDEditorMode>(EditorModeTools->GetActiveScriptableMode(UChaosVDEditorMode::EM_ChaosVisualDebugger)))
+	{
+		if (ScenePtr.IsValid())
+		{
+			CVDEdMode->SetWorld(ScenePtr->GetUnderlyingWorld());
+		}
+	}
 
 	ChildSlot
 	[
@@ -103,21 +93,53 @@ void SChaosVDPlaybackViewport::Construct(const FArguments& InArgs, TWeakPtr<FCha
 	RegisterNewController(InPlaybackController);
 }
 
+TSharedRef<SEditorViewport> SChaosVDPlaybackViewport::GetViewportWidget()
+{
+	return StaticCastSharedRef<SEditorViewport>(AsShared());
+}
 
-void SChaosVDPlaybackViewport::TrackActor(AActor* ActorToTrack, EChaosVDActorTrackingMode TrackingMode)
+TSharedPtr<FExtender> SChaosVDPlaybackViewport::GetExtenders() const
 {
-	if (ensure(PlaybackViewportClient.IsValid()))
-	{
-		PlaybackViewportClient->TrackActor(ActorToTrack, TrackingMode);
-	}	
+	return Extender;
 }
 
-void SChaosVDPlaybackViewport::TrackTransform(const FTransform& TransformToTrack, EChaosVDActorTrackingMode TrackingMode)
+TSharedRef<FEditorViewportClient> SChaosVDPlaybackViewport::MakeEditorViewportClient()
 {
-	if (ensure(PlaybackViewportClient.IsValid()))
-	{
-		PlaybackViewportClient->TrackTransform(TransformToTrack, TrackingMode);
-	}
+	TSharedPtr<FChaosVDPlaybackViewportClient> NewViewport = MakeShared<FChaosVDPlaybackViewportClient>(EditorModeTools);
+
+	NewViewport->SetAllowCinematicControl(false);
+	
+	NewViewport->bSetListenerPosition = false;
+	NewViewport->EngineShowFlags = FEngineShowFlags(ESFIM_Editor);
+	NewViewport->LastEngineShowFlags = FEngineShowFlags(ESFIM_Editor);
+	NewViewport->ViewportType = LVT_Perspective;
+	NewViewport->bDrawAxes = true;
+	NewViewport->bDisableInput = false;
+	NewViewport->VisibilityDelegate.BindLambda([] {return true; });
+
+	NewViewport->EngineShowFlags.DisableAdvancedFeatures();
+	NewViewport->EngineShowFlags.SetSelectionOutline(true);
+	NewViewport->EngineShowFlags.SetSnap(false);
+	NewViewport->EngineShowFlags.SetBillboardSprites(true);
+
+	return StaticCastSharedRef<FEditorViewportClient>(NewViewport.ToSharedRef());
+}
+
+TSharedPtr<SWidget> SChaosVDPlaybackViewport::MakeViewportToolbar()
+{
+	// Build our toolbar level toolbar
+	TSharedRef< SChaosVDViewportToolbar > ToolBar = SNew(SChaosVDViewportToolbar, SharedThis(this));
+
+	return 
+		SNew(SVerticalBox)
+		.Visibility( EVisibility::SelfHitTestInvisible )
+		+SVerticalBox::Slot()
+		.AutoHeight()
+		.Padding(0, 1.0f, 0, 0)
+		.VAlign(VAlign_Top)
+		[
+			ToolBar
+		];
 }
 
 void SChaosVDPlaybackViewport::HandlePlaybackControllerDataUpdated(TWeakPtr<FChaosVDPlaybackController> InController)
@@ -148,17 +170,24 @@ void SChaosVDPlaybackViewport::HandlePlaybackControllerDataUpdated(TWeakPtr<FCha
 
 void SChaosVDPlaybackViewport::HandleControllerTrackFrameUpdated(TWeakPtr<FChaosVDPlaybackController> InController, const FChaosVDTrackInfo* UpdatedTrackInfo, FGuid InstigatorGuid)
 {
-	if (InstigatorGuid == GetInstigatorID())
-	{
-		// Ignore the update if we initiated it
-		return;
-	}
-
 	if (TSharedPtr<FChaosVDPlaybackController> ControllerSharedPtr = InController.Pin())
 	{
 		// The frame number we receive could be from a Solver track, so make sure it is converted to the correct game track frame number 
 		int32 GameTrackFrame = ControllerSharedPtr->ConvertCurrentFrameToOtherTrackFrame(UpdatedTrackInfo, ControllerSharedPtr->GetTrackInfo(EChaosVDTrackType::Game, FChaosVDPlaybackController::GameTrackID));
-		GameFramesTimelineWidget->SetCurrentTimelineFrame(GameTrackFrame, EChaosVDSetTimelineFrameFlags::None);
+
+		// Something other than us advanced the game track, so make sure the timeline widget is updated
+		if (InstigatorGuid != GetInstigatorID())
+		{
+			GameFramesTimelineWidget->SetCurrentTimelineFrame(GameTrackFrame, EChaosVDSetTimelineFrameFlags::None);
+		}
+
+		if (TSharedPtr<FChaosVDRecording> RecordingData = ControllerSharedPtr->GetCurrentRecording().Pin())
+		{
+			if (FChaosVDGameFrameData* FrameData = RecordingData->GetGameFrameData_AssumesLocked(GameTrackFrame))
+			{
+				PlaybackViewportClient->PerformSelectedTrackingForFrame(FrameData);
+			}	
+		}
 	}
 }
 
diff --git a/Engine/Plugins/Experimental/ChaosVD/Source/ChaosVD/Private/Widgets/SChaosVDViewportToolbar.cpp b/Engine/Plugins/Experimental/ChaosVD/Source/ChaosVD/Private/Widgets/SChaosVDViewportToolbar.cpp
new file mode 100644
--- /dev/null
+++ b/Engine/Plugins/Experimental/ChaosVD/Source/ChaosVD/Private/Widgets/SChaosVDViewportToolbar.cpp
@@ -0,0 +1,65 @@
+// Copyright Epic Games, Inc. All Rights Reserved.
+
+#include "Widgets/SChaosVDViewportToolbar.h"
+
+#include "ChaosVDEditorSettings.h"
+#include "Framework/MultiBox/MultiBoxBuilder.h"
+#include "PropertyEditorModule.h"
+#include "DetailsViewArgs.h"
+#include "IDetailsView.h"
+#include "Modules/ModuleManager.h"
+#include "PropertyEditorModule.h"
+#include "SCommonEditorViewportToolbarBase.h"
+#include "SEditorViewport.h"
+#include "Widgets/SChaosVDEditorViewportViewMenu.h"
+#include "Widgets/SChaosVDPlaybackViewport.h"
+#include "Widgets/SChaosVDVisualizationControls.h"
+
+#define LOCTEXT_NAMESPACE "ChaosVisualDebugger"
+
+void SChaosVDViewportToolbar::Construct(const FArguments& InArgs, TSharedPtr<ICommonEditorViewportToolbarInfoProvider> InInfoProvider)
+{
+	SCommonEditorViewportToolbarBase::Construct(SCommonEditorViewportToolbarBase::FArguments(), InInfoProvider);
+}
+
+TSharedRef<SEditorViewportViewMenu> SChaosVDViewportToolbar::MakeViewMenu()
+{
+	TSharedRef<SEditorViewport> ViewportRef = GetInfoProvider().GetViewportWidget();
+	return SNew(SChaosVDEditorViewportViewMenu, ViewportRef, SharedThis(this));
+}
+
+void SChaosVDViewportToolbar::ExtendOptionsMenu(FMenuBuilder& OptionsMenuBuilder) const
+{
+	constexpr bool bInShouldCloseWindowAfterMenuSelection = true;
+	FMenuBuilder CVDOptionsMenuBuilder(bInShouldCloseWindowAfterMenuSelection, GetInfoProvider().GetViewportWidget()->GetCommandList());
+
+	TSharedRef<SVerticalBox> OptionsMenuWidget = SNew(SVerticalBox)
+													+SVerticalBox::Slot()
+													.Padding(5.0f)
+													[
+														SNew(STextBlock)
+														.Text(LOCTEXT("Default Options Message", "Nothing so see yet. An options menu will be implemented soon"))
+													];
+
+	// Intentionally leaving the label empty as we only want to show the temp message defined above
+	CVDOptionsMenuBuilder.AddWidget(OptionsMenuWidget, FText());
+
+	OptionsMenuBuilder = CVDOptionsMenuBuilder;
+}
+
+TSharedRef<SWidget> SChaosVDViewportToolbar::GenerateShowMenu() const
+{
+	FPropertyEditorModule& PropertyEditorModule = FModuleManager::GetModuleChecked<FPropertyEditorModule>("PropertyEditor");
+	FDetailsViewArgs DetailsViewArgs;
+	DetailsViewArgs.bAllowSearch = false;
+	DetailsViewArgs.NameAreaSettings = FDetailsViewArgs::HideNameArea;
+
+	TSharedRef<IDetailsView> DetailsPanel = PropertyEditorModule.CreateDetailView(DetailsViewArgs);
+
+	DetailsPanel->SetObject(GetMutableDefault<UChaosVDEditorSettings>());
+
+	TSharedRef<SChaosVDPlaybackViewport> ViewportRef = StaticCastSharedRef<SChaosVDPlaybackViewport>(GetInfoProvider().GetViewportWidget());
+
+	return SNew(SChaosVDVisualizationControls, ViewportRef->GetObservedController(), ViewportRef);
+}
+#undef LOCTEXT_NAMESPACE
diff --git a/Engine/Plugins/Experimental/ChaosVD/Source/ChaosVD/Private/Widgets/SChaosVDVisualizationControls.cpp b/Engine/Plugins/Experimental/ChaosVD/Source/ChaosVD/Private/Widgets/SChaosVDVisualizationControls.cpp
--- a/Engine/Plugins/Experimental/ChaosVD/Source/ChaosVD/Private/Widgets/SChaosVDVisualizationControls.cpp
+++ b/Engine/Plugins/Experimental/ChaosVD/Source/ChaosVD/Private/Widgets/SChaosVDVisualizationControls.cpp
@@ -139,69 +139,24 @@ void SChaosVDVisualizationControls::HandleControllerTrackFrameUpdated(TWeakPtr<F
 		{
 			UpdateNameListFromKeys(FrameData->RecordedNonSolverTransformsByID, *TransformNamePickerWidget.Get());
 			UpdateNameListFromKeys(FrameData->RecordedNonSolverLocationsByID, *LocationsNamePickerWidget.Get());
-
-			if (const TSharedPtr<FChaosVDScene> CVDSceneSharedPtr = ControllerSharedPtr->GetControllerScene().Pin())
-			{
-				if (const UChaosVDEditorSettings* CVDEditorSettings = GetDefault<UChaosVDEditorSettings>())
-				{
-					switch (CVDEditorSettings->TrackingTarget)
-					{
-						case EChaosVDActorTrackingTarget::SelectedObject:
-							{
-								TArray<AActor*> SelectedActors = CVDSceneSharedPtr->GetElementSelectionSet()->GetSelectedObjects<AActor>();
-								if (SelectedActors.Num() > 0)
-								{
-									//TODO: Update this if we add multi selection support
-									if (AActor* SelectedActor = SelectedActors[0])
-									{
-										PlaybackViewportSharedPtr->TrackActor(SelectedActor, CVDEditorSettings->TrackingOptions);
-									}
-								}
-								break;
-							}
-						case EChaosVDActorTrackingTarget::RecordedTransform:
-							{
-								if (SelectedTrackedTransformName.IsValid())
-								{
-									if (const FChaosVDTrackedTransform* TrackedTransform = FrameData->RecordedNonSolverTransformsByID.Find(*SelectedTrackedTransformName))
-									{
-										PlaybackViewportSharedPtr->TrackTransform(TrackedTransform->Transform, CVDEditorSettings->TrackingOptions);
-									}
-								}
-								
-								break;
-							}
-						case EChaosVDActorTrackingTarget::RecordedLocation:
-							{
-								if (SelectedTrackedLocationName.IsValid())
-								{
-									if (const FChaosVDTrackedLocation* TrackedTransform = FrameData->RecordedNonSolverLocationsByID.Find(*SelectedTrackedLocationName))
-									{
-										FTransform LocationTransform;
-										LocationTransform.SetLocation(TrackedTransform->Location);
-										PlaybackViewportSharedPtr->TrackTransform(LocationTransform, CVDEditorSettings->TrackingOptions);
-									}
-								}
-								
-								break;
-							}
-						default:
-							break;
-						}			
-				}
-			}
 		}
 	}
 }
 
 void SChaosVDVisualizationControls::HandleLocationNameSelected(TSharedPtr<FName> SelectedName)
 {
-	SelectedTrackedLocationName = SelectedName;
+	if (UChaosVDEditorSettings* Settings = GetMutableDefault<UChaosVDEditorSettings>())
+	{
+		Settings->SelectedTrackedLocationName = SelectedName;
+	}
 }
 
 void SChaosVDVisualizationControls::HandleTransformNameSelected(TSharedPtr<FName> SelectedName)
 {
-	SelectedTrackedTransformName = SelectedName;
+	if (UChaosVDEditorSettings* Settings = GetMutableDefault<UChaosVDEditorSettings>())
+	{
+		Settings->SelectedTrackedTransformName = SelectedName;
+	}
 }
 
 EVisibility SChaosVDVisualizationControls::GetLocationsPickerVisibility() const
diff --git a/Engine/Plugins/Experimental/ChaosVD/Source/ChaosVD/Public/ChaosVDEditorMode.h b/Engine/Plugins/Experimental/ChaosVD/Source/ChaosVD/Public/ChaosVDEditorMode.h
new file mode 100644
--- /dev/null
+++ b/Engine/Plugins/Experimental/ChaosVD/Source/ChaosVD/Public/ChaosVDEditorMode.h
@@ -0,0 +1,55 @@
+// Copyright Epic Games, Inc. All Rights Reserved.
+
+#pragma once
+
+#include "Tools/LegacyEdModeWidgetHelpers.h"
+#include "Tools/LegacyEdModeInterfaces.h"
+
+#include "ChaosVDEditorMode.generated.h"
+
+class FChaosVDEditorModeTools;
+
+/**
+ * LegacyEdModeWidgetHelper implementation for the Chaos Visual Debugger tool
+ */
+class FChaosVDEdModeWidgetHelper final : public FLegacyEdModeWidgetHelper
+{
+public:
+	virtual FVector GetWidgetLocation() const override;
+	virtual bool InputDelta(FEditorViewportClient* InViewportClient, FViewport* InViewport, FVector& InDrag, FRotator& InRot, FVector& InScale) override;
+	virtual bool UsesTransformWidget(UE::Widget::EWidgetMode CheckMode) const override;
+	virtual bool UsesTransformWidget() const override;
+	
+private:
+	AActor* GetCurrentSelectedActor() const;
+};
+
+/**
+ * Main Editor Mode for the Chaos Visual Debugger tool.
+ */
+UCLASS()
+class UChaosVDEditorMode : public UBaseLegacyWidgetEdMode
+{
+	GENERATED_BODY()
+
+public:
+	UChaosVDEditorMode();
+
+	virtual bool IsCompatibleWith(FEditorModeID OtherModeID) const override { return false; };
+
+	static const FEditorModeID EM_ChaosVisualDebugger;
+
+	virtual bool UsesToolkits() const override { return false; };
+
+	virtual bool CanAutoSave() const override;
+
+	virtual UWorld* GetWorld() const override;
+
+	void SetWorld(UWorld* InWorld) { CVDWorldPtr = InWorld; }
+
+protected:
+
+	virtual TSharedRef<FLegacyEdModeWidgetHelper> CreateWidgetHelper() override;
+
+	TWeakObjectPtr<UWorld> CVDWorldPtr;
+};
diff --git a/Engine/Plugins/Experimental/ChaosVD/Source/ChaosVD/Public/ChaosVDEditorModeTools.h b/Engine/Plugins/Experimental/ChaosVD/Source/ChaosVD/Public/ChaosVDEditorModeTools.h
new file mode 100644
--- /dev/null
+++ b/Engine/Plugins/Experimental/ChaosVD/Source/ChaosVD/Public/ChaosVDEditorModeTools.h
@@ -0,0 +1,25 @@
+// Copyright Epic Games, Inc. All Rights Reserved.
+
+#pragma once
+#include "EditorModeManager.h"
+
+class FChaosVDScene;
+
+/** Editor Mode tools for the Chaos Visual Debugger tool */
+class FChaosVDEditorModeTools : public FEditorModeTools
+{
+public:
+	FChaosVDEditorModeTools(const TWeakPtr<FChaosVDScene>& InScenePtr);
+
+	virtual FString GetReferencerName() const override;
+	virtual USelection* GetSelectedActors() const override;
+	virtual USelection* GetSelectedObjects() const override;
+	virtual USelection* GetSelectedComponents() const override;
+	virtual UWorld* GetWorld() const override;
+
+	/** Returns the Chaos Visual Debugger Scene */
+	TWeakPtr<FChaosVDScene> GetScene() const { return ScenePtr; }
+
+protected:
+	TWeakPtr<FChaosVDScene> ScenePtr;
+};
diff --git a/Engine/Plugins/Experimental/ChaosVD/Source/ChaosVD/Public/ChaosVDEditorSettings.h b/Engine/Plugins/Experimental/ChaosVD/Source/ChaosVD/Public/ChaosVDEditorSettings.h
--- a/Engine/Plugins/Experimental/ChaosVD/Source/ChaosVD/Public/ChaosVDEditorSettings.h
+++ b/Engine/Plugins/Experimental/ChaosVD/Source/ChaosVD/Public/ChaosVDEditorSettings.h
@@ -69,18 +69,21 @@ public:
 	UPROPERTY(EditAnywhere, Category = "Viewport Tracking", meta=(EditCondition = "TrackingOptions == EChaosVDActorTrackingMode::ByBoundingBox && TrackingTarget != EChaosVDActorTrackingTarget::Disabled", EditConditionHides))
 	float ExpandViewTrackingBy = 60.0f;
 
-	UPROPERTY(Config)
-	TSoftObjectPtr<UWorld> BasePhysicsVDWorld;
-
 	UPROPERTY(EditAnywhere, Category = "Geometry Visibility", meta = (Bitmask, BitmaskEnum = "/Script/ChaosVD.EChaosVDGeometryVisibilityFlags"))
 	uint8 GeometryVisibilityFlags = static_cast<uint8>(EChaosVDGeometryVisibilityFlags::Simulated | EChaosVDGeometryVisibilityFlags::Simple);
 
 	UPROPERTY(Config)
 	TSoftObjectPtr<UMaterial> QueryOnlyMeshesMaterial;
 
+	UPROPERTY(Config)
+	FSoftClassPath SkySphereActorClass;
+
 	virtual void PostEditChangeProperty(FPropertyChangedEvent& PropertyChangedEvent) override;
 
-	FChaosVDVisibilitySettingsChaged& OnVisibilitySettingsChanged() { return VisibilitySettingsChangedDelegate; };
+	FChaosVDVisibilitySettingsChaged& OnVisibilitySettingsChanged() { return VisibilitySettingsChangedDelegate; }
+
+	TSharedPtr<FName> SelectedTrackedTransformName;
+	TSharedPtr<FName> SelectedTrackedLocationName;
 
 protected:
 	FChaosVDVisibilitySettingsChaged VisibilitySettingsChangedDelegate;
diff --git a/Engine/Plugins/Experimental/ChaosVD/Source/ChaosVD/Public/ChaosVDPlaybackControllerObserver.h b/Engine/Plugins/Experimental/ChaosVD/Source/ChaosVD/Public/ChaosVDPlaybackControllerObserver.h
--- a/Engine/Plugins/Experimental/ChaosVD/Source/ChaosVD/Public/ChaosVDPlaybackControllerObserver.h
+++ b/Engine/Plugins/Experimental/ChaosVD/Source/ChaosVD/Public/ChaosVDPlaybackControllerObserver.h
@@ -20,6 +20,9 @@ public:
 
 	virtual ~FChaosVDPlaybackControllerObserver() override;
 
+	/** Returns a Weak Ptr to the Playback Controller we are currently observing */
+	TWeakPtr<FChaosVDPlaybackController> GetObservedController();
+
 protected:
 	virtual void RegisterNewController(TWeakPtr<FChaosVDPlaybackController> NewController);
 	virtual void HandlePlaybackControllerDataUpdated(TWeakPtr<FChaosVDPlaybackController> InController) {};
diff --git a/Engine/Plugins/Experimental/ChaosVD/Source/ChaosVD/Public/ChaosVDPlaybackViewportClient.h b/Engine/Plugins/Experimental/ChaosVD/Source/ChaosVD/Public/ChaosVDPlaybackViewportClient.h
--- a/Engine/Plugins/Experimental/ChaosVD/Source/ChaosVD/Public/ChaosVDPlaybackViewportClient.h
+++ b/Engine/Plugins/Experimental/ChaosVD/Source/ChaosVD/Public/ChaosVDPlaybackViewportClient.h
@@ -3,6 +3,7 @@
 #pragma once
 #include "LevelEditorViewport.h"
 
+struct FChaosVDGameFrameData;
 class FChaosVDScene;
 enum class EChaosVDActorTrackingMode;
 
@@ -12,7 +13,8 @@ enum class EChaosVDActorTrackingMode;
 class FChaosVDPlaybackViewportClient : public FEditorViewportClient
 {
 public:
-	FChaosVDPlaybackViewportClient();
+
+	FChaosVDPlaybackViewportClient(const TSharedPtr<FEditorModeTools>& InModeTools);
 	virtual ~FChaosVDPlaybackViewportClient() override;
 
 	virtual void ProcessClick(FSceneView& View, HHitProxy* HitProxy, FKey Key, EInputEvent Event, uint32 HitX, uint32 HitY) override;
@@ -26,9 +28,12 @@ public:
 	
 	void TrackActor(AActor* ActorToTrack, EChaosVDActorTrackingMode TrackingMode);
 	void TrackTransform(const FTransform& TransformToTrack, EChaosVDActorTrackingMode TrackingMode);
+	
+	void PerformSelectedTrackingForFrame(FChaosVDGameFrameData* FrameData);
 
 private:
 	void HandleObjectFocused(UObject* FocusedObject);
+	void HandleActorMoving(AActor* MovedActor) const;
 	
 	FDelegateHandle ObjectFocusedDelegateHandle;
 	UWorld* CVDWorld;
diff --git a/Engine/Plugins/Experimental/ChaosVD/Source/ChaosVD/Public/ChaosVDScene.h b/Engine/Plugins/Experimental/ChaosVD/Source/ChaosVD/Public/ChaosVDScene.h
--- a/Engine/Plugins/Experimental/ChaosVD/Source/ChaosVD/Public/ChaosVDScene.h
+++ b/Engine/Plugins/Experimental/ChaosVD/Source/ChaosVD/Public/ChaosVDScene.h
@@ -15,6 +15,7 @@ class AChaosVDParticleActor;
 class FReferenceCollector;
 class UObject;
 class UMaterial;
+class USelection;
 class UTypedElementSelectionSet;
 class UWorld;
 
@@ -78,6 +79,10 @@ public:
 	/** Returns a ptr to the current selection set object */
 	UTypedElementSelectionSet* GetElementSelectionSet() const { return SelectionSet; }
 	
+	USelection* GetActorSelectionObject() const { return ActorSelection; }
+	USelection* GetComponentsSelectionObject() const { return ComponentSelection; }
+	USelection* GetObjectsSelectionObject() const { return ObjectSelection; }
+	
 	/** Event triggered when an object is focused in the scene (double click in the scene outliner)*/
 	FChaosVDOnObjectSelectedDelegate& OnObjectFocused() { return ObjectFocusedDelegate; }
 
@@ -87,6 +92,8 @@ public:
 	 */
 	AChaosVDParticleActor* GetParticleActor(int32 SolverID, int32 ParticleID);
 
+	AActor* GetSkySphereActor() const { return SkySphere; }
+
 	TSharedPtr<FChaosVDRecording> LoadedRecording;
 
 private:
@@ -96,6 +103,7 @@ private:
 
 	/** Returns the ID used to track this recorded particle data */
 	int32 GetIDForRecordedParticleData(const FChaosVDParticleDataWrapper& InParticleData) const;
+	void CreateBaseLights(UWorld* TargetWorld) const;
 
 	/** Creates the instance of the World which will be used the recorded data*/
 	UWorld* CreatePhysicsVDWorld() const;
@@ -116,6 +124,9 @@ private:
 
 	void HandleVisibilitySettingsChanged(UChaosVDEditorSettings* SettingsObject);
 
+	void InitializeSelectionSets();
+	void DeInitializeSelectionSets();
+
 	/** UWorld instance used to represent the recorded debug data */
 	TObjectPtr<UWorld> PhysicsVDWorld = nullptr;
 
@@ -130,11 +141,17 @@ private:
 	/** Selection set object holding the current selection state */
 	TObjectPtr<UTypedElementSelectionSet> SelectionSet;
 
+	TObjectPtr<USelection> ActorSelection = nullptr;
+	TObjectPtr<USelection> ComponentSelection = nullptr;
+	TObjectPtr<USelection> ObjectSelection = nullptr;
+
 	/** Array of actors with hit proxies that need to be updated */
 	TArray<AActor*> PendingActorsToUpdateSelectionProxy;
 
 	/** Scene Streamable manager that we'll use to async load any assets we depend on */
 	TSharedPtr<struct FStreamableManager> StreamableManager;
 
+	mutable AActor* SkySphere = nullptr;
+
 	bool bIsInitialized = false;
 };
diff --git a/Engine/Plugins/Experimental/ChaosVD/Source/ChaosVD/Public/ChaosVDSkySphereInterface.h b/Engine/Plugins/Experimental/ChaosVD/Source/ChaosVD/Public/ChaosVDSkySphereInterface.h
new file mode 100644
--- /dev/null
+++ b/Engine/Plugins/Experimental/ChaosVD/Source/ChaosVD/Public/ChaosVDSkySphereInterface.h
@@ -0,0 +1,29 @@
+// Copyright Epic Games, Inc. All Rights Reserved.
+
+#pragma once
+
+#include "UObject/Interface.h"
+#include "ChaosVDSkySphereInterface.generated.h"
+
+class ADirectionalLight;
+
+UINTERFACE()
+class UChaosVDSkySphereInterface : public UInterface
+{
+	GENERATED_BODY()
+};
+
+/**
+ * Simple interface to allow update the light source in the BP SkySphere from C++
+ */
+class CHAOSVD_API IChaosVDSkySphereInterface
+{
+	GENERATED_BODY()
+
+public:
+	UFUNCTION(BlueprintNativeEvent, CallInEditor)
+	void SetDirectionalLightSource(ADirectionalLight* NewLightSource);
+
+	UFUNCTION(BlueprintNativeEvent, CallInEditor)
+	void Refresh();
+};
diff --git a/Engine/Plugins/Experimental/ChaosVD/Source/ChaosVD/Public/Widgets/SChaosVDEditorViewportViewMenu.h b/Engine/Plugins/Experimental/ChaosVD/Source/ChaosVD/Public/Widgets/SChaosVDEditorViewportViewMenu.h
new file mode 100644
--- /dev/null
+++ b/Engine/Plugins/Experimental/ChaosVD/Source/ChaosVD/Public/Widgets/SChaosVDEditorViewportViewMenu.h
@@ -0,0 +1,24 @@
+// Copyright Epic Games, Inc. All Rights Reserved.
+
+#pragma once
+
+#include "SEditorViewportViewMenu.h"
+#include "Templates/SharedPointer.h"
+#include "Widgets/SCompoundWidget.h"
+
+/**
+ * Viewport View modes menu used by Chaos Visual Debugger.
+ * Compared to the normal Editor View modes menu, it has a reduced set of modes available,
+ * showing only the ones supported
+ */
+class CHAOSVD_API SChaosVDEditorViewportViewMenu : public SEditorViewportViewMenu
+{
+public:
+	SLATE_BEGIN_ARGS(SChaosVDEditorViewportViewMenu){}
+	SLATE_END_ARGS()
+
+	void Construct(const FArguments& InArgs, TSharedRef<SEditorViewport> InViewport, TSharedRef<class SViewportToolBar> InParentToolBar);
+
+protected:
+	virtual void RegisterMenus() const override;
+};
diff --git a/Engine/Plugins/Experimental/ChaosVD/Source/ChaosVD/Public/Widgets/SChaosVDPlaybackViewport.h b/Engine/Plugins/Experimental/ChaosVD/Source/ChaosVD/Public/Widgets/SChaosVDPlaybackViewport.h
--- a/Engine/Plugins/Experimental/ChaosVD/Source/ChaosVD/Public/Widgets/SChaosVDPlaybackViewport.h
+++ b/Engine/Plugins/Experimental/ChaosVD/Source/ChaosVD/Public/Widgets/SChaosVDPlaybackViewport.h
@@ -4,9 +4,12 @@
 
 #include "ChaosVDPlaybackControllerInstigator.h"
 #include "ChaosVDPlaybackControllerObserver.h"
+#include "SCommonEditorViewportToolbarBase.h"
+#include "SEditorViewport.h"
 #include "Templates/SharedPointer.h"
 #include "Widgets/SCompoundWidget.h"
 
+class FChaosVDEditorModeTools;
 enum class EChaosVDActorTrackingMode;
 class FChaosVDPlaybackViewportClient;
 class SChaosVDSolverPlaybackControls;
@@ -19,23 +22,28 @@ class FSceneViewport;
 class SViewport;
 
 /* Widget that contains the 3D viewport and playback controls */
-class SChaosVDPlaybackViewport : public SCompoundWidget, public FChaosVDPlaybackControllerObserver, public IChaosVDPlaybackControllerInstigator
+class SChaosVDPlaybackViewport : public SEditorViewport, public FChaosVDPlaybackControllerObserver, public IChaosVDPlaybackControllerInstigator, public ICommonEditorViewportToolbarInfoProvider
 {
 public:
 
-	SLATE_BEGIN_ARGS( SChaosVDPlaybackViewport ){}
+	SLATE_BEGIN_ARGS(SChaosVDPlaybackViewport) {}
 	SLATE_END_ARGS()
 
 	virtual ~SChaosVDPlaybackViewport() override;
 
 	void Construct(const FArguments& InArgs, TWeakPtr<FChaosVDScene> InScene, TWeakPtr<FChaosVDPlaybackController> InPlaybackController);
 
-	void TrackActor(AActor* ActorToTrack, EChaosVDActorTrackingMode TrackingMode);
-	void TrackTransform(const FTransform& TransformToTrack, EChaosVDActorTrackingMode TrackingMode);
+	// BEING ICommonEditorViewportToolbarInfoProvider interface
+	virtual TSharedRef<SEditorViewport> GetViewportWidget() override;
+	virtual TSharedPtr<FExtender> GetExtenders() const override;
+	virtual void OnFloatingButtonClicked() override {};
+	// END ICommonEditorViewportToolbarInfoProvider interface
 
 protected:
 
-	TSharedPtr<FChaosVDPlaybackViewportClient> CreateViewportClient() const;
+	virtual TSharedRef<FEditorViewportClient> MakeEditorViewportClient() override;
+	
+	virtual TSharedPtr<SWidget> MakeViewportToolbar() override;
 
 	virtual void RegisterNewController(TWeakPtr<FChaosVDPlaybackController> NewController )override;
 	virtual void HandlePlaybackControllerDataUpdated(TWeakPtr<FChaosVDPlaybackController> InController) override;
@@ -49,8 +57,10 @@ protected:
 	TSharedPtr<SChaosVDTimelineWidget> GameFramesTimelineWidget;
 
 	TSharedPtr<FChaosVDPlaybackViewportClient> PlaybackViewportClient;
-	TSharedPtr<SViewport> ViewportWidget;
-	TSharedPtr<FSceneViewport> SceneViewport;
 	
 	TWeakPtr<FChaosVDScene> CVDSceneWeakPtr;
+
+	TSharedPtr<FExtender> Extender;
+
+	TSharedPtr<FChaosVDEditorModeTools> EditorModeTools;
 };
diff --git a/Engine/Plugins/Experimental/ChaosVD/Source/ChaosVD/Public/Widgets/SChaosVDViewportToolbar.h b/Engine/Plugins/Experimental/ChaosVD/Source/ChaosVD/Public/Widgets/SChaosVDViewportToolbar.h
new file mode 100644
--- /dev/null
+++ b/Engine/Plugins/Experimental/ChaosVD/Source/ChaosVD/Public/Widgets/SChaosVDViewportToolbar.h
@@ -0,0 +1,33 @@
+// Copyright Epic Games, Inc. All Rights Reserved.
+
+#pragma once
+
+#include "SCommonEditorViewportToolbarBase.h"
+#include "Templates/SharedPointer.h"
+#include "Widgets/SCompoundWidget.h"
+
+class FMenuBuilder;
+class ICommonEditorViewportToolbarInfoProvider;
+class SEditorViewportViewMenu;
+
+/**
+ * Viewport toolbar widget used by the Chaos Visual Debugger
+ */
+class CHAOSVD_API SChaosVDViewportToolbar : public SCommonEditorViewportToolbarBase
+{
+public:
+	SLATE_BEGIN_ARGS(SChaosVDViewportToolbar)
+		{
+		}
+	SLATE_END_ARGS()
+
+	void Construct(const FArguments& InArgs, TSharedPtr<ICommonEditorViewportToolbarInfoProvider> InInfoProvider);
+
+protected:
+	virtual TSharedRef<SEditorViewportViewMenu> MakeViewMenu() override;
+	virtual void ExtendOptionsMenu(FMenuBuilder& OptionsMenuBuilder) const override;
+
+private:
+	virtual TSharedRef<SWidget> GenerateShowMenu() const override;
+
+};
