[ChaosVD] Added Remote Debugging Support

Now the tool is able to search and connect to live trace sessions either in a local machine or a remote machine. It also allows you to browse the local trace store for past recording not done to CVD specific directories.

- Added a new modal to allow browse sessions
- Added a way to hide or/and disable individual elements in the timeline widget
- Added simple Geometry Hashes cache in GeometryDataComponent, so we don't calculate the hash over an over during playback.
- Added a way to disable auto-stop when you reach the end of a recording
- Fixed issues where if the editor is ticking at a lower framerate, the target framerate in CVD is not respected.
- Added support to connect to live trace sessions
- Enabled Connect To session button, now that initial support for it is implemented

diff --git a/Engine/Plugins/Experimental/ChaosVD/Source/ChaosVD/Private/ChaosVDEngine.cpp b/Engine/Plugins/Experimental/ChaosVD/Source/ChaosVD/Private/ChaosVDEngine.cpp
--- a/Engine/Plugins/Experimental/ChaosVD/Source/ChaosVD/Private/ChaosVDEngine.cpp
+++ b/Engine/Plugins/Experimental/ChaosVD/Source/ChaosVD/Private/ChaosVDEngine.cpp
@@ -37,7 +37,7 @@ void FChaosVDEngine::DeInitialize()
 
 	if (const TSharedPtr<FChaosVDTraceManager> CVDTraceManager = FChaosVDModule::Get().GetTraceManager())
 	{
-		CVDTraceManager->CloseSession(CurrentSessionName);
+		CVDTraceManager->CloseSession(CurrentSessionDescriptor.SessionName);
 	}
 
 	CollectGarbage(GARBAGE_COLLECTION_KEEPFLAGS);
@@ -45,9 +45,17 @@ void FChaosVDEngine::DeInitialize()
 
 void FChaosVDEngine::LoadRecording(const FString& FilePath)
 {
-	CurrentSessionName = FChaosVDModule::Get().GetTraceManager()->LoadTraceFile(FilePath);
+	FChaosVDTraceSessionDescriptor NewSessionFromFileDescriptor;
+	NewSessionFromFileDescriptor.SessionName = FChaosVDModule::Get().GetTraceManager()->LoadTraceFile(FilePath);
+	NewSessionFromFileDescriptor.bIsLiveSession = false;
 
-	PlaybackController->LoadChaosVDRecordingFromTraceSession(CurrentSessionName);
+	SetCurrentSession(NewSessionFromFileDescriptor);
+}
+
+void FChaosVDEngine::SetCurrentSession(const FChaosVDTraceSessionDescriptor& SessionDescriptor)
+{
+	CurrentSessionDescriptor = SessionDescriptor;
+	PlaybackController->LoadChaosVDRecordingFromTraceSession(CurrentSessionDescriptor);
 }
 
 bool FChaosVDEngine::Tick(float DeltaTime)
diff --git a/Engine/Plugins/Experimental/ChaosVD/Source/ChaosVD/Private/ChaosVDGeometryDataComponent.cpp b/Engine/Plugins/Experimental/ChaosVD/Source/ChaosVD/Private/ChaosVDGeometryDataComponent.cpp
--- a/Engine/Plugins/Experimental/ChaosVD/Source/ChaosVD/Private/ChaosVDGeometryDataComponent.cpp
+++ b/Engine/Plugins/Experimental/ChaosVD/Source/ChaosVD/Private/ChaosVDGeometryDataComponent.cpp
@@ -71,7 +71,7 @@ void FChaosVDGeometryDataComponentBase::UpdateDataFromShapeArray_Internal(const
 			return true;
 		}
 
-		if (ImplicitA->GetTypeHash() == GeometryID)
+		if (GetCachedGeometryHash(ImplicitA) == GeometryID)
 		{
 			CollisionDataToUpdate = InShapeArray[RootObjectIndexA];
 			CollisionDataToUpdate.bIsComplex = FChaosVDGeometryBuilder::DoesImplicitContainType(ImplicitA, Chaos::ImplicitObjectType::HeightField) || FChaosVDGeometryBuilder::DoesImplicitContainType(ImplicitA, Chaos::ImplicitObjectType::TriangleMesh);
@@ -83,3 +83,24 @@ void FChaosVDGeometryDataComponentBase::UpdateDataFromShapeArray_Internal(const
 		return true;
 	});
 }
+
+
+uint32 FChaosVDGeometryDataComponentBase::GetCachedGeometryHash(const Chaos::FImplicitObject* InImplicitObject)
+{
+	if (!ensure(InImplicitObject))
+	{
+		return 0;
+	}
+
+	if (uint32* FoundCachedHashPtr = CachedGeometryHashes.Find(InImplicitObject))
+	{
+		return *FoundCachedHashPtr;
+	}
+	else
+	{
+		uint32 GeometryHash = InImplicitObject->GetTypeHash();
+		CachedGeometryHashes.Add(InImplicitObject, GeometryHash);
+
+		return GeometryHash;
+	}
+}
diff --git a/Engine/Plugins/Experimental/ChaosVD/Source/ChaosVD/Private/ChaosVDPlaybackController.cpp b/Engine/Plugins/Experimental/ChaosVD/Source/ChaosVD/Private/ChaosVDPlaybackController.cpp
--- a/Engine/Plugins/Experimental/ChaosVD/Source/ChaosVD/Private/ChaosVDPlaybackController.cpp
+++ b/Engine/Plugins/Experimental/ChaosVD/Source/ChaosVD/Private/ChaosVDPlaybackController.cpp
@@ -5,6 +5,7 @@
 #include "ChaosVDModule.h"
 #include "ChaosVDPlaybackControllerInstigator.h"
 #include "ChaosVDRecording.h"
+#include "ChaosVDRuntimeModule.h"
 #include "ChaosVDScene.h"
 #include "Trace/ChaosVDTraceManager.h"
 #include "Trace/ChaosVDTraceProvider.h"
@@ -13,16 +14,24 @@
 FChaosVDPlaybackController::FChaosVDPlaybackController(const TWeakPtr<FChaosVDScene>& InSceneToControl)
 {
 	SceneToControl = InSceneToControl;
+
+	RecordingStoppedHandle = FChaosVDRuntimeModule::Get().RegisterRecordingStopCallback(FChaosVDRecordingStateChangedDelegate::FDelegate::CreateRaw(this, &FChaosVDPlaybackController::HandleDisconnectedFromSession));
 }
 
 FChaosVDPlaybackController::~FChaosVDPlaybackController()
 {
+	// There is a chance the Runtime module is unloaded by now if we had the tool open and we are closing the editor
+	if (FChaosVDRuntimeModule::IsLoaded())
+	{
+		FChaosVDRuntimeModule::Get().RemoveRecordingStopCallback(RecordingStoppedHandle);
+	}
+
 	UnloadCurrentRecording(EChaosVDUnloadRecordingFlags::Silent);
 }
 
-bool FChaosVDPlaybackController::LoadChaosVDRecordingFromTraceSession(const FString& InSessionName)
+bool FChaosVDPlaybackController::LoadChaosVDRecordingFromTraceSession(const FChaosVDTraceSessionDescriptor& InSessionDescriptor)
 {
-	if (InSessionName.IsEmpty())
+	if (InSessionDescriptor.SessionName.IsEmpty())
 	{
 		return false;
 	}
@@ -32,7 +41,7 @@ bool FChaosVDPlaybackController::LoadChaosVDRecordingFromTraceSession(const FStr
 		UnloadCurrentRecording();
 	}
 
-	if (const TSharedPtr<const TraceServices::IAnalysisSession> TraceSession = FChaosVDModule::Get().GetTraceManager()->GetSession(InSessionName))
+	if (const TSharedPtr<const TraceServices::IAnalysisSession> TraceSession = FChaosVDModule::Get().GetTraceManager()->GetSession(InSessionDescriptor.SessionName))
 	{
 		if (const FChaosVDTraceProvider* ChaosVDProvider = TraceSession->ReadProvider<FChaosVDTraceProvider>(FChaosVDTraceProvider::ProviderName))
 		{
@@ -45,6 +54,8 @@ bool FChaosVDPlaybackController::LoadChaosVDRecordingFromTraceSession(const FStr
 		return false;
 	}
 
+	LoadedRecording->SetIsLive(InSessionDescriptor.bIsLiveSession);
+
 	LoadedRecording->OnRecordingUpdated().AddRaw(this, &FChaosVDPlaybackController::HandleCurrentRecordingUpdated);
 
 	HandleCurrentRecordingUpdated();
@@ -92,7 +103,7 @@ void FChaosVDPlaybackController::UnloadCurrentRecording(EChaosVDUnloadRecordingF
 		bHasPendingGTUpdateBroadcast = true;
 	}
 
-	bPlayedFirsFrame = false;
+	bPlayedFirstFrame = false;
 }
 
 void FChaosVDPlaybackController::PlayFromClosestKeyFrame_AssumesLocked(const int32 InTrackID, const int32 FrameNumber, FChaosVDScene& InSceneToControl)
@@ -220,6 +231,12 @@ void FChaosVDPlaybackController::GoToRecordedGameFrame_AssumesLocked(const int32
 		{
 			if (const FChaosVDGameFrameData* FrameData = LoadedRecording->GetGameFrameData_AssumesLocked(FrameNumber))
 			{
+				const bool bIsFrameFullyLoaded = FrameData->EndTime > 0;
+				if (!bIsFrameFullyLoaded)
+				{
+					return;
+				}
+
 				if (TrackInfoByIDMap* TrackInfoByID = TrackInfoPerType.Find(EChaosVDTrackType::Game))
 				{
 					TArray<int32> AvailableSolversID;
@@ -562,20 +579,54 @@ bool FChaosVDPlaybackController::Tick(float DeltaTime)
 		}
 	}
 
-	if (!bPlayedFirsFrame)
+	// Load at least the first frame
+	if (!bPlayedFirstFrame)
 	{
 		if (LoadedRecording.IsValid() && LoadedRecording->GetAvailableSolversNumber_AssumesLocked() > 0)
 		{
-			constexpr int32 GameFrame = 0;
+			constexpr int32 GameFrameToLoad = 0;
 			constexpr int32 Step = 0;
-			GoToTrackFrame(IChaosVDPlaybackControllerInstigator::InvalidGuid, EChaosVDTrackType::Game, GameTrackID, GameFrame, Step);
-			bPlayedFirsFrame = true;
+			GoToTrackFrame(IChaosVDPlaybackControllerInstigator::InvalidGuid, EChaosVDTrackType::Game, GameTrackID, GameFrameToLoad, Step);
+			bPlayedFirstFrame = true;
+		}
+	}
+
+	// If we are live, make sure we don't lag too much behind
+	if (!bPauseRequested && IsPlayingLiveSession())
+	{
+		if (const FChaosVDTrackInfo* GameTrackInfo = GetTrackInfo(EChaosVDTrackType::Game, GameTrackID))
+		{
+			const int32 CurrentFrameDeltaFromLast = FMath::Abs(GameTrackInfo->MaxFrames - GameTrackInfo->CurrentFrame);
+			if (CurrentFrameDeltaFromLast > MaxFramesLaggingBehindDuringLiveSession)
+			{
+				// Playing the middle point between last and the threshold. We don't want to play the last available frame as it could be incomplete,
+				// and we don't want to go to close to the threshold.
+				const int32 GameFrameToLoad = LoadedRecording->GetAvailableGameFramesNumber() - MinFramesLaggingBehindDuringLiveSession;
+				constexpr int32 Step = 0;
+				GoToTrackFrame(IChaosVDPlaybackControllerInstigator::InvalidGuid, EChaosVDTrackType::Game, GameTrackID, GameFrameToLoad, Step);
+			}
 		}
 	}
 
 	return true;
 }
 
+bool FChaosVDPlaybackController::IsPlayingLiveSession() const
+{
+	return LoadedRecording.IsValid() ? LoadedRecording->IsLive() : false;
+}
+
+void FChaosVDPlaybackController::HandleDisconnectedFromSession()
+{
+	if (LoadedRecording.IsValid())
+	{
+		LoadedRecording->SetIsLive(false);
+	}
+
+	// Queue a general update in the Game Thread
+	bHasPendingGTUpdateBroadcast = true;
+}
+
 void FChaosVDPlaybackController::UpdateSolverTracksData()
 {
 	const TMap<int32, TArray<FChaosVDSolverFrameData>>& SolversByID = LoadedRecording->GetAvailableSolvers_AssumesLocked();
diff --git a/Engine/Plugins/Experimental/ChaosVD/Source/ChaosVD/Private/Trace/ChaosVDTraceAnalyzer.cpp b/Engine/Plugins/Experimental/ChaosVD/Source/ChaosVD/Private/Trace/ChaosVDTraceAnalyzer.cpp
--- a/Engine/Plugins/Experimental/ChaosVD/Source/ChaosVD/Private/Trace/ChaosVDTraceAnalyzer.cpp
+++ b/Engine/Plugins/Experimental/ChaosVD/Source/ChaosVD/Private/Trace/ChaosVDTraceAnalyzer.cpp
@@ -55,6 +55,7 @@ bool FChaosVDTraceAnalyzer::OnEvent(uint16 RouteId, EStyle Style, const FOnEvent
 			{
 				FChaosVDGameFrameData FrameData;
 				FrameData.FirstCycle = EventData.GetValue<uint64>("Cycle");
+				FrameData.StartTime = Context.EventTime.AsSeconds(FrameData.FirstCycle);
 				ChaosVDTraceProvider->AddGameFrame(MoveTemp(FrameData));
 			}
 
@@ -70,6 +71,7 @@ bool FChaosVDTraceAnalyzer::OnEvent(uint16 RouteId, EStyle Style, const FOnEvent
 				if (FChaosVDGameFrameData* CurrentFrameData = ChaosVDTraceProvider->GetLastGameFrame_AssumesLocked())
 				{
 					CurrentFrameData->LastCycle = EventData.GetValue<uint64>("Cycle");
+					CurrentFrameData->EndTime = Context.EventTime.AsSeconds(CurrentFrameData->LastCycle);
 				}
 			}
 			break;
@@ -81,11 +83,21 @@ bool FChaosVDTraceAnalyzer::OnEvent(uint16 RouteId, EStyle Style, const FOnEvent
 			NewFrameData.SolverID = EventData.GetValue<int32>("SolverID");
 			NewFrameData.FrameCycle = EventData.GetValue<uint64>("Cycle");
 			NewFrameData.bIsKeyFrame = EventData.GetValue<bool>("IsKeyFrame");
+			NewFrameData.StartTime = Context.EventTime.AsSeconds(NewFrameData.FrameCycle);
 
 			FWideStringView DebugNameView;
 			EventData.GetString("DebugName", DebugNameView);
 			NewFrameData.DebugName = DebugNameView;
 
+			// Currently not all solvers have an end frame event, so lets just set the end frame time of the previous frame, with the start of this new one.
+			{
+				FWriteScopeLock WriteLock(ChaosVDTraceProvider->GetDataLock());
+				if (FChaosVDSolverFrameData* PrevFrameData  = ChaosVDTraceProvider->GetLastSolverFrame_AssumesLocked(NewFrameData.SolverID))
+				{
+					PrevFrameData->EndTime = NewFrameData.StartTime;
+				}
+			}
+
 			// Add an empty frame. It will be filled out by the solver trace events
 			ChaosVDTraceProvider->AddSolverFrame(NewFrameData.SolverID, MoveTemp(NewFrameData));
 			break;
diff --git a/Engine/Plugins/Experimental/ChaosVD/Source/ChaosVD/Private/Trace/ChaosVDTraceManager.cpp b/Engine/Plugins/Experimental/ChaosVD/Source/ChaosVD/Private/Trace/ChaosVDTraceManager.cpp
--- a/Engine/Plugins/Experimental/ChaosVD/Source/ChaosVD/Private/Trace/ChaosVDTraceManager.cpp
+++ b/Engine/Plugins/Experimental/ChaosVD/Source/ChaosVD/Private/Trace/ChaosVDTraceManager.cpp
@@ -2,9 +2,12 @@
 
 #include "Trace/ChaosVDTraceManager.h"
 
-#include "Modules/ModuleManager.h"
+#include "ChaosVDModule.h"
 #include "Features/IModularFeatures.h"
+#include "Modules/ModuleManager.h"
+#include "Misc/Paths.h"
 #include "Trace/ChaosVDTraceModule.h"
+#include "Trace/StoreClient.h"
 #include "TraceServices/AnalysisService.h"
 #include "TraceServices/ITraceServicesModule.h"
 #include "TraceServices/Model/AnalysisSession.h"
@@ -38,6 +41,72 @@ FString FChaosVDTraceManager::LoadTraceFile(const FString& InTraceFilename)
 	return FString();
 }
 
+FString FChaosVDTraceManager::ConnectToLiveSession(FStringView InSessionHost, uint32 SessionID)
+{	
+	using namespace UE::Trace;
+	FStoreClient* StoreClient = FStoreClient::Connect(InSessionHost.GetData());
+
+	FString SessionName;
+
+	if (!StoreClient)
+	{
+		return SessionName;
+	}
+
+	FStoreClient::FTraceData TraceData = StoreClient->ReadTrace(SessionID);
+	if (!TraceData.IsValid())
+	{
+		return SessionName;
+	}
+
+	FString TraceName(StoreClient->GetStatus()->GetStoreDir());
+	const FStoreClient::FTraceInfo* TraceInfo = StoreClient->GetTraceInfoById(SessionID);
+	if (TraceInfo != nullptr)
+	{
+		const FUtf8StringView Utf8NameView = TraceInfo->GetName();
+		FString Name(Utf8NameView);
+		if (!Name.EndsWith(TEXT(".utrace")))
+		{
+			Name += TEXT(".utrace");
+		}
+		TraceName = FPaths::Combine(TraceName, Name);
+		FPaths::NormalizeFilename(TraceName);
+	}
+
+	ITraceServicesModule& TraceServicesModule = FModuleManager::LoadModuleChecked<ITraceServicesModule>("TraceServices");
+	if (const TSharedPtr<TraceServices::IAnalysisService> TraceAnalysisService = TraceServicesModule.GetAnalysisService())
+	{
+		// Close this session in case we were already analysing it
+		CloseSession(TraceName);
+	
+		const TSharedPtr<const TraceServices::IAnalysisSession> NewSession = TraceAnalysisService->StartAnalysis(SessionID, *TraceName, MoveTemp(TraceData));
+		AnalysisSessionByName.Add(TraceName, NewSession);
+		SessionName = NewSession->GetName();
+	}
+
+	return SessionName;
+}
+
+FString FChaosVDTraceManager::GetLocalTraceStoreDirPath()
+{
+	UE::Trace::FStoreClient* StoreClient = UE::Trace::FStoreClient::Connect(TEXT("localhost"));
+
+	if (!StoreClient)
+	{
+		UE_LOG(LogChaosVDEditor, Error, TEXT("[%hs] Failed to connect to local Trace Store client"), ANSI_TO_TCHAR(__FUNCTION__));
+		return TEXT("");
+	}
+
+	const UE::Trace::FStoreClient::FStatus* Status = StoreClient->GetStatus();
+	if (!Status)
+	{
+		UE_LOG(LogChaosVDEditor, Error, TEXT("[%hs] Failed to to get Trace Store staus"), ANSI_TO_TCHAR(__FUNCTION__));
+		return TEXT("");
+	}
+
+	return FString(Status->GetStoreDir());
+}
+
 TSharedPtr<const TraceServices::IAnalysisSession> FChaosVDTraceManager::GetSession(const FString& InSessionName)
 {
 	if (TSharedPtr<const TraceServices::IAnalysisSession>* FoundSession = AnalysisSessionByName.Find(InSessionName))
diff --git a/Engine/Plugins/Experimental/ChaosVD/Source/ChaosVD/Private/Widgets/SChaosBrowseTraceFileSourceModal.cpp b/Engine/Plugins/Experimental/ChaosVD/Source/ChaosVD/Private/Widgets/SChaosBrowseTraceFileSourceModal.cpp
new file mode 100644
--- /dev/null
+++ b/Engine/Plugins/Experimental/ChaosVD/Source/ChaosVD/Private/Widgets/SChaosBrowseTraceFileSourceModal.cpp
@@ -0,0 +1,79 @@
+// Copyright Epic Games, Inc. All Rights Reserved.
+
+#include "Widgets/SChaosBrowseTraceFileSourceModal.h"
+
+#include "Editor.h"
+#include "GenericPlatform/GenericPlatformMisc.h"
+#include "Widgets/Layout/SUniformGridPanel.h"
+#include "Widgets/Text/STextBlock.h"
+#include "Widgets/Input/SButton.h"
+
+#define LOCTEXT_NAMESPACE "ChaosVisualDebugger"
+
+void SChaosBrowseTraceFileSourceModal::Construct(const FArguments& InArgs)
+{
+	SWindow::Construct(SWindow::FArguments()
+	.Title(LOCTEXT("SChaosVDBrowseFileModal_Title", "CVD Recording Source Selector"))
+	.SupportsMinimize(false)
+	.SupportsMaximize(false)
+	.UserResizeBorder(0)
+	.ClientSize(FVector2D(350, 80))
+	[
+		SNew(SVerticalBox)
+		+SVerticalBox::Slot()
+		.Padding(15)
+		[
+			SNew(SVerticalBox)
+			+SVerticalBox::Slot()
+			.AutoHeight()
+			.HAlign(HAlign_Center)
+			[
+				SNew(STextBlock)
+				.Text(LOCTEXT("SelectSourceMessage", "Where is the recording file located?"))
+				.Font(FCoreStyle::GetDefaultFontStyle("Regular", 10))
+			]
+		]
+		+SVerticalBox::Slot()
+			.AutoHeight()
+			.HAlign(HAlign_Center)
+			.Padding(5)
+			[
+				SNew(SUniformGridPanel)
+				.SlotPadding(FAppStyle::GetMargin("StandardDialog.SlotPadding"))
+				.MinDesiredSlotWidth(FAppStyle::GetFloat("StandardDialog.MinDesiredSlotWidth"))
+				.MinDesiredSlotHeight(FAppStyle::GetFloat("StandardDialog.MinDesiredSlotHeight"))
+				+SUniformGridPanel::Slot(0, 0)
+				[
+					SNew(SButton)
+					.HAlign(HAlign_Center)
+					.ContentPadding(FAppStyle::GetMargin("StandardDialog.ContentPadding"))
+					.Text(LOCTEXT("BrowseFolder", "Folder"))
+					.OnClicked(this, &SChaosBrowseTraceFileSourceModal::OnButtonClick, EChaosVDBrowseFileModalResponse::OpenFolder)
+				]
+				+SUniformGridPanel::Slot(1, 0)
+				[
+					SNew(SButton)
+					.HAlign(HAlign_Center)
+					.ContentPadding(FAppStyle::GetMargin("StandardDialog.ContentPadding"))
+					.Text(LOCTEXT("BrowseTraceStore", "Trace Store"))
+					.OnClicked(this, &SChaosBrowseTraceFileSourceModal::OnButtonClick, EChaosVDBrowseFileModalResponse::OpenTraceStore)
+				]
+			]
+	]);
+}
+
+EChaosVDBrowseFileModalResponse SChaosBrowseTraceFileSourceModal::ShowModal()
+{
+	GEditor->EditorAddModalWindow(SharedThis(this));
+	return UserResponse;
+}
+
+FReply SChaosBrowseTraceFileSourceModal::OnButtonClick(EChaosVDBrowseFileModalResponse Response)
+{
+	UserResponse = Response;
+	RequestDestroyWindow();
+	return FReply::Handled();
+}
+
+#undef LOCTEXT_NAMESPACE
+
diff --git a/Engine/Plugins/Experimental/ChaosVD/Source/ChaosVD/Private/Widgets/SChaosVDBrowseSessionsModal.cpp b/Engine/Plugins/Experimental/ChaosVD/Source/ChaosVD/Private/Widgets/SChaosVDBrowseSessionsModal.cpp
new file mode 100644
--- /dev/null
+++ b/Engine/Plugins/Experimental/ChaosVD/Source/ChaosVD/Private/Widgets/SChaosVDBrowseSessionsModal.cpp
@@ -0,0 +1,295 @@
+// Copyright Epic Games, Inc. All Rights Reserved.
+
+#include "Widgets/SChaosVDBrowseSessionsModal.h"
+
+#include "ChaosVDModule.h"
+#include "Editor.h"
+#include "Framework/Application/SlateApplication.h"
+#include "GenericPlatform/GenericPlatformMisc.h"
+#include "Trace/StoreClient.h"
+#include "Widgets/SChaosVDNameListPicker.h"
+#include "Widgets/Layout/SUniformGridPanel.h"
+#include "Widgets/Text/STextBlock.h"
+#include "Widgets/Input/SButton.h"
+
+#define LOCTEXT_NAMESPACE "ChaosVisualDebugger"
+
+/** Reduced Version of the IP4 structure from the Network Module
+ * used just to convert to String an IP saved as host byte order.
+ */
+struct FCVDIPv4Address
+{
+	union
+	{
+		/** The IP address value as A.B.C.D components. */
+		struct
+		{
+#if PLATFORM_LITTLE_ENDIAN
+#ifdef _MSC_VER
+			uint8 D, C, B, A;
+#else
+			uint8 D GCC_ALIGN(4);
+			uint8 C, B, A;
+#endif
+#else
+			uint8 A, B, C, D;
+#endif
+		};
+
+		/** The IP address value in host byte order. */
+		uint32 Value;
+	};
+	
+	/**
+ * Creates and initializes a new IPv4 address with the specified components.
+ *
+ * The created IP address has the value A.B.C.D.
+ *
+ * @param InA The first component.
+ * @param InB The second component.
+ * @param InC The third component.
+ * @param InD The fourth component.
+ */
+	FCVDIPv4Address(uint8 InA, uint8 InB, uint8 InC, uint8 InD)
+#if PLATFORM_LITTLE_ENDIAN
+		: D(InD)
+		, C(InC)
+		, B(InB)
+		, A(InA)
+#else
+		: A(InA)
+		, B(InB)
+		, C(InC)
+		, D(InD)
+#endif // PLATFORM_LITTLE_ENDIAN
+	{ }
+
+	/**
+	 * Creates and initializes a new IPv4 address with the specified value.
+	 *
+	 * @param InValue The address value (in host byte order).
+	 */
+	FCVDIPv4Address(uint32 InValue)
+		: Value(InValue)
+	{ }
+
+	FString ToString() const { return FString::Printf(TEXT("%i.%i.%i.%i"), A, B, C, D); }
+};
+
+SChaosVDBrowseSessionsModal::~SChaosVDBrowseSessionsModal()
+{
+	if (FSlateApplication::IsInitialized())
+	{
+		if (ModalTickHandle.IsValid())
+		{
+			FSlateApplication::Get().GetOnModalLoopTickEvent().Remove(ModalTickHandle);
+		}
+	}
+}
+
+void SChaosVDBrowseSessionsModal::Construct(const FArguments& InArgs)
+{
+	CurrentTraceStoreAddress = TEXT("127.0.0.1");
+	
+	SWindow::Construct(SWindow::FArguments()
+	.Title(LOCTEXT("SChaosVDBrowseSessionsModal_Title", "Live Session Browser"))
+	.SupportsMinimize(false)
+	.SupportsMaximize(false)
+	.UserResizeBorder(0)
+	.ClientSize(FVector2D(450, 120))
+	[
+		SNew(SVerticalBox)
+		+SVerticalBox::Slot()
+		.Padding(5)
+		[
+			SNew(SBorder)
+			.BorderImage(FAppStyle::GetBrush("ToolPanel.GroupBorder"))
+			[
+				SNew(SVerticalBox)
+				+SVerticalBox::Slot()
+				.AutoHeight()
+				[
+					SNew(SHorizontalBox)
+					+SHorizontalBox::Slot()
+					.HAlign(HAlign_Left)
+					.VAlign(VAlign_Center)
+					[
+						SNew(STextBlock)
+						.Text(LOCTEXT("TraceStoreAddress", "Trace Store Address"))
+						.Font(FCoreStyle::GetDefaultFontStyle("Regular", 12))
+					]
+					+SHorizontalBox::Slot()
+					.HAlign(HAlign_Right)
+					.VAlign(VAlign_Center)
+					[
+						SAssignNew(TraceStoreAddressWidget, SEditableTextBox)
+						.Text(FText::AsCultureInvariant(CurrentTraceStoreAddress))
+						.Font(FCoreStyle::GetDefaultFontStyle("Regular", 12))
+						.OnTextCommitted_Raw(this, &SChaosVDBrowseSessionsModal::OnTraceStoreAddressUpdated)
+					]
+				]
+				+SVerticalBox::Slot()
+				.AutoHeight()
+				[
+					SNew(SHorizontalBox)
+					+SHorizontalBox::Slot()
+					.HAlign(HAlign_Left)
+					.VAlign(VAlign_Center)
+					[
+						SNew(STextBlock)
+						.Text(LOCTEXT("SelectedLiveSession", "Selected Live Session"))
+						.Font(FCoreStyle::GetDefaultFontStyle("Regular", 12))
+					]
+					+SHorizontalBox::Slot()
+					.HAlign(HAlign_Right)
+					.VAlign(VAlign_Center)
+					[
+						SAssignNew(NamePickerWidget, SChaosVDNameListPicker)
+						.OnNameSleceted_Raw(this, &SChaosVDBrowseSessionsModal::HandleSessionNameSelected)
+					]	
+				]
+			]
+		]
+		+SVerticalBox::Slot()
+		.AutoHeight()
+		.HAlign(HAlign_Right)
+		.Padding(5)
+		[
+			SNew(SUniformGridPanel)
+			.SlotPadding(FAppStyle::GetMargin("StandardDialog.SlotPadding"))
+			.MinDesiredSlotWidth(FAppStyle::GetFloat("StandardDialog.MinDesiredSlotWidth"))
+			.MinDesiredSlotHeight(FAppStyle::GetFloat("StandardDialog.MinDesiredSlotHeight"))
+			+SUniformGridPanel::Slot(0, 0)
+			[
+				SNew(SButton)
+				.HAlign(HAlign_Center)
+				.ContentPadding(FAppStyle::GetMargin("StandardDialog.ContentPadding"))
+				.Text(LOCTEXT("ConnectToSession", "Connect to Session"))
+				.IsEnabled_Raw(this, &SChaosVDBrowseSessionsModal::CanOpenSession)
+				.OnClicked(this, &SChaosVDBrowseSessionsModal::OnButtonClick, EAppReturnType::Ok)
+			]
+			+SUniformGridPanel::Slot(1, 0)
+			[
+				SNew(SButton)
+				.HAlign(HAlign_Center)
+				.ContentPadding(FAppStyle::GetMargin("StandardDialog.ContentPadding"))
+				.Text(LOCTEXT("Cancel", "Cancel"))
+				.OnClicked(this, &SChaosVDBrowseSessionsModal::OnButtonClick, EAppReturnType::Cancel)
+			]
+		]
+	]);
+
+	UpdateCurrentSessionInfoMap();
+
+	if (FSlateApplication::IsInitialized())
+	{
+		ModalTickHandle = FSlateApplication::Get().GetOnModalLoopTickEvent().AddRaw(this, &SChaosVDBrowseSessionsModal::ModalTick);
+	}
+}
+
+EAppReturnType::Type SChaosVDBrowseSessionsModal::ShowModal()
+{
+	GEditor->EditorAddModalWindow(SharedThis(this));
+	return UserResponse;
+}
+
+FChaosVDTraceSessionInfo SChaosVDBrowseSessionsModal::GetSelectedTraceInfo()
+{
+	if (ensure(CurrentTraceSessionSelected.IsValid()))
+	{
+		if (const FChaosVDTraceSessionInfo* SessionInfoPtr = CurrentSessionInfosMap.Find(CurrentTraceSessionSelected))
+		{
+			return *SessionInfoPtr;
+		}
+		else
+		{
+			UE_LOG(LogChaosVDEditor, Error, TEXT("Selected session [%s] is no longer available"), *CurrentTraceSessionSelected.ToString())
+		}
+	}
+	else
+	{
+		UE_LOG(LogChaosVDEditor, Error, TEXT("Selected session [%s] is invalid"), *CurrentTraceSessionSelected.ToString())
+	}
+
+	return FChaosVDTraceSessionInfo();
+}
+
+void SChaosVDBrowseSessionsModal::ModalTick(float InDeltaTime)
+{
+	AccumulatedTimeBetweenTicks += InDeltaTime;
+
+	// Update the sessions at least once per desired interval
+	constexpr float PollingInterval = 1.0f;
+	if (AccumulatedTimeBetweenTicks > PollingInterval)
+	{
+		UpdateCurrentSessionInfoMap();
+		AccumulatedTimeBetweenTicks = 0.0f;
+	}
+}
+
+bool SChaosVDBrowseSessionsModal::CanOpenSession() const
+{
+	return !CurrentTraceSessionSelected.IsNone();
+}
+
+void SChaosVDBrowseSessionsModal::UpdateCurrentSessionInfoMap()
+{
+	using namespace UE::Trace;
+	const FStoreClient* StoreClient = FStoreClient::Connect(*CurrentTraceStoreAddress);
+
+	if (!StoreClient)
+	{
+		UE_LOG(LogChaosVDEditor, Error, TEXT("Failed to connect to trace store at [%s]"), *CurrentTraceStoreAddress)
+		return;
+	}
+
+	const uint32 SessionCount = StoreClient->GetSessionCount();
+	CurrentSessionInfosMap.Empty(SessionCount);
+
+	for (uint32 SessionIndex = 0; SessionIndex < SessionCount; SessionIndex++)
+	{
+		if (const FStoreClient::FSessionInfo* SessionInfo = StoreClient->GetSessionInfo(SessionIndex))
+		{
+			//TODO: We need a way to know if it is an IPV6
+			FCVDIPv4Address IP(SessionInfo->GetIpAddress());
+			
+			FStringFormatOrderedArguments Args {IP.ToString(), FString::FromInt(SessionInfo->GetId()) };
+			FString SessionDisplayNameString = FString::Format(TEXT("IP: {0} + {1}"), Args);
+
+			FChaosVDTraceSessionInfo CVDSessionInfo { SessionInfo->GetTraceId(), SessionInfo->GetIpAddress(), SessionInfo->GetControlPort(), true };
+
+			CurrentSessionInfosMap.Add(FName(SessionDisplayNameString), CVDSessionInfo);
+		}	
+	}
+
+	TArray<TSharedPtr<FName>> NewNameList;
+	NewNameList.Reserve(CurrentSessionInfosMap.Num());
+
+	Algo::Transform(CurrentSessionInfosMap, NewNameList, [](const TPair<FName, FChaosVDTraceSessionInfo>& TransformData){ return MakeShared<FName>(TransformData.Key);});
+	
+	NamePickerWidget->UpdateNameList(MoveTemp(NewNameList));
+}
+
+FReply SChaosVDBrowseSessionsModal::OnButtonClick(EAppReturnType::Type ButtonID)
+{
+	UserResponse = ButtonID;
+	RequestDestroyWindow();
+	return FReply::Handled();
+}
+
+void SChaosVDBrowseSessionsModal::HandleSessionNameSelected(TSharedPtr<FName> SelectedName)
+{
+	CurrentTraceSessionSelected = SelectedName.IsValid() ? *SelectedName : FName();
+}
+
+void SChaosVDBrowseSessionsModal::OnTraceStoreAddressUpdated(const FText& InText, ETextCommit::Type CommitType)
+{
+	if (CommitType == ETextCommit::OnEnter || CommitType == ETextCommit::Default)
+	{
+		CurrentTraceStoreAddress = InText.ToString();
+	}
+
+	UpdateCurrentSessionInfoMap();
+}
+
+#undef LOCTEXT_NAMESPACE
diff --git a/Engine/Plugins/Experimental/ChaosVD/Source/ChaosVD/Private/Widgets/SChaosVDMainTab.cpp b/Engine/Plugins/Experimental/ChaosVD/Source/ChaosVD/Private/Widgets/SChaosVDMainTab.cpp
--- a/Engine/Plugins/Experimental/ChaosVD/Source/ChaosVD/Private/Widgets/SChaosVDMainTab.cpp
+++ b/Engine/Plugins/Experimental/ChaosVD/Source/ChaosVD/Private/Widgets/SChaosVDMainTab.cpp
@@ -4,6 +4,7 @@
 
 #include "ChaosVDEditorVisualizationSettingsTab.h"
 #include "ChaosVDEngine.h"
+#include "ChaosVDModule.h"
 #include "ChaosVDObjectDetailsTab.h"
 #include "ChaosVDOutputLogTab.h"
 #include "ChaosVDPlaybackController.h"
@@ -19,9 +20,13 @@
 #include "Framework/Commands/UICommandList.h"
 #include "Framework/MultiBox/MultiBoxBuilder.h"
 #include "IDesktopPlatform.h"
+#include "Misc/MessageDialog.h"
 #include "StatusBarSubsystem.h"
 #include "Styling/StyleColors.h"
 #include "Styling/ToolBarStyle.h"
+#include "Trace/ChaosVDTraceManager.h"
+#include "Widgets/SChaosBrowseTraceFileSourceModal.h"
+#include "Widgets/SChaosVDBrowseSessionsModal.h"
 #include "Widgets/SChaosVDRecordingControls.h"
 #include "Widgets/Docking/SDockTab.h"
 #include "Widgets/Input/SComboButton.h"
@@ -82,7 +87,7 @@ void SChaosVDMainTab::Construct(const FArguments& InArgs, TSharedPtr<FChaosVDEng
 					.ContentPadding(FMargin(0, 5.f, 0, 4.f))
 					.OnClicked_Lambda([this]()
 					{
-						BrowseAndOpenChaosVDFile();
+						BrowseAndOpenChaosVDRecording();
 						return FReply::Handled();
 					})
 					.Content()
@@ -116,14 +121,8 @@ void SChaosVDMainTab::Construct(const FArguments& InArgs, TSharedPtr<FChaosVDEng
 				.AutoWidth()
 				[
 					SNew(SButton)
-					.ToolTip(SNew(SToolTip).Text(LOCTEXT("RemoteConnectionDesc", "Not Supported yet - Click here to connect to a remote debug session.")))
 					.ContentPadding(FMargin(0, 5.f, 0, 4.f))
-					.IsEnabled(false) // We don't support remote debugging yet
-					.OnClicked_Lambda([this]()
-					{
-						//TODO : Add a call to the method that will open the Trace Session Browser
-						return FReply::Handled();
-					})
+					.OnClicked_Raw(this, &SChaosVDMainTab::HandleSessionConnectionClicked)
 					.Content()
 					[
 						SNew(SHorizontalBox)
@@ -142,7 +141,7 @@ void SChaosVDMainTab::Construct(const FArguments& InArgs, TSharedPtr<FChaosVDEng
 						[
 							SNew(STextBlock)
 							.TextStyle(FAppStyle::Get(), "SmallButtonText")
-							.Text(LOCTEXT("ConnectToSession", "Connect to Session"))
+							.Text_Raw(this, &SChaosVDMainTab::GetConnectButtonText)
 						]
 					]
 				]
@@ -282,7 +281,7 @@ TSharedRef<FTabManager::FLayout> SChaosVDMainTab::GenerateMainLayout()
 					FTabManager::NewStack()
 					->SetSizeCoefficient(0.3f)
 					->AddTab(FChaosVDTabID::DetailsPanel, ETabState::OpenedTab)
-					->AddTab(FChaosVDTabID::CVDEditorSettings, ETabState::OpenedTab)
+					->AddTab(FChaosVDTabID::CVDEditorSettings, ETabState::ClosedTab)
 				)
 			)
 		);
@@ -300,7 +299,7 @@ void SChaosVDMainTab::GenerateMainWindowMenu()
 			FUIAction(
 				FExecuteAction::CreateLambda([this]()
 				{
-					BrowseAndOpenChaosVDFile();
+					BrowseAndOpenChaosVDRecording();
 				})));
 		}),
 		"File"
@@ -319,12 +318,37 @@ void SChaosVDMainTab::GenerateMainWindowMenu()
 	TabManager->SetMenuMultiBox(MenuBarBuilder.GetMultiBox(), MenuBarBuilder.MakeWidget());
 }
 
-void SChaosVDMainTab::BrowseAndOpenChaosVDFile()
+void SChaosVDMainTab::BrowseAndOpenChaosVDRecording()
+{
+	const TSharedRef<SChaosBrowseTraceFileSourceModal> SessionBrowserModal = SNew(SChaosBrowseTraceFileSourceModal);
+
+	EChaosVDBrowseFileModalResponse Response = SessionBrowserModal->ShowModal();
+	switch(Response)
+	{
+		case EChaosVDBrowseFileModalResponse::OpenFolder:
+			{
+				BrowseChaosVDRecordingFromFolder();
+				break;
+			}
+		case EChaosVDBrowseFileModalResponse::OpenTraceStore:
+			{
+				//TODO: Support remote Trace Stores 
+				BrowseChaosVDRecordingFromFolder(FChaosVDModule::Get().GetTraceManager()->GetLocalTraceStoreDirPath());
+				break;
+			}
+		case EChaosVDBrowseFileModalResponse::Cancel:
+			break;
+		default:
+				ensureMsgf(false, TEXT("Invalid responce received"));
+			break;
+	}
+}
+
+void SChaosVDMainTab::BrowseChaosVDRecordingFromFolder(FStringView FolderPath)
 {
 	TArray<FString> OutOpenFilenames;
 	if (IDesktopPlatform* DesktopPlatform = FDesktopPlatformModule::Get())
 	{
-		
 		FString ExtensionStr;
 		ExtensionStr += TEXT("Unreal Trace|*.utrace|");
 		//TODO: Re-enable this when we add "Clips" support as these will use our own format
@@ -333,7 +357,7 @@ void SChaosVDMainTab::BrowseAndOpenChaosVDFile()
 		DesktopPlatform->OpenFileDialog(
 			FSlateApplication::Get().FindBestParentWindowHandleForDialogs(AsShared()),
 			LOCTEXT("OpenDialogTitle", "Open Chaos Visual Debug File").ToString(),
-			TEXT(""),
+			FolderPath.GetData(),
 			TEXT(""),
 			*ExtensionStr,
 			EFileDialogFlags::None,
@@ -350,4 +374,57 @@ void SChaosVDMainTab::BrowseAndOpenChaosVDFile()
 	}
 }
 
+void SChaosVDMainTab::BrowseLiveSessionsFromTraceStore() const
+{
+	const TSharedRef<SChaosVDBrowseSessionsModal> SessionBrowserModal = SNew(SChaosVDBrowseSessionsModal);
+
+	if (SessionBrowserModal->ShowModal() != EAppReturnType::Cancel)
+	{
+		bool bSuccess = false;
+		const FChaosVDTraceSessionInfo SessionInfo = SessionBrowserModal->GetSelectedTraceInfo();
+		if (SessionInfo.bIsValid)
+		{
+			FChaosVDTraceSessionDescriptor NewSessionFromFileDescriptor;
+			NewSessionFromFileDescriptor.SessionName =  FChaosVDModule::Get().GetTraceManager()->ConnectToLiveSession(SessionBrowserModal->GetSelectedTraceStoreAddress(), SessionInfo.TraceID);
+			NewSessionFromFileDescriptor.bIsLiveSession = true;
+
+			bSuccess = !NewSessionFromFileDescriptor.SessionName.IsEmpty();
+
+			// If it failed we want to clean the current session name, so it is ok calling it either way
+			GetChaosVDEngineInstance()->SetCurrentSession(NewSessionFromFileDescriptor);
+		}
+
+		if (!bSuccess)
+		{
+			FMessageDialog::Open(EAppMsgType::Ok, LOCTEXT("FailedToConnectToSessionMessage", "Failed to connect to session"));	
+		}
+	}
+}
+
+FReply SChaosVDMainTab::HandleSessionConnectionClicked()
+{
+	if (TSharedPtr<FChaosVDPlaybackController> PlaybackControllerPtr =  GetChaosVDEngineInstance()->GetPlaybackController())
+	{
+		const bool bIsAlreadyInLiveSession = PlaybackControllerPtr->IsPlayingLiveSession();
+	
+		if (bIsAlreadyInLiveSession)
+		{
+			FChaosVDModule::Get().GetTraceManager()->CloseSession(GetChaosVDEngineInstance()->GetCurrentSessionDescriptor().SessionName);
+			PlaybackControllerPtr->HandleDisconnectedFromSession();
+		}
+		else
+		{
+			BrowseLiveSessionsFromTraceStore();
+		}
+	}
+
+	return FReply::Handled();
+}
+
+FText SChaosVDMainTab::GetConnectButtonText() const
+{
+	const bool bIsAlreadyInLiveSession = GetChaosVDEngineInstance()->GetPlaybackController()->IsPlayingLiveSession();
+	return bIsAlreadyInLiveSession ? LOCTEXT("DisconnectFromSession", "Disconnect from Session") : LOCTEXT("ConnectToSession", "Connect to Session");
+}
+
 #undef LOCTEXT_NAMESPACE
diff --git a/Engine/Plugins/Experimental/ChaosVD/Source/ChaosVD/Private/Widgets/SChaosVDPlaybackViewport.cpp b/Engine/Plugins/Experimental/ChaosVD/Source/ChaosVD/Private/Widgets/SChaosVDPlaybackViewport.cpp
--- a/Engine/Plugins/Experimental/ChaosVD/Source/ChaosVD/Private/Widgets/SChaosVDPlaybackViewport.cpp
+++ b/Engine/Plugins/Experimental/ChaosVD/Source/ChaosVD/Private/Widgets/SChaosVDPlaybackViewport.cpp
@@ -82,9 +82,9 @@ void SChaosVDPlaybackViewport::Construct(const FArguments& InArgs, TWeakPtr<FCha
 			+SVerticalBox::Slot()
 			[
 				SAssignNew(GameFramesTimelineWidget, SChaosVDTimelineWidget)
-				.HidePlayStopButtons(false)
-				.HideLockButton(true)
+				.ButtonVisibilityFlags(static_cast<uint16>(EChaosVDTimelineElementIDFlags::AllPlayback))
 				.OnFrameChanged_Raw(this, &SChaosVDPlaybackViewport::OnFrameSelectionUpdated)
+				.OnButtonClicked(this, &SChaosVDPlaybackViewport::HandlePlaybackButtonClicked)
 				.MaxFrames(0)
 			]
 		]	
@@ -157,6 +157,25 @@ void SChaosVDPlaybackViewport::HandlePlaybackControllerDataUpdated(TWeakPtr<FCha
 			// Max is inclusive and we use this to request as the index on the recorded frames/steps arrays so we need to -1 to the available frames/steps
 			GameFramesTimelineWidget->UpdateMinMaxValue(0, TrackInfo->MaxFrames != INDEX_NONE ? TrackInfo->MaxFrames -1  : 0);
 			GameFramesTimelineWidget->SetCurrentTimelineFrame(TrackInfo->CurrentFrame);
+
+			constexpr uint16 PlaybackElementDisabledDuringLiveSession = static_cast<uint16>(EChaosVDTimelineElementIDFlags::Stop | EChaosVDTimelineElementIDFlags::Next | EChaosVDTimelineElementIDFlags::Prev);
+			if (ControllerSharedPtr->IsPlayingLiveSession())
+			{
+				if (!ControllerSharedPtr->HasPauseRequest())
+				{
+					GameFramesTimelineWidget->SetAutoStopEnabled(false);
+					GameFramesTimelineWidget->Play();
+				}
+
+				uint16& CurrentEnabledFlags = GameFramesTimelineWidget->GetMutableElementEnabledFlagsRef();
+				CurrentEnabledFlags = CurrentEnabledFlags &~ PlaybackElementDisabledDuringLiveSession;
+			}
+			else
+			{
+				uint16& CurrentEnabledFlags = GameFramesTimelineWidget->GetMutableElementEnabledFlagsRef();
+				CurrentEnabledFlags = CurrentEnabledFlags | PlaybackElementDisabledDuringLiveSession;
+				GameFramesTimelineWidget->SetAutoStopEnabled(true);
+			}
 		}
 	}
 	else
@@ -186,6 +205,7 @@ void SChaosVDPlaybackViewport::HandleControllerTrackFrameUpdated(TWeakPtr<FChaos
 			if (FChaosVDGameFrameData* FrameData = RecordingData->GetGameFrameData_AssumesLocked(GameTrackFrame))
 			{
 				PlaybackViewportClient->PerformSelectedTrackingForFrame(FrameData);
+				GameFramesTimelineWidget->SetTargetFrameTime(FrameData->GetFrameTime());
 			}	
 		}
 	}
@@ -236,4 +256,26 @@ void SChaosVDPlaybackViewport::OnFrameSelectionUpdated(int32 NewFrameIndex) cons
 	}
 }
 
+void SChaosVDPlaybackViewport::HandlePlaybackButtonClicked(EChaosVDPlaybackButtonsID ButtonID)
+{
+	if (const TSharedPtr<FChaosVDPlaybackController> PlaybackControllerPtr = PlaybackController.Pin())
+	{
+		switch (ButtonID)
+		{
+		case EChaosVDPlaybackButtonsID::Play:
+			{
+				PlaybackControllerPtr->RequestUnpause();
+				break;
+			}
+		case EChaosVDPlaybackButtonsID::Pause:
+			{
+				PlaybackControllerPtr->RequestPause();
+				break;
+			}
+		default:
+			break;
+		}
+	}
+}
+
 #undef LOCTEXT_NAMESPACE
diff --git a/Engine/Plugins/Experimental/ChaosVD/Source/ChaosVD/Private/Widgets/SChaosVDRecordingControls.cpp b/Engine/Plugins/Experimental/ChaosVD/Source/ChaosVD/Private/Widgets/SChaosVDRecordingControls.cpp
--- a/Engine/Plugins/Experimental/ChaosVD/Source/ChaosVD/Private/Widgets/SChaosVDRecordingControls.cpp
+++ b/Engine/Plugins/Experimental/ChaosVD/Source/ChaosVD/Private/Widgets/SChaosVDRecordingControls.cpp
@@ -90,16 +90,28 @@ const FSlateBrush* SChaosVDRecordingControls::GetRecordOrStopButton() const
 
 void SChaosVDRecordingControls::HandleRecordingStop()
 {
+	const TSharedPtr<SChaosVDMainTab> MainTabSharedPtr = MainTabWeakPtr.Pin();
+	if (!MainTabSharedPtr.IsValid())
+	{
+		return;
+	}
+
 	UStatusBarSubsystem* StatusBarSubsystem = GEditor ? GEditor->GetEditorSubsystem<UStatusBarSubsystem>() : nullptr;
 	check(StatusBarSubsystem);
 	StatusBarSubsystem->PopStatusBarMessage(StatusBarID, RecordingMessageHandle);
 
-	FText RecordingPathMessage = FText::Format(LOCTEXT("RecordingSavedPathMessage"," Recoring saved at {0} "), FText::AsCultureInvariant(FChaosVDRuntimeModule::Get().GetActiveRecordingFileName()));
-	RecordingPathMessageHandle = StatusBarSubsystem->PushStatusBarMessage(StatusBarID, RecordingPathMessage);
+	if (MainTabSharedPtr->GetChaosVDEngineInstance()->GetCurrentSessionDescriptor().bIsLiveSession)
+	{
+		FText LiveSessionEnded = LOCTEXT("LiveSessionEndedMessage"," Live session has ended");
+		LiveSessionEndedMessageHandle = StatusBarSubsystem->PushStatusBarMessage(StatusBarID, LiveSessionEnded);
 
-	if (FMessageDialog::Open(EAppMsgType::YesNo, LOCTEXT("OpenLastRecordingMessage", "Do you want to load the recorded file now? ")) == EAppReturnType::Yes)
+	}
+	else
 	{
-		if (TSharedPtr<SChaosVDMainTab> MainTabSharedPtr = MainTabWeakPtr.Pin())
+		FText RecordingPathMessage = FText::Format(LOCTEXT("RecordingSavedPathMessage"," Recoring saved at {0} "), FText::AsCultureInvariant(FChaosVDRuntimeModule::Get().GetActiveRecordingFileName()));
+		RecordingPathMessageHandle = StatusBarSubsystem->PushStatusBarMessage(StatusBarID, RecordingPathMessage);
+		
+		if (FMessageDialog::Open(EAppMsgType::YesNo, LOCTEXT("OpenLastRecordingMessage", "Do you want to load the recorded file now? ")) == EAppReturnType::Yes)
 		{
 			MainTabSharedPtr->GetChaosVDEngineInstance()->LoadRecording(FChaosVDRuntimeModule::Get().GetActiveRecordingFileName());
 		}
@@ -117,6 +129,13 @@ void SChaosVDRecordingControls::HandleRecordingStart()
 	if (RecordingPathMessageHandle.IsValid())
 	{
 		StatusBarSubsystem->PopStatusBarMessage(StatusBarID, RecordingPathMessageHandle);
+		RecordingPathMessageHandle = FStatusBarMessageHandle();
+	}
+	
+	if (LiveSessionEndedMessageHandle.IsValid())
+	{
+		StatusBarSubsystem->PopStatusBarMessage(StatusBarID, LiveSessionEndedMessageHandle);
+		LiveSessionEndedMessageHandle = FStatusBarMessageHandle();
 	}
 
 	RecordingMessageHandle = StatusBarSubsystem->PushStatusBarMessage(StatusBarID, LOCTEXT("RecordingMessgae", "Recording..."));
@@ -138,7 +157,7 @@ FReply SChaosVDRecordingControls::ToggleRecordingState()
 
 bool SChaosVDRecordingControls::IsRecording() const
 {
-#if WITH_CHAOS_VISUAL_DEBUGGER
+#if WITH_CHAOS_VISUAL_DEBUGGER && UE_TRACE_ENABLED
 	return FChaosVisualDebuggerTrace::IsTracing();
 #endif
 	return false;
diff --git a/Engine/Plugins/Experimental/ChaosVD/Source/ChaosVD/Private/Widgets/SChaosVDSolverPlaybackControls.cpp b/Engine/Plugins/Experimental/ChaosVD/Source/ChaosVD/Private/Widgets/SChaosVDSolverPlaybackControls.cpp
--- a/Engine/Plugins/Experimental/ChaosVD/Source/ChaosVD/Private/Widgets/SChaosVDSolverPlaybackControls.cpp
+++ b/Engine/Plugins/Experimental/ChaosVD/Source/ChaosVD/Private/Widgets/SChaosVDSolverPlaybackControls.cpp
@@ -39,8 +39,7 @@ void SChaosVDSolverPlaybackControls::Construct(const FArguments& InArgs, int32 I
 				+SVerticalBox::Slot()
 				[
 					SAssignNew(FramesTimelineWidget, SChaosVDTimelineWidget)
-						.HidePlayStopButtons(false)
-						.HideLockButton(true)
+						.ButtonVisibilityFlags(static_cast<uint16>(EChaosVDTimelineElementIDFlags::AllPlayback))
 						.OnFrameChanged_Raw(this, &SChaosVDSolverPlaybackControls::OnFrameSelectionUpdated)
 						.MaxFrames(0)
 				]
@@ -60,7 +59,7 @@ void SChaosVDSolverPlaybackControls::Construct(const FArguments& InArgs, int32 I
 				+SVerticalBox::Slot()
 				[
 					SAssignNew(StepsTimelineWidget, SChaosVDTimelineWidget)
-					.HidePlayStopButtons(true)
+					.ButtonVisibilityFlags(static_cast<uint16>(EChaosVDTimelineElementIDFlags::AllManualStepping | EChaosVDTimelineElementIDFlags::Lock))
 					.OnFrameLockStateChanged_Raw(this, &SChaosVDSolverPlaybackControls::HandleLockStateChanged)
 					.OnFrameChanged_Raw(this, &SChaosVDSolverPlaybackControls::OnStepSelectionUpdated)
 					.MaxFrames(0)
@@ -105,6 +104,10 @@ void SChaosVDSolverPlaybackControls::HandlePlaybackControllerDataUpdated(TWeakPt
 		//TODO: This will show steps 0/0 if only one step is recorded, we need to add a way to override that functionality
 		// or just set the slider to start from 1 and handle the offset later 
 		StepsTimelineWidget->UpdateMinMaxValue(0,AvailableSteps != INDEX_NONE ? AvailableSteps -1 : 0);
+
+		// On Live Sessions, only the Game Frames timeline controls are allowed for now
+		FramesTimelineWidget->SetIsLocked(ControllerSharedPtr->IsPlayingLiveSession());
+		StepsTimelineWidget->SetIsLocked(ControllerSharedPtr->IsPlayingLiveSession());
 	}
 	else
 	{
@@ -149,19 +152,24 @@ void SChaosVDSolverPlaybackControls::UpdateStepsWidgetForFrame(const FChaosVDPla
 
 void SChaosVDSolverPlaybackControls::HandleControllerTrackFrameUpdated(TWeakPtr<FChaosVDPlaybackController> InController, const FChaosVDTrackInfo* UpdatedTrackInfo, FGuid InstigatorGuid)
 {
-	if (InstigatorGuid == GetInstigatorID())
-	{
-		// Ignore the update if we initiated it
-		return;
-	}
-
 	if (const TSharedPtr<FChaosVDPlaybackController> CurrentPlaybackControllerPtr = InController.Pin())
 	{
 		if (const FChaosVDTrackInfo* SolverTrackInfo = CurrentPlaybackControllerPtr->GetTrackInfo(EChaosVDTrackType::Solver, SolverID))
 		{
-			FramesTimelineWidget->SetCurrentTimelineFrame(SolverTrackInfo->CurrentFrame, EChaosVDSetTimelineFrameFlags::None);
+			if (InstigatorGuid != GetInstigatorID())
+			{
+				// No Need to manually update the widget state if the widget it-self instigated the update 
+				FramesTimelineWidget->SetCurrentTimelineFrame(SolverTrackInfo->CurrentFrame, EChaosVDSetTimelineFrameFlags::None);
+				UpdateStepsWidgetForFrame(*CurrentPlaybackControllerPtr.Get(), SolverTrackInfo->CurrentFrame, SolverTrackInfo->CurrentStep);
+			}
 
-			UpdateStepsWidgetForFrame(*CurrentPlaybackControllerPtr.Get(), SolverTrackInfo->CurrentFrame, SolverTrackInfo->CurrentStep);
+			if (const TSharedPtr<FChaosVDRecording> RecordingData = CurrentPlaybackControllerPtr->GetCurrentRecording().Pin())
+			{
+				if (const FChaosVDSolverFrameData* FrameData = RecordingData->GetSolverFrameData_AssumesLocked(SolverID, SolverTrackInfo->CurrentFrame))
+				{
+					FramesTimelineWidget->SetTargetFrameTime(FrameData->GetFrameTime());
+				}	
+			}
 		}
 	}
 }
diff --git a/Engine/Plugins/Experimental/ChaosVD/Source/ChaosVD/Private/Widgets/SChaosVDTimelineWidget.cpp b/Engine/Plugins/Experimental/ChaosVD/Source/ChaosVD/Private/Widgets/SChaosVDTimelineWidget.cpp
--- a/Engine/Plugins/Experimental/ChaosVD/Source/ChaosVD/Private/Widgets/SChaosVDTimelineWidget.cpp
+++ b/Engine/Plugins/Experimental/ChaosVD/Source/ChaosVD/Private/Widgets/SChaosVDTimelineWidget.cpp
@@ -17,6 +17,9 @@ void SChaosVDTimelineWidget::Construct(const FArguments& InArgs)
 	MaxFrames = InArgs._MaxFrames;
 	FrameChangedDelegate = InArgs._OnFrameChanged;
 	FrameLockedDelegate = InArgs._OnFrameLockStateChanged;
+	ButtonClickedDelegate = InArgs._OnButtonClicked;
+	bAutoStopEnabled = InArgs._AutoStopEnabled;
+	ElementVisibilityFlags = InArgs._ButtonVisibilityFlags;
 
 	SetCanTick(false);
 
@@ -32,9 +35,9 @@ void SChaosVDTimelineWidget::Construct(const FArguments& InArgs)
 				+SHorizontalBox::Slot()
 				[
 					SNew(SButton)
-					.Visibility(InArgs._HidePlayStopButtons.Get() ? EVisibility::Collapsed : EVisibility::Visible)
-					.IsEnabled_Raw(this, &SChaosVDTimelineWidget::IsUnlocked)
-					.OnClicked( FOnClicked::CreateRaw(this, &SChaosVDTimelineWidget::Play))
+					.Visibility_Raw(this, &SChaosVDTimelineWidget::GetElementVisibility, EChaosVDTimelineElementIDFlags::Play)
+					.IsEnabled_Raw(this, &SChaosVDTimelineWidget::GetElementEnabled, EChaosVDTimelineElementIDFlags::Play)
+					.OnClicked( FOnClicked::CreateRaw(this, &SChaosVDTimelineWidget::TogglePlay))
 					.ContentPadding( 2.0f )
 					.ForegroundColor( FSlateColor::UseForeground() )
 					.IsFocusable( false )
@@ -48,8 +51,8 @@ void SChaosVDTimelineWidget::Construct(const FArguments& InArgs)
 				+SHorizontalBox::Slot()
 				[
 					SNew(SButton)
-					.Visibility(InArgs._HidePlayStopButtons.Get() ? EVisibility::Collapsed : EVisibility::Visible)
-					.IsEnabled_Raw(this, &SChaosVDTimelineWidget::IsUnlocked)
+					.Visibility_Raw(this, &SChaosVDTimelineWidget::GetElementVisibility, EChaosVDTimelineElementIDFlags::Stop)
+					.IsEnabled_Raw(this, &SChaosVDTimelineWidget::GetElementEnabled, EChaosVDTimelineElementIDFlags::Stop)
 					.OnClicked( FOnClicked::CreateRaw(this, &SChaosVDTimelineWidget::Stop))
 					.ContentPadding( 2.0f )
 					.ForegroundColor( FSlateColor::UseForeground() )
@@ -64,8 +67,8 @@ void SChaosVDTimelineWidget::Construct(const FArguments& InArgs)
 				+SHorizontalBox::Slot()
 				[
 					SNew(SButton)
-					.Visibility(InArgs._HideNextPrevButtons.Get() ? EVisibility::Collapsed : EVisibility::Visible)
-					.IsEnabled_Raw(this, &SChaosVDTimelineWidget::IsUnlocked)
+					.Visibility_Raw(this, &SChaosVDTimelineWidget::GetElementVisibility, EChaosVDTimelineElementIDFlags::Prev)
+					.IsEnabled_Raw(this, &SChaosVDTimelineWidget::GetElementEnabled, EChaosVDTimelineElementIDFlags::Prev)
 					.OnClicked( FOnClicked::CreateRaw(this, &SChaosVDTimelineWidget::Prev))
 					.ContentPadding( 2.0f )
 					.ForegroundColor( FSlateColor::UseForeground() )
@@ -80,8 +83,8 @@ void SChaosVDTimelineWidget::Construct(const FArguments& InArgs)
 				+SHorizontalBox::Slot()
 				[
 					SNew(SButton)
-					.Visibility(InArgs._HideNextPrevButtons.Get() ? EVisibility::Collapsed : EVisibility::Visible)
-					.IsEnabled_Raw(this, &SChaosVDTimelineWidget::IsUnlocked)
+					.Visibility_Raw(this, &SChaosVDTimelineWidget::GetElementVisibility, EChaosVDTimelineElementIDFlags::Next)
+					.IsEnabled_Raw(this, &SChaosVDTimelineWidget::GetElementEnabled, EChaosVDTimelineElementIDFlags::Next)
 					.OnClicked( FOnClicked::CreateRaw(this, &SChaosVDTimelineWidget::Next))
 					.ContentPadding( 2.0f )
 					.ForegroundColor( FSlateColor::UseForeground() )
@@ -96,7 +99,8 @@ void SChaosVDTimelineWidget::Construct(const FArguments& InArgs)
 				+SHorizontalBox::Slot()
 				[
 					SNew(SButton)
-					.Visibility(InArgs._HideLockButton.Get() ? EVisibility::Collapsed : EVisibility::Visible)
+					.Visibility_Raw(this, &SChaosVDTimelineWidget::GetElementVisibility, EChaosVDTimelineElementIDFlags::Lock)
+					.IsEnabled_Raw(this, &SChaosVDTimelineWidget::GetElementEnabled, EChaosVDTimelineElementIDFlags::Lock)
 					.OnClicked( FOnClicked::CreateRaw(this, &SChaosVDTimelineWidget::ToggleLockState))
 					.ContentPadding( 2.0f )
 					.ForegroundColor( FSlateColor::UseForeground())
@@ -114,6 +118,9 @@ void SChaosVDTimelineWidget::Construct(const FArguments& InArgs)
 			.FillWidth(0.65f)
 			[
 			  SAssignNew(TimelineSlider, SSlider)
+			  .Visibility_Raw(this, &SChaosVDTimelineWidget::GetElementVisibility, EChaosVDTimelineElementIDFlags::Timeline)
+			  //TODO: Enable locking when we have custom images for the slider elements. Currently locking it messes up with the transparency/images
+			  //.Locked_Raw(this, &SChaosVDTimelineWidget::GetElementEnabled, EChaosVDTimelineElementIDFlags::Timeline)
 			  .ToolTipText_Lambda([this]()-> FText{ return FText::AsNumber(CurrentFrame); })
 			  .Value(CurrentFrame)
 			  .OnValueChanged_Raw(this, &SChaosVDTimelineWidget::SetCurrentTimelineFrame, EChaosVDSetTimelineFrameFlags::BroadcastChange)
@@ -145,12 +152,25 @@ void SChaosVDTimelineWidget::UpdateMinMaxValue(float NewMin, float NewMax)
 	MinFrames = NewMin;
 	MaxFrames = NewMax;
 	
-	if(CurrentFrame < NewMin || CurrentFrame > NewMax)
+	if (CurrentFrame < NewMin || CurrentFrame > NewMax)
 	{
 		CurrentFrame = NewMin;
 	}
 }
 
+void SChaosVDTimelineWidget::SetTargetFrameTime(float TargetFrameTimeSeconds)
+{
+	if (TargetFrameTimeSeconds > 0)
+	{
+		CurrentPlaybackRate = TargetFrameTimeSeconds; 
+	}
+	else
+	{
+		// Default to 60 FPS
+		CurrentPlaybackRate = 1 / 60.0f;
+	}
+}
+
 void SChaosVDTimelineWidget::ResetTimeline()
 {
 	TimelineSlider->SetValue(MinFrames);
@@ -171,22 +191,42 @@ void SChaosVDTimelineWidget::SetCurrentTimelineFrame(float FrameNumber, EChaosVD
 	}
 }
 
-FReply SChaosVDTimelineWidget::Play()
+void SChaosVDTimelineWidget::SetIsLocked(bool NewIsLocked)
+{
+	bIsLocked = NewIsLocked;
+
+	if (bIsLocked)
+	{
+		ElementEnabledFlags = static_cast<uint16>(EChaosVDTimelineElementIDFlags::Lock);	
+	}
+	else
+	{
+		ElementEnabledFlags = DefaultEnabledElementsFlags;
+	}
+}
+
+FReply SChaosVDTimelineWidget::TogglePlay()
 {
 	if (bIsPlaying)
 	{
-		bIsPlaying = false;
-		SetCanTick(false);
+		Pause();
+		ButtonClickedDelegate.ExecuteIfBound(EChaosVDPlaybackButtonsID::Pause);
 	}
 	else
 	{
-		bIsPlaying = true;
-		SetCanTick(true);
+		Play();
+		ButtonClickedDelegate.ExecuteIfBound(EChaosVDPlaybackButtonsID::Play);
 	}
 
 	return FReply::Handled();
 }
 
+void SChaosVDTimelineWidget::Play()
+{
+	bIsPlaying = true;
+	SetCanTick(true);
+}
+
 FReply SChaosVDTimelineWidget::Stop()
 {
 	CurrentFrame = 0;
@@ -197,9 +237,17 @@ FReply SChaosVDTimelineWidget::Stop()
 
 	SetCanTick(false);
 
+	ButtonClickedDelegate.ExecuteIfBound(EChaosVDPlaybackButtonsID::Stop);
+
 	return FReply::Handled();
 }
 
+void SChaosVDTimelineWidget::Pause()
+{
+	bIsPlaying = false;
+	SetCanTick(false);
+}
+
 FReply SChaosVDTimelineWidget::Next()
 {
 	if (CurrentFrame >= MaxFrames)
@@ -212,6 +260,8 @@ FReply SChaosVDTimelineWidget::Next()
 	
 	SetCurrentTimelineFrame(CurrentFrame);
 
+	ButtonClickedDelegate.ExecuteIfBound(EChaosVDPlaybackButtonsID::Next);
+
 	return FReply::Handled();
 }
 
@@ -225,13 +275,15 @@ FReply SChaosVDTimelineWidget::Prev()
 	CurrentFrame--;
 
 	SetCurrentTimelineFrame(CurrentFrame);
+	
+	ButtonClickedDelegate.ExecuteIfBound(EChaosVDPlaybackButtonsID::Prev);
 
 	return FReply::Handled();
 }
 
 FReply SChaosVDTimelineWidget::ToggleLockState()
 {
-	bIsLocked = !bIsLocked;
+	SetIsLocked(!bIsLocked);
 
 	FrameLockedDelegate.ExecuteIfBound(bIsLocked);
 
@@ -248,25 +300,37 @@ const FSlateBrush* SChaosVDTimelineWidget::GetLockStateIcon() const
 	return bIsLocked ? FChaosVDStyle::Get().GetBrush("LockIcon") : FChaosVDStyle::Get().GetBrush("UnlockedIcon");
 }
 
+EVisibility SChaosVDTimelineWidget::GetElementVisibility(EChaosVDTimelineElementIDFlags ElementID) const
+{
+	 return ((static_cast<uint16>(ElementID) & ElementVisibilityFlags) == static_cast<uint16>(ElementID)) ? EVisibility::Visible : EVisibility::Collapsed;
+}
+
+bool SChaosVDTimelineWidget::GetElementEnabled(EChaosVDTimelineElementIDFlags ElementID) const
+{
+	return ((static_cast<uint16>(ElementID) & ElementEnabledFlags) == static_cast<uint16>(ElementID));
+}
+
 void SChaosVDTimelineWidget::Tick(const FGeometry& AllottedGeometry, const double InCurrentTime, const float InDeltaTime)
 {
 	SCompoundWidget::Tick(AllottedGeometry, InCurrentTime, InDeltaTime);
-	
+
+	//TODO: We should move the Ticking logic to advance frames outside of this widget.
+	// The logic to update the visual state is already controlled externally
 	if (bIsPlaying)
 	{
 		if (CurrentFrame == MaxFrames)
 		{
-			Stop();
+			if (bAutoStopEnabled)
+			{
+				Stop();
+			}
 		}
 
-		//TODO: Allow to customize the playback frametimes so it can be made an option or read from the recorded file
-		constexpr float PlaybackFrameTime = 0.016f;
-		
 		CurrentPlaybackTime += InDeltaTime;
 
-		if (CurrentPlaybackTime > PlaybackFrameTime)
+		while (CurrentPlaybackTime > CurrentPlaybackRate)
 		{
-			CurrentPlaybackTime = 0.0f;
+			CurrentPlaybackTime -= CurrentPlaybackRate;
 			Next();
 		}
 	}
diff --git a/Engine/Plugins/Experimental/ChaosVD/Source/ChaosVD/Public/ChaosVDEngine.h b/Engine/Plugins/Experimental/ChaosVD/Source/ChaosVD/Public/ChaosVDEngine.h
--- a/Engine/Plugins/Experimental/ChaosVD/Source/ChaosVD/Public/ChaosVDEngine.h
+++ b/Engine/Plugins/Experimental/ChaosVD/Source/ChaosVD/Public/ChaosVDEngine.h
@@ -5,6 +5,7 @@
 #include "Templates/SharedPointer.h"
 #include "Containers/Ticker.h"
 #include "Misc/Guid.h"
+#include "Trace/ChaosVDTraceManager.h"
 
 class FChaosVDTraceManager;
 class FChaosVDPlaybackController;
@@ -30,13 +31,17 @@ public:
 	TSharedPtr<FChaosVDScene>& GetCurrentScene() { return CurrentScene; };
 	TSharedPtr<FChaosVDPlaybackController>& GetPlaybackController() { return PlaybackController; };
 
+	const FChaosVDTraceSessionDescriptor& GetCurrentSessionDescriptor() { return CurrentSessionDescriptor; };
+
 	void LoadRecording(const FString& FilePath);
+
+	void SetCurrentSession(const FChaosVDTraceSessionDescriptor& SessionDescriptor);
 	
 private:
 
 	FGuid InstanceGUID;
 
-	FString CurrentSessionName;
+	FChaosVDTraceSessionDescriptor CurrentSessionDescriptor;
 
 	TSharedPtr<FChaosVDScene> CurrentScene;
 	TSharedPtr<FChaosVDPlaybackController> PlaybackController;
diff --git a/Engine/Plugins/Experimental/ChaosVD/Source/ChaosVD/Public/ChaosVDGeometryDataComponent.h b/Engine/Plugins/Experimental/ChaosVD/Source/ChaosVD/Public/ChaosVDGeometryDataComponent.h
--- a/Engine/Plugins/Experimental/ChaosVD/Source/ChaosVD/Public/ChaosVDGeometryDataComponent.h
+++ b/Engine/Plugins/Experimental/ChaosVD/Source/ChaosVD/Public/ChaosVDGeometryDataComponent.h
@@ -74,10 +74,14 @@ protected:
 	void UpdateVisibility_Internal(const FChaosVDShapeCollisionData& InCollisionData, UMeshComponent* MeshComponent);
 	
 	void UpdateDataFromShapeArray_Internal(const TArray<FChaosVDShapeCollisionData>& InShapeArray, FChaosVDShapeCollisionData& CollisionDataToUpdate);
+
+	uint32 GetCachedGeometryHash(const Chaos::FImplicitObject* InImplicitObject);
 	
 	uint32 GeometryID = 0;
 	bool bIsMeshReady = false;
 
 	FChaosVDMeshReadyDelegate MeshReadyDelegate;
+
+	TMap<const Chaos::FImplicitObject*, uint32> CachedGeometryHashes;
 	Chaos::FConstImplicitObjectPtr RootImplicitObject = nullptr;
 };
diff --git a/Engine/Plugins/Experimental/ChaosVD/Source/ChaosVD/Public/ChaosVDPlaybackController.h b/Engine/Plugins/Experimental/ChaosVD/Source/ChaosVD/Public/ChaosVDPlaybackController.h
--- a/Engine/Plugins/Experimental/ChaosVD/Source/ChaosVD/Public/ChaosVDPlaybackController.h
+++ b/Engine/Plugins/Experimental/ChaosVD/Source/ChaosVD/Public/ChaosVDPlaybackController.h
@@ -10,6 +10,7 @@
 #include "Delegates/Delegate.h"
 #include "HAL/ThreadSafeBool.h"
 
+struct FChaosVDTraceSessionDescriptor;
 struct FChaosVDTrackInfo;
 class FChaosVDScene;
 struct FChaosVDRecording;
@@ -67,8 +68,8 @@ public:
 	FChaosVDPlaybackController(const TWeakPtr<FChaosVDScene>& InSceneToControl);
 	virtual ~FChaosVDPlaybackController() override;
 
-	/** Loads a recording using a Trace Session Name */
-	bool LoadChaosVDRecordingFromTraceSession(const FString& InSessionName);
+	/** Loads a recording using a CVD Trace Session Descriptor */
+	bool LoadChaosVDRecordingFromTraceSession(const FChaosVDTraceSessionDescriptor& InSessionDescriptor);
 
 	/** Unloads the currently loaded recording
 	 * @param UnloadOptions Options flags to change the steps performed during the unload
@@ -206,6 +207,16 @@ public:
 	
 	virtual bool Tick(float DeltaTime) override;
 
+	/** Returns true if we are playing a live debugging session */
+	bool IsPlayingLiveSession() const;
+
+	/** Updates the loaded recording state to indicate is not longer receiving live updates */
+	void HandleDisconnectedFromSession();
+
+	void RequestPause() { bPauseRequested = true; }
+	void RequestUnpause() { bPauseRequested = false; }
+	bool HasPauseRequest() const { return bPauseRequested; }
+
 protected:
 
 	/** Updates (or adds) solvers data from the loaded recording to the solver tracks */
@@ -242,10 +253,17 @@ protected:
 	FChaosVDPlaybackControllerFrameUpdated ControllerFrameUpdatedDelegate;
 
 	/** Set to true when the recording data controlled by this Playback Controller is updated, the update delegate will be called on the GT */
-	FThreadSafeBool bHasPendingGTUpdateBroadcast;
+	std::atomic<bool> bHasPendingGTUpdateBroadcast;
 
 	/** Queue with a copy of all Track Info Updates that needs to be done in the Game thread */
 	TQueue<FChaosVDQueuedTrackInfoUpdate, EQueueMode::Mpsc> TrackInfoUpdateGTQueue;
 
-	bool bPlayedFirsFrame = false;
+	bool bPlayedFirstFrame = false;
+
+	int32 MaxFramesLaggingBehindDuringLiveSession = 50;
+	int32 MinFramesLaggingBehindDuringLiveSession = 5;
+
+	bool bPauseRequested = false;
+
+	FDelegateHandle RecordingStoppedHandle;
 };
diff --git a/Engine/Plugins/Experimental/ChaosVD/Source/ChaosVD/Public/Trace/ChaosVDTraceManager.h b/Engine/Plugins/Experimental/ChaosVD/Source/ChaosVD/Public/Trace/ChaosVDTraceManager.h
--- a/Engine/Plugins/Experimental/ChaosVD/Source/ChaosVD/Public/Trace/ChaosVDTraceManager.h
+++ b/Engine/Plugins/Experimental/ChaosVD/Source/ChaosVD/Public/Trace/ChaosVDTraceManager.h
@@ -10,6 +10,15 @@ class FChaosVDTraceModule;
 
 namespace TraceServices { class IAnalysisSession; }
 
+/**
+ * Structure containing info about a trace session used by Chaos Visual Debugger
+ */
+struct FChaosVDTraceSessionDescriptor
+{
+	FString SessionName;
+	bool bIsLiveSession = false;
+};
+
 /** Manager class used by Chaos VD to interact/control UE Trace systems */
 class FChaosVDTraceManager
 {
@@ -17,10 +26,26 @@ public:
 	FChaosVDTraceManager();
 	~FChaosVDTraceManager();
 
+	/** Load a trace file and starts analyzing it
+	 * @param InTraceFilename File Name including (Path Included) of the Trace file to load
+	 */
 	FString LoadTraceFile(const FString& InTraceFilename);
-
+	
+	/**
+	 * Connects to a live Trace Session and starts analyzing it.
+	 * @param InSessionHost Trace Store Address for this session
+	 * @param SessionID Trace ID in the Trace Store provided as host
+	 * @return Session Name
+	 */
+	FString ConnectToLiveSession(FStringView InSessionHost, uint32 SessionID);
+
+	/** Returns the path to the local trace store */
+	FString GetLocalTraceStoreDirPath();
+
+	/** Returns a ptr to the session registered with the provided session name. Null if no session is found */
 	TSharedPtr<const TraceServices::IAnalysisSession> GetSession(const FString& InSessionName);
 
+	/** Stops and de-registers a trace session registered with the provided session name */
 	void CloseSession(const FString& InSessionName);
 
 private:
diff --git a/Engine/Plugins/Experimental/ChaosVD/Source/ChaosVD/Public/Widgets/SChaosBrowseTraceFileSourceModal.h b/Engine/Plugins/Experimental/ChaosVD/Source/ChaosVD/Public/Widgets/SChaosBrowseTraceFileSourceModal.h
new file mode 100644
--- /dev/null
+++ b/Engine/Plugins/Experimental/ChaosVD/Source/ChaosVD/Public/Widgets/SChaosBrowseTraceFileSourceModal.h
@@ -0,0 +1,40 @@
+// Copyright Epic Games, Inc. All Rights Reserved.
+
+#pragma once
+
+#include "Widgets/SCompoundWidget.h"
+#include "Widgets/SWindow.h"
+
+class FReply;
+
+enum class EChaosVDBrowseFileModalResponse
+{
+	OpenFolder,
+	OpenTraceStore,
+	Cancel
+};
+
+/**
+ * 
+ */
+class SChaosBrowseTraceFileSourceModal : public SWindow
+{
+public:
+	SLATE_BEGIN_ARGS(SChaosBrowseTraceFileSourceModal)
+		{
+		}
+
+	SLATE_END_ARGS()
+
+	/** Constructs this widget with InArgs */
+	void Construct(const FArguments& InArgs);
+
+	/** Displays the dialog in a blocking fashion */
+	EChaosVDBrowseFileModalResponse ShowModal();
+
+protected:
+
+	FReply OnButtonClick(EChaosVDBrowseFileModalResponse Response);
+
+	EChaosVDBrowseFileModalResponse UserResponse = EChaosVDBrowseFileModalResponse::Cancel;
+};
diff --git a/Engine/Plugins/Experimental/ChaosVD/Source/ChaosVD/Public/Widgets/SChaosVDBrowseSessionsModal.h b/Engine/Plugins/Experimental/ChaosVD/Source/ChaosVD/Public/Widgets/SChaosVDBrowseSessionsModal.h
new file mode 100644
--- /dev/null
+++ b/Engine/Plugins/Experimental/ChaosVD/Source/ChaosVD/Public/Widgets/SChaosVDBrowseSessionsModal.h
@@ -0,0 +1,82 @@
+// Copyright Epic Games, Inc. All Rights Reserved.
+
+#pragma once
+
+#include "Containers/Map.h"
+#include "Templates/SharedPointer.h"
+#include "Widgets/SWindow.h"
+
+class FName;
+class FText;
+class FLiveSessionTracker;
+class SChaosVDNameListPicker;
+class SEditableTextBox;
+
+namespace EAppReturnType { enum Type; }
+
+/** Structure containing info about a Trace Session */
+struct FChaosVDTraceSessionInfo
+{
+	uint32 TraceID = 0;
+	uint32 IPAddress = 0;
+	uint32 ControlPort = 0;
+	bool bIsValid = false;
+};
+
+/**
+ * Modal window used to find, show and select active Trace Sessions to be used
+ * with Chaos Visual Debugger
+ */
+class SChaosVDBrowseSessionsModal : public SWindow
+{
+public:
+	SLATE_BEGIN_ARGS(SChaosVDBrowseSessionsModal)
+		{
+		}
+	SLATE_END_ARGS()
+
+	virtual ~SChaosVDBrowseSessionsModal() override;
+
+	void Construct(const FArguments& InArgs);
+
+	/** Displays the dialog in a blocking fashion */
+	EAppReturnType::Type ShowModal();
+
+	/** Response that triggered the closing of this modal */
+	EAppReturnType::Type GetResponse() const { return UserResponse; }
+
+	/** Returns information about the active trace session selected */
+	FChaosVDTraceSessionInfo GetSelectedTraceInfo();
+
+	/** Returns the address of the trace store selected while looking for active Trace Sessions */
+	FString GetSelectedTraceStoreAddress() const { return CurrentTraceStoreAddress; };
+
+	void ModalTick(float InDeltaTime);
+
+protected:
+
+	bool CanOpenSession() const;
+	
+	void UpdateCurrentSessionInfoMap();
+
+	FReply OnButtonClick(EAppReturnType::Type ButtonID);
+	
+	void HandleSessionNameSelected(TSharedPtr<FName> SelectedName);
+	
+	void OnTraceStoreAddressUpdated(const FText& InText, ETextCommit::Type CommitType);
+
+	TSharedPtr<SChaosVDNameListPicker> NamePickerWidget;
+	
+	TSharedPtr<SEditableTextBox> TraceStoreAddressWidget;
+	
+	TMap<FName, FChaosVDTraceSessionInfo> CurrentSessionInfosMap;
+
+	EAppReturnType::Type UserResponse = EAppReturnType::Cancel;
+
+	FString CurrentTraceStoreAddress;
+
+	FName CurrentTraceSessionSelected;
+
+	FDelegateHandle ModalTickHandle;
+	float AccumulatedTimeBetweenTicks = 0.0f;
+};
diff --git a/Engine/Plugins/Experimental/ChaosVD/Source/ChaosVD/Public/Widgets/SChaosVDEditorViewportViewMenu.h b/Engine/Plugins/Experimental/ChaosVD/Source/ChaosVD/Public/Widgets/SChaosVDEditorViewportViewMenu.h
--- a/Engine/Plugins/Experimental/ChaosVD/Source/ChaosVD/Public/Widgets/SChaosVDEditorViewportViewMenu.h
+++ b/Engine/Plugins/Experimental/ChaosVD/Source/ChaosVD/Public/Widgets/SChaosVDEditorViewportViewMenu.h
@@ -11,7 +11,7 @@
  * Compared to the normal Editor View modes menu, it has a reduced set of modes available,
  * showing only the ones supported
  */
-class CHAOSVD_API SChaosVDEditorViewportViewMenu : public SEditorViewportViewMenu
+class SChaosVDEditorViewportViewMenu : public SEditorViewportViewMenu
 {
 public:
 	SLATE_BEGIN_ARGS(SChaosVDEditorViewportViewMenu){}
diff --git a/Engine/Plugins/Experimental/ChaosVD/Source/ChaosVD/Public/Widgets/SChaosVDMainTab.h b/Engine/Plugins/Experimental/ChaosVD/Source/ChaosVD/Public/Widgets/SChaosVDMainTab.h
--- a/Engine/Plugins/Experimental/ChaosVD/Source/ChaosVD/Public/Widgets/SChaosVDMainTab.h
+++ b/Engine/Plugins/Experimental/ChaosVD/Source/ChaosVD/Public/Widgets/SChaosVDMainTab.h
@@ -38,7 +38,11 @@ private:
 
 	void GenerateMainWindowMenu();
 
-	void BrowseAndOpenChaosVDFile();
+	void BrowseAndOpenChaosVDRecording();
+
+	void BrowseChaosVDRecordingFromFolder(FStringView FolderPath = TEXT(""));
+
+	void BrowseLiveSessionsFromTraceStore() const;
 
 	TSharedPtr<FChaosVDEngine> ChaosVDEngine;
 
@@ -53,4 +57,7 @@ private:
 	FName StatusBarID;
 
 	TSharedPtr<FTabManager> TabManager;
+
+	FReply HandleSessionConnectionClicked();
+	FText GetConnectButtonText() const;
 };
diff --git a/Engine/Plugins/Experimental/ChaosVD/Source/ChaosVD/Public/Widgets/SChaosVDPlaybackViewport.h b/Engine/Plugins/Experimental/ChaosVD/Source/ChaosVD/Public/Widgets/SChaosVDPlaybackViewport.h
--- a/Engine/Plugins/Experimental/ChaosVD/Source/ChaosVD/Public/Widgets/SChaosVDPlaybackViewport.h
+++ b/Engine/Plugins/Experimental/ChaosVD/Source/ChaosVD/Public/Widgets/SChaosVDPlaybackViewport.h
@@ -9,6 +9,7 @@
 #include "Templates/SharedPointer.h"
 #include "Widgets/SCompoundWidget.h"
 
+enum class EChaosVDPlaybackButtonsID : uint8;
 class FChaosVDEditorModeTools;
 enum class EChaosVDActorTrackingMode;
 class FChaosVDPlaybackViewportClient;
@@ -54,6 +55,8 @@ protected:
 
 	void OnFrameSelectionUpdated(int32 NewFrameIndex) const;
 
+	void HandlePlaybackButtonClicked(EChaosVDPlaybackButtonsID ButtonID);
+
 	TSharedPtr<SChaosVDTimelineWidget> GameFramesTimelineWidget;
 
 	TSharedPtr<FChaosVDPlaybackViewportClient> PlaybackViewportClient;
diff --git a/Engine/Plugins/Experimental/ChaosVD/Source/ChaosVD/Public/Widgets/SChaosVDRecordingControls.h b/Engine/Plugins/Experimental/ChaosVD/Source/ChaosVD/Public/Widgets/SChaosVDRecordingControls.h
--- a/Engine/Plugins/Experimental/ChaosVD/Source/ChaosVD/Public/Widgets/SChaosVDRecordingControls.h
+++ b/Engine/Plugins/Experimental/ChaosVD/Source/ChaosVD/Public/Widgets/SChaosVDRecordingControls.h
@@ -37,6 +37,7 @@ protected:
 	
 	FStatusBarMessageHandle RecordingMessageHandle;
 	FStatusBarMessageHandle RecordingPathMessageHandle;
+	FStatusBarMessageHandle LiveSessionEndedMessageHandle;
 	FDelegateHandle RecordingStartedHandle;
 	FDelegateHandle RecordingStoppedHandle;
 
diff --git a/Engine/Plugins/Experimental/ChaosVD/Source/ChaosVD/Public/Widgets/SChaosVDTimelineWidget.h b/Engine/Plugins/Experimental/ChaosVD/Source/ChaosVD/Public/Widgets/SChaosVDTimelineWidget.h
--- a/Engine/Plugins/Experimental/ChaosVD/Source/ChaosVD/Public/Widgets/SChaosVDTimelineWidget.h
+++ b/Engine/Plugins/Experimental/ChaosVD/Source/ChaosVD/Public/Widgets/SChaosVDTimelineWidget.h
@@ -9,6 +9,16 @@ class SSlider;
 class FReply;
 struct FSlateBrush;
 
+enum class EChaosVDPlaybackButtonsID : uint8
+{
+	Play,
+	Pause,
+	Stop,
+	Next,
+	Prev
+};
+
+DECLARE_DELEGATE_OneParam(FChaosControlButtonClicked, EChaosVDPlaybackButtonsID)
 DECLARE_DELEGATE_OneParam(FChaosVDFrameChangedDelegate, int32)
 DECLARE_DELEGATE_OneParam(FChaosVDFrameLockStateDelegate, bool)
 
@@ -20,6 +30,23 @@ enum class EChaosVDSetTimelineFrameFlags
 };
 ENUM_CLASS_FLAGS(EChaosVDSetTimelineFrameFlags)
 
+enum class EChaosVDTimelineElementIDFlags : uint16
+{
+	None = 0,
+	Play = 1 << 0,
+	Stop = 1 << 1,
+	Next = 1 << 2,
+	Prev = 1 << 3,
+	Lock = 1 << 4,
+	Timeline = 1 << 5,
+	
+	ManualSteppingButtons = Next | Prev,
+	AllManualStepping = Next | Prev | Timeline,
+	AllPlaybackButtons = Play | Stop | Next | Prev,
+	AllPlayback = Play | Stop | Next | Prev | Timeline,
+	All = Play | Stop | Next | Prev | Timeline | Lock,
+};
+ENUM_CLASS_FLAGS(EChaosVDTimelineElementIDFlags)
 
 /** Simple timeline control widget */
 class SChaosVDTimelineWidget : public SCompoundWidget
@@ -28,11 +55,11 @@ public:
 
 	SLATE_BEGIN_ARGS( SChaosVDTimelineWidget ){}
 		SLATE_ARGUMENT(int32, MaxFrames)
+		SLATE_ARGUMENT(uint16, ButtonVisibilityFlags)
+		SLATE_ARGUMENT(bool, AutoStopEnabled)
 		SLATE_EVENT(FChaosVDFrameChangedDelegate, OnFrameChanged)
 		SLATE_EVENT(FChaosVDFrameLockStateDelegate, OnFrameLockStateChanged)
-		SLATE_ATTRIBUTE( bool, HidePlayStopButtons)
-		SLATE_ATTRIBUTE( bool, HideNextPrevButtons)
-		SLATE_ATTRIBUTE( bool, HideLockButton)
+		SLATE_EVENT(FChaosControlButtonClicked, OnButtonClicked)
 	SLATE_END_ARGS()
 
 	void Construct(const FArguments& InArgs);
@@ -41,10 +68,12 @@ public:
 
 	void SetCurrentTimelineFrame(float FrameNumber, EChaosVDSetTimelineFrameFlags Options = EChaosVDSetTimelineFrameFlags::BroadcastChange);
 
-	void SetIsLocked(bool NewIsLocked) { bIsLocked = NewIsLocked; }
+	void SetIsLocked(bool NewIsLocked);
 
 	bool IsUnlocked() const { return !bIsLocked; }
 
+	void SetTargetFrameTime(float TargetFrameTimeSeconds);
+
 	/** Brings back the state of the timeline to its original state*/
 	void ResetTimeline();
 
@@ -54,10 +83,20 @@ public:
 	/** Called when this timeline is locked or unlocked */
 	FChaosVDFrameLockStateDelegate& OnFrameLockStateChanged() { return FrameLockedDelegate; }
 
-protected:
+	uint16& GetCurrentElementVisibilityFlags() { return ElementVisibilityFlags; }
+	void SetCurrentElementVisibilityFlags(uint16 NewVisibilityFlags) { ElementVisibilityFlags = NewVisibilityFlags; }
 
-	FReply  Play();
+	uint16& GetMutableElementEnabledFlagsRef() { return ElementEnabledFlags; }
+
+	void SetAutoStopEnabled(bool bNewEnabled) { bAutoStopEnabled = bNewEnabled; }
+
+	void Play();
 	FReply  Stop();
+
+protected:
+
+	void Pause();
+	FReply  TogglePlay();
 	FReply  Next();
 	FReply  Prev();
 
@@ -66,6 +105,9 @@ protected:
 	const FSlateBrush* GetPlayOrPauseIcon() const;
 	const FSlateBrush* GetLockStateIcon() const;
 
+	EVisibility GetElementVisibility(EChaosVDTimelineElementIDFlags ElementID) const;
+	bool GetElementEnabled(EChaosVDTimelineElementIDFlags ElementID) const;
+
 	virtual void Tick(const FGeometry& AllottedGeometry, const double InCurrentTime, const float InDeltaTime) override;
 
 	TSharedPtr<SSlider> TimelineSlider;
@@ -76,9 +118,17 @@ protected:
 
 	bool bIsPlaying = false;
 	float CurrentPlaybackTime = 0.0f;
+	float CurrentPlaybackRate = 1.0f / 60.0f;
 
 	FChaosVDFrameChangedDelegate FrameChangedDelegate;
 	FChaosVDFrameLockStateDelegate FrameLockedDelegate;
+	FChaosControlButtonClicked ButtonClickedDelegate;
 
 	bool bIsLocked = false;
+	bool bAutoStopEnabled = false;
+
+	uint16 ElementVisibilityFlags = 0;
+	uint16 ElementEnabledFlags = static_cast<uint16>(EChaosVDTimelineElementIDFlags::All);
+	uint16 DefaultEnabledElementsFlags = static_cast<uint16>(EChaosVDTimelineElementIDFlags::All);
+	
 };
diff --git a/Engine/Source/Runtime/Experimental/ChaosVDData/Public/ChaosVDRecording.h b/Engine/Source/Runtime/Experimental/ChaosVDData/Public/ChaosVDRecording.h
--- a/Engine/Source/Runtime/Experimental/ChaosVDData/Public/ChaosVDRecording.h
+++ b/Engine/Source/Runtime/Experimental/ChaosVDData/Public/ChaosVDRecording.h
@@ -47,12 +47,42 @@ struct CHAOSVDDATA_API FChaosVDSolverFrameData
 	bool bIsKeyFrame = false;
 	FChaosVDStepsContainer SolverSteps;
 	TSet<int32> ParticlesDestroyedIDs;
+	double StartTime = -1.0;
+	double EndTime = -1.0;
+
+	/** Calculates and returns the frame time for this recorded frame.
+	 * @return Calculated frame time. -1 if it was not recorded
+	 */
+	double GetFrameTime() const
+	{
+		if (StartTime < 0 || EndTime < 0)
+		{
+			return -1.0;
+		}
+
+		return EndTime - StartTime;
+	}
 };
 
 struct FChaosVDGameFrameData
 {
 	uint64 FirstCycle;
 	uint64 LastCycle;
+	double StartTime = -1.0;
+	double EndTime = -1.0;
+
+	/** Calculates and returns the frame time for this recorded frame.
+	 * @return Calculated frame time. -1 if it was not recorded
+	 */
+	double GetFrameTime() const
+	{
+		if (StartTime < 0 || EndTime < 0)
+		{
+			return -1.0;
+		}
+
+		return EndTime - StartTime;
+	}
 
 	TMap<FName, FChaosVDTrackedLocation> RecordedNonSolverLocationsByID;
 	TMap<FName, FChaosVDTrackedTransform> RecordedNonSolverTransformsByID;
@@ -206,6 +236,12 @@ struct CHAOSVDDATA_API FChaosVDRecording
 
 	FRWLock& GetRecordingDataLock() { return RecordingDataLock; }
 
+	/** Returns true if this recording is being populated from a live session */
+	bool IsLive() const { return bIsLive; }
+
+	/** Sets if this recording is being populated from a live session */
+	void SetIsLive(bool bNewIsLive) { bIsLive = bNewIsLive; }
+
 protected:
 
 	/** Adds an Implicit Object to the recording and takes ownership of it */
@@ -229,6 +265,9 @@ protected:
 
 	mutable FRWLock RecordingDataLock;
 
+	/** True if this recording is being populated from a live session */
+	bool bIsLive = false;
+
 	friend class FChaosVDTraceProvider;
 	friend class FChaosVDTraceImplicitObjectProcessor;
 };
diff --git a/Engine/Source/Runtime/Experimental/ChaosVisualDebugger/Private/ChaosVDRuntimeModule.cpp b/Engine/Source/Runtime/Experimental/ChaosVisualDebugger/Private/ChaosVDRuntimeModule.cpp
--- a/Engine/Source/Runtime/Experimental/ChaosVisualDebugger/Private/ChaosVDRuntimeModule.cpp
+++ b/Engine/Source/Runtime/Experimental/ChaosVisualDebugger/Private/ChaosVDRuntimeModule.cpp
@@ -46,6 +46,11 @@ FChaosVDRuntimeModule& FChaosVDRuntimeModule::Get()
 	return FModuleManager::Get().LoadModuleChecked<FChaosVDRuntimeModule>(TEXT("ChaosVDRuntime"));
 }
 
+bool FChaosVDRuntimeModule::IsLoaded()
+{
+	return FModuleManager::Get().IsModuleLoaded(TEXT("ChaosVDRuntime"));
+}
+
 void FChaosVDRuntimeModule::StartupModule()
 {
 #if UE_TRACE_ENABLED
@@ -159,9 +164,14 @@ void FChaosVDRuntimeModule::StartRecording(const TArray<FString>& Args)
 	}
 	else if(Args[0] == TEXT("Server"))
 	{
+		if (FTraceAuxiliary::IsConnected())
+		{
+			FTraceAuxiliary::Stop();
+		}
+
 		const FString Target = Args.IsValidIndex(1) ? Args[1] : TEXT("127.0.0.1");
 
-		FTraceAuxiliary::Start(
+		bIsRecording = FTraceAuxiliary::Start(
 		FTraceAuxiliary::EConnectionType::Network,
 		*Target,
 		nullptr);
@@ -211,6 +221,8 @@ void FChaosVDRuntimeModule::HandleTraceStopRequest(FTraceAuxiliary::EConnectionT
 		{
 			UE_LOG(LogChaosVDRuntime, Warning, TEXT("Trace Recording has been stopped unexpectedly"));
 		}
+
+		StopRecording();
 	}
 
 	bRequestedStop = false;
diff --git a/Engine/Source/Runtime/Experimental/ChaosVisualDebugger/Public/ChaosVDRuntimeModule.h b/Engine/Source/Runtime/Experimental/ChaosVisualDebugger/Public/ChaosVDRuntimeModule.h
--- a/Engine/Source/Runtime/Experimental/ChaosVisualDebugger/Public/ChaosVDRuntimeModule.h
+++ b/Engine/Source/Runtime/Experimental/ChaosVisualDebugger/Public/ChaosVDRuntimeModule.h
@@ -26,6 +26,7 @@ class CHAOSVDRUNTIME_API FChaosVDRuntimeModule : public IModuleInterface
 public:
 
 	static FChaosVDRuntimeModule& Get();
+	static bool IsLoaded();
 	
 	virtual void StartupModule() override;
 	virtual void ShutdownModule() override;
