[ChaosVD] Fix crash on shutdown if the editor is closed while the CVD tool is open

diff --git a/Engine/Plugins/Experimental/ChaosVD/Source/ChaosVD/Private/ChaosVDEngine.cpp b/Engine/Plugins/Experimental/ChaosVD/Source/ChaosVD/Private/ChaosVDEngine.cpp
--- a/Engine/Plugins/Experimental/ChaosVD/Source/ChaosVD/Private/ChaosVDEngine.cpp
+++ b/Engine/Plugins/Experimental/ChaosVD/Source/ChaosVD/Private/ChaosVDEngine.cpp
@@ -41,6 +41,8 @@ void FChaosVDEngine::DeInitialize()
 	}
 
 	CollectGarbage(GARBAGE_COLLECTION_KEEPFLAGS);
+
+	bIsInitialized = false;
 }
 
 void FChaosVDEngine::LoadRecording(const FString& FilePath)
diff --git a/Engine/Plugins/Experimental/ChaosVD/Source/ChaosVD/Private/ChaosVDModule.cpp b/Engine/Plugins/Experimental/ChaosVD/Source/ChaosVD/Private/ChaosVDModule.cpp
--- a/Engine/Plugins/Experimental/ChaosVD/Source/ChaosVD/Private/ChaosVDModule.cpp
+++ b/Engine/Plugins/Experimental/ChaosVD/Source/ChaosVD/Private/ChaosVDModule.cpp
@@ -39,6 +39,8 @@ void FChaosVDModule::StartupModule()
 								.SetGroup(WorkspaceMenu::GetMenuStructure().GetDeveloperToolsDebugCategory());
 
 	ChaosVDTraceManager = MakeShared<FChaosVDTraceManager>();
+
+	FCoreDelegates::OnEnginePreExit.AddRaw(this, &FChaosVDModule::CloseActiveInstances);
 }
 
 void FChaosVDModule::ShutdownModule()
@@ -46,6 +48,10 @@ void FChaosVDModule::ShutdownModule()
 	FChaosVDStyle::Shutdown();
 
 	FGlobalTabmanager::Get()->UnregisterNomadTabSpawner(FChaosVDTabID::ChaosVisualDebuggerTab);
+	
+	FCoreDelegates::OnEnginePreExit.RemoveAll(this);
+
+	CloseActiveInstances();
 }
 
 void FChaosVDModule::RegisterClassesCustomDetails() const
@@ -79,31 +85,48 @@ TSharedRef<SDockTab> FChaosVDModule::SpawnMainTab(const FSpawnTabArgs& Args)
 	MainTabInstance->SetTabIcon(FChaosVDStyle::Get().GetBrush("TabIconPlaybackViewport"));
 
 	const FGuid InstanceGuid = ChaosVDEngineInstance->GetInstanceGuid();
-	RegisterChaosVDInstance(InstanceGuid, ChaosVDEngineInstance);
+	RegisterChaosVDEngineInstance(InstanceGuid, ChaosVDEngineInstance);
+
+	const SDockTab::FOnTabClosedCallback ClosedCallback = SDockTab::FOnTabClosedCallback::CreateRaw(this, &FChaosVDModule::HandleTabClosed, InstanceGuid);
+	MainTabInstance->SetOnTabClosed(ClosedCallback);
+
+	RegisterChaosVDTabInstance(InstanceGuid, MainTabInstance.ToSharedPtr());
+
+	return MainTabInstance;
+}
 
+void FChaosVDModule::HandleTabClosed(TSharedRef<SDockTab> ClosedTab, FGuid InstanceGUID)
+{
 	// Workaround. Currently the ChaosVD Engine instance determines the lifetime of the Editor world and other objects
 	// Some widgets, like UE Level viewport tries to iterate on these objects on destruction
 	// For now we can avoid any crashes by just de-initializing ChaosVD Engine on the next frame but that is not the real fix.
+	// Unless we are shutting down the engine
 	
 	//TODO: Ensure that systems that uses the Editor World we create know beforehand when it is about to be Destroyed and GC'd
-	MainTabInstance->SetOnTabClosed(SDockTab::FOnTabClosedCallback::CreateLambda( [this, InstanceGuid](TSharedRef<SDockTab>)
+	// Related Jira Task UE-191876
+	if (bIsShuttingDown)
+	{
+		DeregisterChaosVDTabInstance(InstanceGUID);
+		DeregisterChaosVDEngineInstance(InstanceGUID);
+	}
+	else
 	{
-		FTSTicker::GetCoreTicker().AddTicker(FTickerDelegate::CreateLambda([this, InstanceGuid](float DeltaTime)->bool
+		DeregisterChaosVDTabInstance(InstanceGUID);
+
+		FTSTicker::GetCoreTicker().AddTicker(FTickerDelegate::CreateLambda([this, InstanceGUID](float DeltaTime)->bool
 		{
-			DeregisterChaosVDInstance(InstanceGuid);
+			DeregisterChaosVDEngineInstance(InstanceGUID);
 			return false;
 		}));
-	}));
-
-	return MainTabInstance;
+	}
 }
 
-void FChaosVDModule::RegisterChaosVDInstance(const FGuid& InstanceGuid, TSharedPtr<FChaosVDEngine> Instance)
+void FChaosVDModule::RegisterChaosVDEngineInstance(const FGuid& InstanceGuid, TSharedPtr<FChaosVDEngine> Instance)
 {
 	ActiveChaosVDInstances.Add(InstanceGuid, Instance);
 }
 
-void FChaosVDModule::DeregisterChaosVDInstance(const FGuid& InstanceGuid)
+void FChaosVDModule::DeregisterChaosVDEngineInstance(const FGuid& InstanceGuid)
 {
 	if (TSharedPtr<FChaosVDEngine>* InstancePtrPtr = ActiveChaosVDInstances.Find(InstanceGuid))
 	{
@@ -116,6 +139,37 @@ void FChaosVDModule::DeregisterChaosVDInstance(const FGuid& InstanceGuid)
 	}	
 }
 
+void FChaosVDModule::RegisterChaosVDTabInstance(const FGuid& InstanceGuid, TSharedPtr<SDockTab> Instance)
+{
+	ActiveCVDTabs.Add(InstanceGuid, Instance);
+}
+
+void FChaosVDModule::DeregisterChaosVDTabInstance(const FGuid& InstanceGuid)
+{
+	ActiveCVDTabs.Remove(InstanceGuid);
+}
+
+void FChaosVDModule::CloseActiveInstances()
+{
+	bIsShuttingDown = true;
+	for (const TPair<FGuid, TWeakPtr<SDockTab>>& CVDTabWithID : ActiveCVDTabs)
+	{
+		if (TSharedPtr<SDockTab> CVDTab = CVDTabWithID.Value.Pin())
+		{
+			CVDTab->RequestCloseTab();
+		}
+		else
+		{
+			// if the tab Instance no longer exist, make sure the CVD engine instance is shutdown
+			DeregisterChaosVDEngineInstance(CVDTabWithID.Key);
+		}
+	}
+
+	ActiveChaosVDInstances.Reset();
+	ActiveCVDTabs.Reset();
+}
+
+
 #undef LOCTEXT_NAMESPACE
 	
 IMPLEMENT_MODULE(FChaosVDModule, ChaosVD)
\ No newline at end of file
diff --git a/Engine/Plugins/Experimental/ChaosVD/Source/ChaosVD/Private/Widgets/SChaosVDRecordingControls.cpp b/Engine/Plugins/Experimental/ChaosVD/Source/ChaosVD/Private/Widgets/SChaosVDRecordingControls.cpp
--- a/Engine/Plugins/Experimental/ChaosVD/Source/ChaosVD/Private/Widgets/SChaosVDRecordingControls.cpp
+++ b/Engine/Plugins/Experimental/ChaosVD/Source/ChaosVD/Private/Widgets/SChaosVDRecordingControls.cpp
@@ -79,8 +79,11 @@ void SChaosVDRecordingControls::Construct(const FArguments& InArgs, const TWeakP
 
 SChaosVDRecordingControls::~SChaosVDRecordingControls()
 {
-	 FChaosVDRuntimeModule::Get().RemoveRecordingStartedCallback(RecordingStartedHandle);
-	 FChaosVDRuntimeModule::Get().RemoveRecordingStopCallback(RecordingStoppedHandle);
+	if (FChaosVDRuntimeModule::IsLoaded())
+	{
+		FChaosVDRuntimeModule::Get().RemoveRecordingStartedCallback(RecordingStartedHandle);
+		FChaosVDRuntimeModule::Get().RemoveRecordingStopCallback(RecordingStoppedHandle);
+	}
 }
 
 const FSlateBrush* SChaosVDRecordingControls::GetRecordOrStopButton() const
diff --git a/Engine/Plugins/Experimental/ChaosVD/Source/ChaosVD/Public/ChaosVDModule.h b/Engine/Plugins/Experimental/ChaosVD/Source/ChaosVD/Public/ChaosVDModule.h
--- a/Engine/Plugins/Experimental/ChaosVD/Source/ChaosVD/Public/ChaosVDModule.h
+++ b/Engine/Plugins/Experimental/ChaosVD/Source/ChaosVD/Public/ChaosVDModule.h
@@ -31,10 +31,21 @@ private:
 
 	TSharedRef<SDockTab> SpawnMainTab(const FSpawnTabArgs& Args);
 
-	void RegisterChaosVDInstance(const FGuid& InstanceGuid, TSharedPtr<FChaosVDEngine> Instance);
-	void DeregisterChaosVDInstance(const FGuid& InstanceGuid);
+	void HandleTabClosed(TSharedRef<SDockTab> ClosedTab, FGuid InstanceGUID);
+
+	void RegisterChaosVDEngineInstance(const FGuid& InstanceGuid, TSharedPtr<FChaosVDEngine> Instance);
+	void DeregisterChaosVDEngineInstance(const FGuid& InstanceGuid);
+
+	void RegisterChaosVDTabInstance(const FGuid& InstanceGuid, TSharedPtr<SDockTab> Instance);
+	void DeregisterChaosVDTabInstance(const FGuid& InstanceGuid);
+	
+	void CloseActiveInstances();
 	
 	TMap<FGuid, TSharedPtr<FChaosVDEngine>> ActiveChaosVDInstances;
 
+	TMap<FGuid, TWeakPtr<SDockTab>> ActiveCVDTabs;
+
 	TSharedPtr<FChaosVDTraceManager> ChaosVDTraceManager;
+
+	bool bIsShuttingDown = false;
 };
