[ChaosVD] Remote Debugging performance and stability improvements

- Added loading state to frames. Now frame can be Unloaded/Loaded/Corrupt. This will be used if a frame cannot be read after multiple attempts and allow the tool decide if a frame should be skipped or should wait for it to be loaded
- Now the tool waits for a frame to be fully loaded before attempting to read it, and waits if it is not ready. Frames not being fully loaded all the time happens during remote debugging
- During scrubbing (or the automatic catch up done during live debugging), now instead of playing back all delta frames needed between keyframes, we collapse the all the delta frames data into a generated "keyframe" and play that one.
This increases performance, needed to ensure we don't fall back constantly during remote debugging playback of large maps (and makes scrubbing better)
- More granular time slicing for Geometry application after generation. Now we count each mesh component processed instead each mesh (which could be shared between multiple components) against the limit of geometry to process each tick.
- Fix to ensure geometry components are created/updated in the game thread (geometry generation is still async)
- Now we only destroy the Particle actors when the solver is not available anymore, in any other situation we just "deactivate" them, to avoid the cost of having to destroy and rebuild them during playback of large maps where a lot of level streaming was happening. The next step is implement an actor pool system as just deactivate them will not scale.

diff --git a/Engine/Plugins/Experimental/ChaosVD/Source/ChaosVD/Private/ChaosVDGeometryBuilder.cpp b/Engine/Plugins/Experimental/ChaosVD/Source/ChaosVD/Private/ChaosVDGeometryBuilder.cpp
--- a/Engine/Plugins/Experimental/ChaosVD/Source/ChaosVD/Private/ChaosVDGeometryBuilder.cpp
+++ b/Engine/Plugins/Experimental/ChaosVD/Source/ChaosVD/Private/ChaosVDGeometryBuilder.cpp
@@ -13,7 +13,7 @@
 
 static FAutoConsoleVariable CVarChaosVDGeometryToProcessPerTick(
 	TEXT("p.Chaos.VD.Tool.GeometryToProcessPerTick"),
-	50,
+	200,
 	TEXT("Number of generated geometry to process each tick when loading a teace file in the CVD tool"));
 
 void FChaosVDGeometryBuilder::AddReferencedObjects(FReferenceCollector& Collector)
@@ -65,10 +65,21 @@ bool FChaosVDGeometryBuilder::HasNegativeScale(const Chaos::FRigidTransform3& In
 	return ScaleSignVector.X * ScaleSignVector.Y * ScaleSignVector.Z < 0;
 }
 
+bool FChaosVDGeometryBuilder::HasGeometryInCache(uint32 GeometryKey)
+{
+	FReadScopeLock ReadLock(GeometryCacheRWLock);
+	return HasGeometryInCache_AssumesLocked(GeometryKey);
+}
+
+bool FChaosVDGeometryBuilder::HasGeometryInCache_AssumesLocked(uint32 GeometryKey) const
+{
+	return StaticMeshCacheMap.Contains(GeometryKey) || DynamicMeshCacheMap.Contains(GeometryKey);
+}
+
 UDynamicMesh* FChaosVDGeometryBuilder::CreateAndCacheDynamicMesh(const uint32 GeometryCacheKey, UE::Geometry::FMeshShapeGenerator& MeshGenerator)
 {
 	{
-		FReadScopeLock ReadLock(RWLock);
+		FReadScopeLock ReadLock(GeometryCacheRWLock);
 		if (TObjectPtr<UDynamicMesh>* DynamicMeshPtrPtr = DynamicMeshCacheMap.Find(GeometryCacheKey))
 		{
 			return *DynamicMeshPtrPtr;
@@ -81,7 +92,7 @@ UDynamicMesh* FChaosVDGeometryBuilder::CreateAndCacheDynamicMesh(const uint32 Ge
 	Mesh->SetMesh(&MeshGenerator.Generate());
 
 	{
-		FWriteScopeLock WriteLock(RWLock);
+		FWriteScopeLock WriteLock(GeometryCacheRWLock);
 		DynamicMeshCacheMap.Add(GeometryCacheKey, Mesh);
 	}
 
@@ -91,7 +102,7 @@ UDynamicMesh* FChaosVDGeometryBuilder::CreateAndCacheDynamicMesh(const uint32 Ge
 UStaticMesh* FChaosVDGeometryBuilder::CreateAndCacheStaticMesh(const uint32 GeometryCacheKey, UE::Geometry::FMeshShapeGenerator& MeshGenerator, const int32 LODsToGenerateNum)
 {
 	{
-		FReadScopeLock ReadLock(RWLock);
+		FReadScopeLock ReadLock(GeometryCacheRWLock);
 		if (TObjectPtr<UStaticMesh>* StaticMeshPtrPtr = StaticMeshCacheMap.Find(GeometryCacheKey))
 		{
 			return *StaticMeshPtrPtr;
@@ -156,7 +167,7 @@ UStaticMesh* FChaosVDGeometryBuilder::CreateAndCacheStaticMesh(const uint32 Geom
 	}
 
 	{
-		FWriteScopeLock WriteLock(RWLock);
+		FWriteScopeLock WriteLock(GeometryCacheRWLock);
 		StaticMeshCacheMap.Add(GeometryCacheKey, MainStaticMesh);
 	}
 
@@ -202,46 +213,49 @@ void FChaosVDGeometryBuilder::ApplyMeshToComponentFromKey(TWeakObjectPtr<UMeshCo
 
 bool FChaosVDGeometryBuilder::GameThreadTick(float DeltaTime)
 {
-	if (bHasPendingJobs)
-	{
-		FReadScopeLock ReadLock(RWLock);
-		int32 PendingJobs = MeshComponentsWaitingForGeometryByKey.Num();
-		GeometryGenerationNotification.Update(PendingJobs);
-		
-		if (PendingJobs == 0)
-		{
-			bHasPendingJobs = false;
-		}
-	}
-	
-	int32 CurrentGeometryProcessed = 0;
-	while (!GeometryReadyToApplyQueue.IsEmpty() && CurrentGeometryProcessed < CVarChaosVDGeometryToProcessPerTick->GetInt())
+	GeometryGenerationNotification.Update(MeshComponentsWaitingGeometryNum);
+
 	{
-		uint32 GeometryKey = 0;
-		GeometryReadyToApplyQueue.Dequeue(GeometryKey);
+		int32 CurrentGeometryProcessedNum = 0;
+		bool bCanContinueProcessing = true;
 
-		CurrentGeometryProcessed++;
-		
-		TArray<TWeakObjectPtr<UMeshComponent>>* MeshComponentsWaiting = nullptr;
+		FWriteScopeLock WaitListWriteLock(GeometryWaitListRWLock);
+		for (FChaosVDWaitListMeshMap::TIterator MeshWaitListRemoveIterator = MeshComponentsWaitingForGeometryByKey.CreateIterator(); MeshWaitListRemoveIterator; ++MeshWaitListRemoveIterator)
 		{
-			FReadScopeLock ReadLock(RWLock);
-			MeshComponentsWaiting = MeshComponentsWaitingForGeometryByKey.Find(GeometryKey);
-		}
+			const uint32 GeometryKey = MeshWaitListRemoveIterator.Key();
 
-		if (MeshComponentsWaiting)
-		{
-			for (const TWeakObjectPtr<UMeshComponent>& MeshComponent : *MeshComponentsWaiting)
 			{
-				ApplyMeshToComponentFromKey(MeshComponent, GeometryKey);
+				FReadScopeLock CacheReadLocLock(GeometryCacheRWLock);
+				if (HasGeometryInCache_AssumesLocked(GeometryKey))
+				{
+					for (TArray<TWeakObjectPtr<UMeshComponent>>::TIterator MeshRemoveIterator = MeshWaitListRemoveIterator.Value().CreateIterator(); MeshRemoveIterator; ++MeshRemoveIterator)
+					{
+						CurrentGeometryProcessedNum++;
+						MeshComponentsWaitingGeometryNum -=1;
+						bCanContinueProcessing = CurrentGeometryProcessedNum < CVarChaosVDGeometryToProcessPerTick->GetInt();
+						ApplyMeshToComponentFromKey(*MeshRemoveIterator, GeometryKey);
+						MeshRemoveIterator.RemoveCurrent();
+					}
+				
+					if (!bCanContinueProcessing)
+					{
+						break;
+					}	
+				}
+			}
+
+			if (MeshWaitListRemoveIterator.Value().IsEmpty())
+			{
+				MeshWaitListRemoveIterator.RemoveCurrent();
+			}
+
+			if (!bCanContinueProcessing)
+			{
+				break;
 			}
-		}
-		
-		{
-			FWriteScopeLock WriteLock(RWLock);
-			MeshComponentsWaitingForGeometryByKey.Remove(GeometryKey);
 		}
 	}
-	
+
 	return true;
 }
 
@@ -259,7 +273,7 @@ void FChaosVDGeometryBuilder::RegisterMeshComponentWaitingForGeometry(uint32 Geo
 	}
 
 	{
-		FWriteScopeLock WriteLock(RWLock);
+		FWriteScopeLock WriteLock(GeometryWaitListRWLock);
 	
 		if (TArray<TWeakObjectPtr<UMeshComponent>>* MeshComponentsWaiting = MeshComponentsWaitingForGeometryByKey.Find(GeometryKey))
 		{
@@ -268,7 +282,8 @@ void FChaosVDGeometryBuilder::RegisterMeshComponentWaitingForGeometry(uint32 Geo
 		else
 		{
 			MeshComponentsWaitingForGeometryByKey.Add(GeometryKey, {MesComponent});
-			bHasPendingJobs = true;
 		}
+
+		MeshComponentsWaitingGeometryNum += 1;
 	}
 }
diff --git a/Engine/Plugins/Experimental/ChaosVD/Source/ChaosVD/Private/ChaosVDParticleActor.cpp b/Engine/Plugins/Experimental/ChaosVD/Source/ChaosVD/Private/ChaosVDParticleActor.cpp
--- a/Engine/Plugins/Experimental/ChaosVD/Source/ChaosVD/Private/ChaosVDParticleActor.cpp
+++ b/Engine/Plugins/Experimental/ChaosVD/Source/ChaosVD/Private/ChaosVDParticleActor.cpp
@@ -3,6 +3,7 @@
 #include "ChaosVDParticleActor.h"
 
 #include "ChaosVDGeometryBuilder.h"
+#include "ChaosVDModule.h"
 #include "ChaosVDScene.h"
 #include "Components/ChaosVDInstancedStaticMeshComponent.h"
 #include "Components/ChaosVDStaticMeshComponent.h"
@@ -28,8 +29,17 @@ void AChaosVDParticleActor::UpdateFromRecordedParticleData(const FChaosVDParticl
 
 	if (InRecordedData.ParticlePositionRotation.HasValidData())
 	{
-		// TODO: Only update if the transform (either the simulation or the particle) has changed;
-		SetActorLocationAndRotation(SimulationTransform.TransformPosition(InRecordedData.ParticlePositionRotation.MX), SimulationTransform.GetRotation() * InRecordedData.ParticlePositionRotation.MR, false);
+		const FVector TargetLocation = SimulationTransform.TransformPosition(InRecordedData.ParticlePositionRotation.MX);
+		if (GetActorLocation() != TargetLocation)
+		{
+			SetActorLocation(InRecordedData.ParticlePositionRotation.MX);
+		}
+
+		const FQuat TargetRotation = SimulationTransform.GetRotation() * InRecordedData.ParticlePositionRotation.MR;
+		if (GetActorRotation() != TargetRotation.Rotator())
+		{
+			SetActorRotation(TargetRotation);
+		}
 	}
 
 	if (ParticleDataViewer.GeometryHash != InRecordedData.GeometryHash)
@@ -37,12 +47,19 @@ void AChaosVDParticleActor::UpdateFromRecordedParticleData(const FChaosVDParticl
 		UpdateGeometry(InRecordedData.GeometryHash, EChaosVDActorGeometryUpdateFlags::ForceUpdate);
 	}
 
+	// This is iterating and comparing each element of the array,
+	// We might need to find a faster way of determine if the data changed, but for now this is faster than assuming it changed
+	const bool bShapeDataIsDirty = ParticleDataViewer.CollisionDataPerShape != InRecordedData.CollisionDataPerShape;
+
 	// TODO: We should store a ptr to the data and in our custom details panel draw it
 	ParticleDataViewer = InRecordedData;
 
 	// Now that we have updated particle data, update the Shape data and visibility as needed
-	UpdateShapeDataComponents();
-	UpdateGeometryComponentsVisibility();
+	if (bShapeDataIsDirty)
+	{
+		UpdateShapeDataComponents();
+		UpdateGeometryComponentsVisibility();
+	}
 }
 
 void AChaosVDParticleActor::UpdateCollisionData(const TArray<TSharedPtr<FChaosVDParticlePairMidPhase>>& InRecordedMidPhases)
@@ -264,6 +281,27 @@ void AChaosVDParticleActor::UpdateGeometryComponentsVisibility()
 	}
 }
 
+void AChaosVDParticleActor::SetIsActive(bool bNewActive)
+{
+	if (bIsActive != bNewActive)
+	{
+#if WITH_EDITOR
+		//TODO: We need to add support for this to our Scene Outliner
+		// This will hide the actor and disable it in the outliner but it will still be listed
+		// We need to add a way to unlist inactive particle actors without a full hierarchy rebuild, which would be too costly
+		bEditable = bNewActive;
+		bListedInSceneOutliner = bNewActive;
+		SetIsTemporarilyHiddenInEditor(!bNewActive);
+#endif
+		bIsActive = bNewActive;
+
+		if (const TSharedPtr<FChaosVDScene> ScenePtr = OwningScene.Pin())
+		{
+			ScenePtr->OnActorActiveStateChanged().Broadcast(this);
+		}
+	}
+}
+
 void AChaosVDParticleActor::DrawVisualization(const FSceneView* View, FPrimitiveDrawInterface* PDI)
 {
 	for (const TPair<FStringView, TUniquePtr<FChaosVDDataVisualizerBase>>& VisualizerWithID : CVDVisualizers)
diff --git a/Engine/Plugins/Experimental/ChaosVD/Source/ChaosVD/Private/ChaosVDPlaybackController.cpp b/Engine/Plugins/Experimental/ChaosVD/Source/ChaosVD/Private/ChaosVDPlaybackController.cpp
--- a/Engine/Plugins/Experimental/ChaosVD/Source/ChaosVD/Private/ChaosVDPlaybackController.cpp
+++ b/Engine/Plugins/Experimental/ChaosVD/Source/ChaosVD/Private/ChaosVDPlaybackController.cpp
@@ -60,13 +60,7 @@ bool FChaosVDPlaybackController::LoadChaosVDRecordingFromTraceSession(const FCha
 
 	HandleCurrentRecordingUpdated();
 
-	LoadedRecording->OnGeometryDataLoaded().AddLambda([this](const Chaos::FConstImplicitObjectPtr& NewGeometry, const uint32 GeometryID)
-	{
-		if (const TSharedPtr<FChaosVDScene> ScenePtr = SceneToControl.Pin())
-		{
-			ScenePtr->HandleNewGeometryData(NewGeometry, GeometryID);
-		}
-	});
+	LoadedRecording->OnGeometryDataLoaded().AddRaw(this, &FChaosVDPlaybackController::EnqueueGeometryDataUpdate);
 	
 	if (TSharedPtr<FChaosVDScene> ScenePtr = SceneToControl.Pin())
 	{
@@ -106,31 +100,23 @@ void FChaosVDPlaybackController::UnloadCurrentRecording(EChaosVDUnloadRecordingF
 	bPlayedFirstFrame = false;
 }
 
-void FChaosVDPlaybackController::PlayFromClosestKeyFrame_AssumesLocked(const int32 InTrackID, const int32 FrameNumber, FChaosVDScene& InSceneToControl)
+void FChaosVDPlaybackController::PlayFromClosestKeyFrame_AssumesLocked(const int32 InTrackID, const int32 FrameNumber, FChaosVDScene& InSceneToControl) const
 {
 	const int32 KeyFrameNumber = LoadedRecording->FindFirstSolverKeyFrameNumberFromFrame_AssumesLocked(InTrackID, FrameNumber);
-
 	if (!ensure(KeyFrameNumber >= 0))
 	{
 		return;
 	}
 
-	for (int32 CurrentFrameNumber = KeyFrameNumber; CurrentFrameNumber < FrameNumber; CurrentFrameNumber++)
-	{
-		if (const FChaosVDSolverFrameData* SolverFrameData = LoadedRecording->GetSolverFrameData_AssumesLocked(InTrackID, CurrentFrameNumber))
-		{
-			const int32 LastStepNumber = SolverFrameData->SolverSteps.Num() - 1;
+	// Instead of playing back each delta frame since the key frame, generate a new solver frame with all the deltas collapsed in one
+	// This increases the tool performance while scrubbing or live debugging if there are few keyframes
+	const int32 LastFrameToEvaluateIndex = FrameNumber - 1;
+	FChaosVDSolverFrameData CollapsedFrameData;
+	LoadedRecording->CollapseSolverFramesRange_AssumesLocked(InTrackID, KeyFrameNumber, LastFrameToEvaluateIndex, CollapsedFrameData);
 
-			if (SolverFrameData->SolverSteps.IsValidIndex(LastStepNumber))
-			{
-				InSceneToControl.UpdateFromRecordedStepData(InTrackID, SolverFrameData->DebugName, SolverFrameData->SolverSteps[LastStepNumber], *SolverFrameData);
-			}
-		}
-		else
-		{
-			// This is common if we stop PIE, change worlds, and PIE again without stopping the recording
-			UE_LOG(LogChaosVDEditor, Verbose, TEXT("[%s] Failed to read solver frame data for frame [%d] in track [%d]"), ANSI_TO_TCHAR(__FUNCTION__), CurrentFrameNumber, InTrackID);
-		}
+	if (CollapsedFrameData.SolverSteps.Num() > 0)
+	{
+		InSceneToControl.UpdateFromRecordedStepData(InTrackID, CollapsedFrameData.DebugName, CollapsedFrameData.SolverSteps[0], CollapsedFrameData);
 	}
 }
 
@@ -148,7 +134,12 @@ void FChaosVDPlaybackController::EnqueueTrackInfoUpdate(const FChaosVDTrackInfo&
 	TrackInfoUpdateGTQueue.Enqueue(InfoUpdate);
 }
 
-void FChaosVDPlaybackController::GoToRecordedSolverStep_AssumesLocked(const int32 InTrackID, const int32 FrameNumber, const int32 Step, FGuid InstigatorID)
+void FChaosVDPlaybackController::EnqueueGeometryDataUpdate(const Chaos::FConstImplicitObjectPtr& NewGeometry, const uint32 GeometryID)
+{	
+	GeometryDataUpdateGTQueue.Enqueue({NewGeometry, GeometryID });
+}
+
+void FChaosVDPlaybackController::GoToRecordedSolverStep_AssumesLocked(const int32 InTrackID, const int32 FrameNumber, const int32 Step, FGuid InstigatorID, int32 Attempts)
 {
 	if (const TSharedPtr<FChaosVDScene> SceneToControlSharedPtr = SceneToControl.Pin())
 	{
@@ -168,16 +159,16 @@ void FChaosVDPlaybackController::GoToRecordedSolverStep_AssumesLocked(const int3
 				}
 			}
 
-			const int32 FrameDiff = FrameNumber - CurrentTrackInfo->CurrentFrame;
-			constexpr int32 FrameDriftTolerance = 1;
-			if (FMath::Abs(FrameDiff) > FrameDriftTolerance || CurrentTrackInfo->CurrentFrame == 0)
-			{
-				// As Frames are recorded as delta, we need to make sure of playing back all the deltas since the closest keyframe
-				PlayFromClosestKeyFrame_AssumesLocked(InTrackID, FrameNumber, *SceneToControlSharedPtr.Get());
-			}
-			
+			if (FChaosVDSolverFrameData* SolverFrameData = LoadedRecording->GetSolverFrameData_AssumesLocked(InTrackID, FrameNumber))
 			{
-				const FChaosVDSolverFrameData* SolverFrameData = LoadedRecording->GetSolverFrameData_AssumesLocked(InTrackID, FrameNumber);
+				const int32 FrameDiff = FrameNumber - CurrentTrackInfo->CurrentFrame;
+				constexpr int32 FrameDriftTolerance = 1;
+				if (FMath::Abs(FrameDiff) > FrameDriftTolerance || CurrentTrackInfo->CurrentFrame == 0)
+				{
+					// As Frames are recorded as delta, we need to make sure of playing back all the deltas since the closest keyframe
+					PlayFromClosestKeyFrame_AssumesLocked(InTrackID, FrameNumber, *SceneToControlSharedPtr.Get());
+				}
+
 				if (CurrentTrackInfo->LockedOnStep != INDEX_NONE)
 				{
 					// If this track is locked to a specific step, we need to play back the previous steps on the current frame, because not all steps capture the same data.
@@ -208,13 +199,12 @@ void FChaosVDPlaybackController::GoToRecordedSolverStep_AssumesLocked(const int3
 						UE_LOG(LogChaosVDEditor, Verbose, TEXT("[%s] Tried to scrub to an invalid step | Step Number [%d] ..."), ANSI_TO_TCHAR(__FUNCTION__), CurrentTrackInfo->LockedOnStep);
 					}
 				}
-			}
-			
 
-			CurrentTrackInfo->CurrentFrame = FrameNumber;
-			CurrentTrackInfo->CurrentStep = Step;
+				CurrentTrackInfo->CurrentFrame = FrameNumber;
+				CurrentTrackInfo->CurrentStep = Step;
 			
-			EnqueueTrackInfoUpdate(*CurrentTrackInfo.Get(), InstigatorID);
+				EnqueueTrackInfoUpdate(*CurrentTrackInfo.Get(), InstigatorID);
+			}
 		}
 	}
 	else
@@ -223,44 +213,37 @@ void FChaosVDPlaybackController::GoToRecordedSolverStep_AssumesLocked(const int3
 	}
 }
 
-void FChaosVDPlaybackController::GoToRecordedGameFrame_AssumesLocked(const int32 FrameNumber, FGuid InstigatorID)
+void FChaosVDPlaybackController::GoToRecordedGameFrame_AssumesLocked(const int32 FrameNumber, FGuid InstigatorID, int32 Attempts)
 {
 	if (const TSharedPtr<FChaosVDScene> SceneToControlSharedPtr = SceneToControl.Pin())
 	{
 		if (ensure(LoadedRecording.IsValid()))
 		{
-			if (const FChaosVDGameFrameData* FrameData = LoadedRecording->GetGameFrameData_AssumesLocked(FrameNumber))
+			if (TrackInfoByIDMap* TrackInfoByID = TrackInfoPerType.Find(EChaosVDTrackType::Game))
 			{
-				const bool bIsFrameFullyLoaded = FrameData->EndTime > 0;
-				if (!bIsFrameFullyLoaded)
-				{
-					return;
-				}
-
-				if (TrackInfoByIDMap* TrackInfoByID = TrackInfoPerType.Find(EChaosVDTrackType::Game))
+				if (const TSharedPtr<FChaosVDTrackInfo>& TrackInfoSharedPtr = TrackInfoByID->FindChecked(GameTrackID))
 				{
-					TArray<int32> AvailableSolversID;
-					LoadedRecording->GetAvailableSolverIDsAtGameFrameNumber_AssumesLocked(FrameNumber, AvailableSolversID);
+					if (const FChaosVDGameFrameData* FoundGameFrameData = LoadedRecording->GetGameFrameData_AssumesLocked(FrameNumber))
+					{
+						TArray<int32> AvailableSolversID;
+						LoadedRecording->GetAvailableSolverIDsAtGameFrameNumber_AssumesLocked(FrameNumber, AvailableSolversID);
 
-					SceneToControlSharedPtr->HandleEnterNewGameFrame(FrameNumber, AvailableSolversID);
+						SceneToControlSharedPtr->HandleEnterNewGameFrame(FrameNumber, AvailableSolversID);
 
-					for (int32 SolverID : AvailableSolversID)
-					{
-						// When Scrubbing the timeline by game frames instead of solvers, try to go to the first solver frame on the first platform cycle of the game frame.
-						// Game Frames are not in sync with Solver Frames and Solver steps.
-						const int32 SolverFrameNumber = LoadedRecording->GetLowestSolverFrameNumberAtCycle_AssumesLocked(SolverID, FrameData->FirstCycle);
+						for (const int32 SolverID : AvailableSolversID)
+						{
+							// When Scrubbing the timeline by game frames instead of solvers, try to go to the first solver frame on the first platform cycle of the game frame.
+							// Game Frames are not in sync with Solver Frames and Solver steps.
+							const int32 SolverFrameNumber = LoadedRecording->GetLowestSolverFrameNumberAtCycle_AssumesLocked(SolverID, FoundGameFrameData->FirstCycle);
 	
-						const int32 StepNumber = GetTrackLastStepAtFrame(EChaosVDTrackType::Solver, SolverID, SolverFrameNumber);
+							const int32 StepNumber = GetTrackLastStepAtFrame(EChaosVDTrackType::Solver, SolverID, SolverFrameNumber);
 
-						GoToTrackFrame_AssumesLocked(InstigatorID, EChaosVDTrackType::Solver, SolverID, SolverFrameNumber, StepNumber);
-					}
-
-					TSharedPtr<FChaosVDTrackInfo>* TrackInfoPtrPtr = TrackInfoByID->Find(GameTrackID);
-					if (TSharedPtr<FChaosVDTrackInfo> TrackInfoSharedPtr = TrackInfoPtrPtr? *TrackInfoPtrPtr : nullptr)
-					{
-						TrackInfoSharedPtr->CurrentFrame = FrameNumber;
-						EnqueueTrackInfoUpdate(*TrackInfoSharedPtr.Get(), InstigatorID);
+							GoToTrackFrame_AssumesLocked(InstigatorID, EChaosVDTrackType::Solver, SolverID, SolverFrameNumber, StepNumber);
+						}
 					}
+					
+					TrackInfoSharedPtr->CurrentFrame = FrameNumber;
+					EnqueueTrackInfoUpdate(*TrackInfoSharedPtr.Get(), InstigatorID);
 				}
 			}
 		}
@@ -557,7 +540,19 @@ void FChaosVDPlaybackController::GetAvailableTrackInfosAtTrackFrame(EChaosVDTrac
 }
 
 bool FChaosVDPlaybackController::Tick(float DeltaTime)
-{
+{	
+	if (!GeometryDataUpdateGTQueue.IsEmpty())
+	{
+		FChaosVDGeometryDataUpdate GeometryDataUpdate;
+		while (GeometryDataUpdateGTQueue.Dequeue(GeometryDataUpdate))
+		{
+			if (const TSharedPtr<FChaosVDScene> ScenePtr = SceneToControl.Pin())
+			{
+				ScenePtr->HandleNewGeometryData(GeometryDataUpdate.NewGeometry, GeometryDataUpdate.GeometryID);
+			}
+		}
+	}
+
 	const TWeakPtr<FChaosVDPlaybackController> ThisWeakPtr = DoesSharedInstanceExist() ? AsWeak() : nullptr;
 	if (!ThisWeakPtr.IsValid())
 	{
@@ -661,7 +656,7 @@ void FChaosVDPlaybackController::HandleCurrentRecordingUpdated()
 		GameTrackInfo = MakeShared<FChaosVDTrackInfo>();
 		GameTrackInfo->CurrentFrame = 0;
 		GameTrackInfo->CurrentStep = 0;
-	};
+	}
 
 	GameTrackInfo->MaxFrames = LoadedRecording->GetAvailableGameFrames_AssumesLocked().Num();
 	GameTrackInfo->TrackType = EChaosVDTrackType::Game;
diff --git a/Engine/Plugins/Experimental/ChaosVD/Source/ChaosVD/Private/ChaosVDScene.cpp b/Engine/Plugins/Experimental/ChaosVD/Source/ChaosVD/Private/ChaosVDScene.cpp
--- a/Engine/Plugins/Experimental/ChaosVD/Source/ChaosVD/Private/ChaosVDScene.cpp
+++ b/Engine/Plugins/Experimental/ChaosVD/Source/ChaosVD/Private/ChaosVDScene.cpp
@@ -4,6 +4,7 @@
 
 #include "ChaosVDEditorSettings.h"
 #include "ChaosVDGeometryBuilder.h"
+#include "ChaosVDModule.h"
 #include "ChaosVDParticleActor.h"
 #include "Chaos/ImplicitObject.h"
 #include "ChaosVDRecording.h"
@@ -133,6 +134,12 @@ void FChaosVDScene::UpdateFromRecordedStepData(const int32 SolverID, const FStri
 			{
 				if (AChaosVDParticleActor* ExistingParticleVDInstancePtr = *ExistingParticleVDInstancePtrPtr)
 				{
+					// We have new data for this particle, so re-activate the existing actor
+					if (!ExistingParticleVDInstancePtr->IsActive())
+					{
+						ExistingParticleVDInstancePtr->SetIsActive(true);
+					}
+
 					ExistingParticleVDInstancePtr->UpdateFromRecordedParticleData(Particle, InFrameData.SimulationTransform);
 				}
 				else
@@ -166,27 +173,28 @@ void FChaosVDScene::UpdateFromRecordedStepData(const int32 SolverID, const FStri
 
 	UpdateParticlesCollisionData(InRecordedStepData, SolverID);
 
-	int32 AmountRemoved = 0;
-	for (FChaosVDParticlesByIDMap::TIterator RemoveIterator = SolverParticlesByID.CreateIterator(); RemoveIterator; ++RemoveIterator)
+	for (const TPair<int32, AChaosVDParticleActor*>& ParticleActorWithID : SolverParticlesByID)
 	{
 		// If we are playing back a keyframe, the scene should only contain what it is in the recorded data
-		const bool bShouldDestroyParticleAnyway = InFrameData.bIsKeyFrame && !ParticlesIDsInRecordedStepData.Contains(RemoveIterator.Key());
+		const bool bShouldDestroyParticleAnyway = InFrameData.bIsKeyFrame && !ParticlesIDsInRecordedStepData.Contains(ParticleActorWithID.Key);
 		
-		if (bShouldDestroyParticleAnyway || InFrameData.ParticlesDestroyedIDs.Contains(RemoveIterator.Key()))
+		if (bShouldDestroyParticleAnyway || InFrameData.ParticlesDestroyedIDs.Contains(ParticleActorWithID.Key))
 		{
-			if (AChaosVDParticleActor* ActorToRemove = RemoveIterator.Value())
+			// In large maps moving at high speed (like when moving on a vehicle), level streaming adds/removes hundreds of actors (and therefore particles) constantly.
+			// Destroying particle actors is expensive, specially if we need to spawn them again sooner as we will nee to rebuild-them.
+			// So, we deactivate them instead.
+
+			// TODO: We need an actor pool system, so we can keep memory under control as well.
+			if (AChaosVDParticleActor* ActorToDeactivate = ParticleActorWithID.Value)
 			{
-				if (IsObjectSelected(ActorToRemove))
+				if (IsObjectSelected(ActorToDeactivate))
 				{
 					ClearSelectionAndNotify();
 				}
 
-				PhysicsVDWorld->DestroyActor(ActorToRemove);
+				ActorToDeactivate->SetIsActive(false);
 			}
 
-			RemoveIterator.RemoveCurrent();
-
-			AmountRemoved++;
 		}
 	}
 	
@@ -269,6 +277,8 @@ void FChaosVDScene::HandleEnterNewGameFrame(int32 FrameNumber, const TArray<int3
 	{
 		if (!AvailableSolversSet.Contains(RemoveIterator.Key()))
 		{
+			UE_LOG(LogChaosVDEditor, Log, TEXT("[%s] Removing Solver [%d] as it is no longer present in the recording"), ANSI_TO_TCHAR(__FUNCTION__), RemoveIterator.Key());
+	
 			for (const TPair<int32, AChaosVDParticleActor*>& ParticleVDInstanceWithID : RemoveIterator.Value())
 			{
 				if (IsObjectSelected(ParticleVDInstanceWithID.Value))
@@ -338,6 +348,7 @@ AChaosVDParticleActor* FChaosVDScene::SpawnParticleFromRecordedData(const FChaos
 
 	if (AChaosVDParticleActor* NewActor = PhysicsVDWorld->SpawnActor<AChaosVDParticleActor>(Params))
 	{
+		NewActor->SetIsActive(true);
 		NewActor->UpdateFromRecordedParticleData(InParticleData, InFrameData.SimulationTransform);
 
 		if (!InParticleData.DebugName.IsEmpty())
diff --git a/Engine/Plugins/Experimental/ChaosVD/Source/ChaosVD/Private/ChaosVDWorldOutlinerMode.cpp b/Engine/Plugins/Experimental/ChaosVD/Source/ChaosVD/Private/ChaosVDWorldOutlinerMode.cpp
--- a/Engine/Plugins/Experimental/ChaosVD/Source/ChaosVD/Private/ChaosVDWorldOutlinerMode.cpp
+++ b/Engine/Plugins/Experimental/ChaosVD/Source/ChaosVD/Private/ChaosVDWorldOutlinerMode.cpp
@@ -4,9 +4,14 @@
 
 #include "ActorTreeItem.h"
 #include "ChaosVDModule.h"
+#include "ChaosVDParticleActor.h"
 #include "ChaosVDScene.h"
 #include "Elements/Framework/TypedElementSelectionSet.h"
 
+static FAutoConsoleVariable CVarChaosVDQueueAndCombineSceneOutlinerEvents(
+	TEXT("p.Chaos.VD.Tool.QueueAndCombineSceneOutlinerEvents"),
+	true,
+	TEXT("If set to true, scene outliner events will be queued and sent once per frame. If there was a unprocessed event for an item, the las queued event will replace it"));
 
 FChaosVDWorldOutlinerMode::FChaosVDWorldOutlinerMode(const FActorModeParams& InModeParams, TWeakPtr<FChaosVDScene> InScene)
 	: FActorMode(InModeParams),
@@ -18,6 +23,8 @@ FChaosVDWorldOutlinerMode::FChaosVDWorldOutlinerMode(const FActorModeParams& InM
 		return;
 	}
 
+	ScenePtr->OnActorActiveStateChanged().AddRaw(this, &FChaosVDWorldOutlinerMode::HandleActorActiveStateChanged);
+
 	RegisterSelectionSetObject(ScenePtr->GetElementSelectionSet());
 
 	ActorLabelChangedDelegateHandle = FCoreDelegates::OnActorLabelChanged.AddRaw(this, &FChaosVDWorldOutlinerMode::HandleActorLabelChanged);
@@ -26,6 +33,11 @@ FChaosVDWorldOutlinerMode::FChaosVDWorldOutlinerMode(const FActorModeParams& InM
 FChaosVDWorldOutlinerMode::~FChaosVDWorldOutlinerMode()
 {
 	FCoreDelegates::OnActorLabelChanged.Remove(ActorLabelChangedDelegateHandle);
+
+	if (TSharedPtr<FChaosVDScene> ScenePtr = CVDScene.Pin())
+	{
+		ScenePtr->OnActorActiveStateChanged().RemoveAll(this);
+	}
 }
 
 void FChaosVDWorldOutlinerMode::OnItemSelectionChanged(FSceneOutlinerTreeItemPtr Item, ESelectInfo::Type SelectionType, const FSceneOutlinerItemSelection& Selection)
@@ -67,6 +79,33 @@ void FChaosVDWorldOutlinerMode::OnItemDoubleClick(FSceneOutlinerTreeItemPtr Item
 	}
 }
 
+void FChaosVDWorldOutlinerMode::ProcessPendingHierarchyEvents()
+{
+	for (const TPair<FSceneOutlinerTreeItemID, FSceneOutlinerHierarchyChangedData>& PendingEvent : PendingOutlinerEventsMap)
+	{
+		Hierarchy->OnHierarchyChanged().Broadcast(PendingEvent.Value);
+	}
+
+	PendingOutlinerEventsMap.Reset();
+}
+
+bool FChaosVDWorldOutlinerMode::Tick(float DeltaTime)
+{
+	ProcessPendingHierarchyEvents();
+	return true;
+}
+
+void FChaosVDWorldOutlinerMode::EnqueueAndCombineHierarchyEvent(const FSceneOutlinerTreeItemID& ItemID, const FSceneOutlinerHierarchyChangedData& EnventToProcess)
+{
+	if (FSceneOutlinerHierarchyChangedData* EventData = PendingOutlinerEventsMap.Find(ItemID))
+	{
+		*EventData = EnventToProcess;
+	}
+	else
+	{
+		PendingOutlinerEventsMap.Add(ItemID, EnventToProcess);
+	}	
+}
 void FChaosVDWorldOutlinerMode::HandleActorLabelChanged(AActor* ChangedActor)
 {
 	if (!ensure(ChangedActor))
@@ -84,6 +123,43 @@ void FChaosVDWorldOutlinerMode::HandleActorLabelChanged(AActor* ChangedActor)
 	}
 }
 
+void FChaosVDWorldOutlinerMode::HandleActorActiveStateChanged(AChaosVDParticleActor* ChangedActor)
+{
+	if (!ChangedActor)
+	{
+		return;
+	}
+
+	if (Hierarchy.IsValid())
+	{
+		FSceneOutlinerHierarchyChangedData EventData;
+
+		if (ChangedActor->IsActive())
+		{
+			EventData.Type = FSceneOutlinerHierarchyChangedData::Added;
+			EventData.Items.Emplace(CreateItemFor<FActorTreeItem>(ChangedActor));
+		}
+		else
+		{
+			EventData.ItemIDs.Emplace(ChangedActor);
+			EventData.Type = FSceneOutlinerHierarchyChangedData::Removed;
+		}
+
+		// There is currently a bug in the Scene Outliner where if opposite events happen multiple times within the same tick, the last ones get dropped
+		// (UE-193877). As our current use case is fairly simple, as a workaround we can just queue the events and process them once per frame
+		// Only taking into account only the last requested event for each item.
+		// Keeping this behind a cvar enabled by default so when the Scene Outliner bug is fixed, we can test it easily.
+		if (CVarChaosVDQueueAndCombineSceneOutlinerEvents->GetBool())
+		{
+			EnqueueAndCombineHierarchyEvent(ChangedActor, EventData);
+		}
+		else
+		{
+			Hierarchy->OnHierarchyChanged().Broadcast(EventData);
+		}
+	}
+}
+
 void FChaosVDWorldOutlinerMode::HandlePostSelectionChange(const UTypedElementSelectionSet* ChangesSelectionSet)
 {
 	TArray<AActor*> SelectedActors = ChangesSelectionSet->GetSelectedObjects<AActor>();
diff --git a/Engine/Plugins/Experimental/ChaosVD/Source/ChaosVD/Private/ChaosVDWorldOutlinerMode.h b/Engine/Plugins/Experimental/ChaosVD/Source/ChaosVD/Private/ChaosVDWorldOutlinerMode.h
--- a/Engine/Plugins/Experimental/ChaosVD/Source/ChaosVD/Private/ChaosVDWorldOutlinerMode.h
+++ b/Engine/Plugins/Experimental/ChaosVD/Source/ChaosVD/Private/ChaosVDWorldOutlinerMode.h
@@ -4,6 +4,8 @@
 #include "ActorBrowsingMode.h"
 #include "ChaosVDScene.h"
 #include "ChaosVDSceneSelectionObserver.h"
+#include "Containers/Map.h"
+#include "Templates/SharedPointer.h"
 
 class FChaosVDScene;
 
@@ -12,7 +14,7 @@ class FChaosVDScene;
  * It has a more limited view compared to the normal outliner, hiding features we don't support,
  * and it is integrated with the CVD local selection system
  */
-class FChaosVDWorldOutlinerMode : public FActorMode, public FChaosVDSceneSelectionObserver
+class FChaosVDWorldOutlinerMode : public FActorMode, public FChaosVDSceneSelectionObserver, public FTSTickerObjectBase
 {
 public:
 
@@ -22,21 +24,30 @@ public:
 
 	virtual void OnItemSelectionChanged(FSceneOutlinerTreeItemPtr Item, ESelectInfo::Type SelectionType, const FSceneOutlinerItemSelection& Selection) override;
 	virtual void OnItemDoubleClick(FSceneOutlinerTreeItemPtr Item) override;
+	void ProcessPendingHierarchyEvents();
+
 
 	virtual void CreateViewContent(FMenuBuilder& MenuBuilder) override
 	{
 		// Intentionally not implemented, as we don't support the built in menu to switch worlds
 	}
 
-	virtual bool ShouldShowFolders() const override { return true;};
+	virtual bool ShouldShowFolders() const override { return true;}
+
+	virtual bool Tick(float DeltaTime) override;
 
 private:
 
+	void EnqueueAndCombineHierarchyEvent(const FSceneOutlinerTreeItemID& ItemID, const FSceneOutlinerHierarchyChangedData& EnventToProcess);
 	void HandleActorLabelChanged(AActor* ChangedActor);
+	void HandleActorActiveStateChanged(AChaosVDParticleActor* ChangedActor);
 
 	virtual void HandlePostSelectionChange(const UTypedElementSelectionSet* ChangesSelectionSet) override;
 
+private:
 	FDelegateHandle ActorLabelChangedDelegateHandle;
 
 	TWeakPtr<FChaosVDScene> CVDScene;
+
+	TMap<FSceneOutlinerTreeItemID, FSceneOutlinerHierarchyChangedData> PendingOutlinerEventsMap;
 };
diff --git a/Engine/Plugins/Experimental/ChaosVD/Source/ChaosVD/Private/Trace/ChaosVDTraceAnalyzer.cpp b/Engine/Plugins/Experimental/ChaosVD/Source/ChaosVD/Private/Trace/ChaosVDTraceAnalyzer.cpp
--- a/Engine/Plugins/Experimental/ChaosVD/Source/ChaosVD/Private/Trace/ChaosVDTraceAnalyzer.cpp
+++ b/Engine/Plugins/Experimental/ChaosVD/Source/ChaosVD/Private/Trace/ChaosVDTraceAnalyzer.cpp
@@ -53,10 +53,10 @@ bool FChaosVDTraceAnalyzer::OnEvent(uint16 RouteId, EStyle Style, const FOnEvent
 			uint8 FrameType = EventData.GetValue<uint8>("FrameType");
 			if (static_cast<ETraceFrameType>(FrameType) == TraceFrameType_Game)
 			{
-				FChaosVDGameFrameData FrameData;
-				FrameData.FirstCycle = EventData.GetValue<uint64>("Cycle");
-				FrameData.StartTime = Context.EventTime.AsSeconds(FrameData.FirstCycle);
-				ChaosVDTraceProvider->AddGameFrame(MoveTemp(FrameData));
+				TSharedPtr<FChaosVDGameFrameData> FrameData = MakeShared<FChaosVDGameFrameData>();
+				FrameData->FirstCycle = EventData.GetValue<uint64>("Cycle");
+				FrameData->StartTime = Context.EventTime.AsSeconds(FrameData->FirstCycle);
+				ChaosVDTraceProvider->StartGameFrame(FrameData);
 			}
 
 			break;
@@ -67,11 +67,10 @@ bool FChaosVDTraceAnalyzer::OnEvent(uint16 RouteId, EStyle Style, const FOnEvent
 			uint8 FrameType = EventData.GetValue<uint8>("FrameType");
 			if (static_cast<ETraceFrameType>(FrameType) == TraceFrameType_Game)
 			{
-				FWriteScopeLock WriteLock(ChaosVDTraceProvider->GetDataLock());
-				if (FChaosVDGameFrameData* CurrentFrameData = ChaosVDTraceProvider->GetLastGameFrame_AssumesLocked())
+				if (TSharedPtr<FChaosVDGameFrameData> CurrentFrameData = ChaosVDTraceProvider->GetCurrentGameFrame().Pin())
 				{
 					CurrentFrameData->LastCycle = EventData.GetValue<uint64>("Cycle");
-					CurrentFrameData->EndTime = Context.EventTime.AsSeconds(CurrentFrameData->LastCycle);
+					CurrentFrameData->EndTime = Context.EventTime.AsSeconds(CurrentFrameData->LastCycle);	
 				}
 			}
 			break;
@@ -91,15 +90,14 @@ bool FChaosVDTraceAnalyzer::OnEvent(uint16 RouteId, EStyle Style, const FOnEvent
 
 			// Currently not all solvers have an end frame event, so lets just set the end frame time of the previous frame, with the start of this new one.
 			{
-				FWriteScopeLock WriteLock(ChaosVDTraceProvider->GetDataLock());
-				if (FChaosVDSolverFrameData* PrevFrameData  = ChaosVDTraceProvider->GetLastSolverFrame_AssumesLocked(NewFrameData.SolverID))
+				if (FChaosVDSolverFrameData* PrevFrameData = ChaosVDTraceProvider->GetCurrentSolverFrame(NewFrameData.SolverID))
 				{
 					PrevFrameData->EndTime = NewFrameData.StartTime;
 				}
 			}
 
 			// Add an empty frame. It will be filled out by the solver trace events
-			ChaosVDTraceProvider->AddSolverFrame(NewFrameData.SolverID, MoveTemp(NewFrameData));
+			ChaosVDTraceProvider->StartSolverFrame(NewFrameData.SolverID, MoveTemp(NewFrameData));
 			break;
 		}
 	case RouteId_ChaosVDSolverFrameEnd:
@@ -110,9 +108,8 @@ bool FChaosVDTraceAnalyzer::OnEvent(uint16 RouteId, EStyle Style, const FOnEvent
 		{
 			const int32 SolverID = EventData.GetValue<int32>("SolverID");
 
-			FWriteScopeLock WriteLock(ChaosVDTraceProvider->GetDataLock());
 			// This can be null if the recording started Mid-Frame. In this case we just discard the data for now
-			if (FChaosVDSolverFrameData* FrameData  = ChaosVDTraceProvider->GetLastSolverFrame_AssumesLocked(SolverID))
+			if (FChaosVDSolverFrameData* FrameData = ChaosVDTraceProvider->GetCurrentSolverFrame(SolverID))
 			{
 				// Add an empty step. It will be filled out by the particle (and later on other objects/elements) events
 				FChaosVDStepData& StepData = FrameData->SolverSteps.AddDefaulted_GetRef();
@@ -132,8 +129,7 @@ bool FChaosVDTraceAnalyzer::OnEvent(uint16 RouteId, EStyle Style, const FOnEvent
 		{
 			const int32 SolverID = EventData.GetValue<int32>("SolverID");
 
-			FWriteScopeLock WriteLock(ChaosVDTraceProvider->GetDataLock());
-			if (FChaosVDSolverFrameData* FrameData = ChaosVDTraceProvider->GetLastSolverFrame_AssumesLocked(SolverID))
+			if (FChaosVDSolverFrameData* FrameData = ChaosVDTraceProvider->GetCurrentSolverFrame(SolverID))
 			{
 				int32 ParticleDestroyedID = EventData.GetValue<int32>("ParticleID");
 
@@ -198,9 +194,8 @@ bool FChaosVDTraceAnalyzer::OnEvent(uint16 RouteId, EStyle Style, const FOnEvent
 			FQuat Rotation;
 			CVD_READ_TRACE_QUAT(Rotation, Rotation, float, EventData);
 
-			FWriteScopeLock WriteLock(ChaosVDTraceProvider->GetDataLock());
 			// This can be null if the recording started Mid-Frame. In this case we just discard the data for now
-			if (FChaosVDSolverFrameData* FrameData = ChaosVDTraceProvider->GetLastSolverFrame_AssumesLocked(SolverID))
+			if (FChaosVDSolverFrameData* FrameData = ChaosVDTraceProvider->GetCurrentSolverFrame(SolverID))
 			{
 				FrameData->SimulationTransform.SetLocation(Position);
 				FrameData->SimulationTransform.SetRotation(Rotation);
@@ -214,12 +209,11 @@ bool FChaosVDTraceAnalyzer::OnEvent(uint16 RouteId, EStyle Style, const FOnEvent
 
 			EventData.GetString("DebugName", TrackedLocation.DebugName);
 			
-			FWriteScopeLock WriteLock(ChaosVDTraceProvider->GetDataLock());
-			if (FChaosVDGameFrameData* CurrentFrameData = ChaosVDTraceProvider->GetLastGameFrame_AssumesLocked())
+			if (TSharedPtr<FChaosVDGameFrameData> CurrentFrameData = ChaosVDTraceProvider->GetCurrentGameFrame().Pin())
 			{
 				CurrentFrameData->RecordedNonSolverLocationsByID.Add(FName(TrackedLocation.DebugName), MoveTemp(TrackedLocation));
 			}
-			
+
 			break;
 		}
 	case RouteId_ChaosVDNonSolverTransform:
@@ -228,14 +222,14 @@ bool FChaosVDTraceAnalyzer::OnEvent(uint16 RouteId, EStyle Style, const FOnEvent
 			CVD_READ_TRACE_TRANSFORM(TrackedTransform.Transform, float, EventData);
 
 			EventData.GetString("DebugName", TrackedTransform.DebugName);
-			
-			FWriteScopeLock WriteLock(ChaosVDTraceProvider->GetDataLock());
-			if (FChaosVDGameFrameData* CurrentFrameData = ChaosVDTraceProvider->GetLastGameFrame_AssumesLocked())
+			if (TSharedPtr<FChaosVDGameFrameData> CurrentFrameData = ChaosVDTraceProvider->GetCurrentGameFrame().Pin())
 			{
 				CurrentFrameData->RecordedNonSolverTransformsByID.Add(FName(TrackedTransform.DebugName), MoveTemp(TrackedTransform));
 			}
+
 			break;
 		}
+
 	default:
 		break;
 	}
diff --git a/Engine/Plugins/Experimental/ChaosVD/Source/ChaosVD/Private/Trace/ChaosVDTraceProvider.cpp b/Engine/Plugins/Experimental/ChaosVD/Source/ChaosVD/Private/Trace/ChaosVDTraceProvider.cpp
--- a/Engine/Plugins/Experimental/ChaosVD/Source/ChaosVD/Private/Trace/ChaosVDTraceProvider.cpp
+++ b/Engine/Plugins/Experimental/ChaosVD/Source/ChaosVD/Private/Trace/ChaosVDTraceProvider.cpp
@@ -33,68 +33,81 @@ void FChaosVDTraceProvider::DeleteRecordingInstanceForSession()
 	InternalRecording.Reset();
 }
 
-void FChaosVDTraceProvider::AddSolverFrame(const int32 InSolverID, FChaosVDSolverFrameData&& FrameData)
+void FChaosVDTraceProvider::StartSolverFrame(const int32 InSolverID, FChaosVDSolverFrameData&& FrameData)
 {
-	if (InternalRecording.IsValid())
+	if (!InternalRecording.IsValid())
 	{
-		InternalRecording->AddFrameForSolver(InSolverID, MoveTemp(FrameData));
+		return;
+	}
+
+	if (FChaosVDSolverFrameData* SolveFrameData = CurrentSolverFramesByID.Find(InSolverID))
+	{
+		InternalRecording->AddFrameForSolver(InSolverID, MoveTemp(*SolveFrameData));
+		CurrentSolverFramesByID[InSolverID] = FrameData;
+	}
+	else
+	{
+		CurrentSolverFramesByID.Add(InSolverID, FrameData);
 	}
 }
 
-void FChaosVDTraceProvider::AddGameFrame(FChaosVDGameFrameData&& FrameData)
+void FChaosVDTraceProvider::CommitProcessedGameFramesToRecording()
 {
-	if (InternalRecording.IsValid())
+	// This is a kind of made up number, this includes RBAN solvers so we have more than 1 for sure.
+	constexpr int32 MinSolverIDsExpected = 10;
+	
+	static TArray<int32> SolverIDs;
+	SolverIDs.Reserve(MinSolverIDsExpected);
+
+	// The Game Frame events are not generated by CVD trace code and we don't have control over them
+	// we use them as general timestamps.
+	// These are generated even when no solvers are available (specially in PIE), so we need to discard any game frame that will not resolve to a solver frame
+	// Solver data might be slightly behind, so to make sure we have all the solver data for that time range, we queue a handful of game frames before processing them.
+	constexpr int32 MaxQueueSize = 2;
+	if (CurrentGameFrameQueueSize > MaxQueueSize)
 	{
-		// In PIE, we can have a lot of empty frames at the beginning, so we discard them here
+		TSharedPtr<FChaosVDGameFrameData> ProcessedGameFrameData;
+		DeQueueGameFrameForProcessing(ProcessedGameFrameData);
+
+		if (ProcessedGameFrameData.IsValid())
 		{
-			// This is a kind of made up number, this includes RBAN solvers so we have more than 1 for sure.
-			constexpr int32 MinSolverIDsExpected = 10;
-			
-			static TArray<int32> SolverIDs;
-			SolverIDs.Reserve(MinSolverIDsExpected);
-			
-			FWriteScopeLock WriteLock(GetDataLock());
-			const int32 AvailableGameFrames = InternalRecording->GetAvailableGameFramesNumber_AssumesLocked();
-			InternalRecording->GetAvailableSolverIDsAtGameFrameNumber_AssumesLocked(AvailableGameFrames -1, SolverIDs);
-			if (SolverIDs.IsEmpty())
+			InternalRecording->GetAvailableSolverIDsAtGameFrame(*ProcessedGameFrameData, SolverIDs);
+
+			if (!SolverIDs.IsEmpty())
 			{
-				if (FChaosVDGameFrameData* GameFrame = InternalRecording->GetLastGameFrameData_AssumesLocked())
-				{
-					*GameFrame = FrameData;
-					return;
-				}
+				InternalRecording->AddGameFrameData(*ProcessedGameFrameData);	
 			}
-
-			SolverIDs.Reset();
-		}
-
-		InternalRecording->AddGameFrameData(MoveTemp(FrameData));
+		}	
 	}
+
+	SolverIDs.Reset();
 }
 
-FChaosVDSolverFrameData* FChaosVDTraceProvider::GetSolverFrame_AssumesLocked(const int32 InSolverID, const int32 FrameNumber) const
+void FChaosVDTraceProvider::StartGameFrame(const TSharedPtr<FChaosVDGameFrameData>& InFrameData)
 {
-	return InternalRecording.IsValid() ? InternalRecording->GetSolverFrameData_AssumesLocked(InSolverID, FrameNumber) : nullptr;
+	if (!InternalRecording.IsValid())
+	{
+		return;
+	}
+
+	CommitProcessedGameFramesToRecording();
+
+	EnqueueGameFrameForProcessing(InFrameData);
 }
 
-FChaosVDSolverFrameData* FChaosVDTraceProvider::GetLastSolverFrame_AssumesLocked(const int32 InSolverID) const
+FChaosVDSolverFrameData* FChaosVDTraceProvider::GetCurrentSolverFrame(const int32 InSolverID)
 {
-	if (InternalRecording.IsValid())
+	if (FChaosVDSolverFrameData* SolveFrameData = CurrentSolverFramesByID.Find(InSolverID))
 	{
-		const int32 AvailableFramesNumber = InternalRecording->GetAvailableSolverFramesNumber_AssumesLocked(InSolverID);
-
-		if (AvailableFramesNumber > 0)
-		{
-			return GetSolverFrame_AssumesLocked(InSolverID, AvailableFramesNumber - 1);
-		}
+		return SolveFrameData;
 	}
 
 	return nullptr;
 }
 
-FChaosVDGameFrameData* FChaosVDTraceProvider::GetLastGameFrame_AssumesLocked() const
+TWeakPtr<FChaosVDGameFrameData> FChaosVDTraceProvider::GetCurrentGameFrame()
 {
-	return InternalRecording.IsValid() ? InternalRecording->GetLastGameFrameData_AssumesLocked() : nullptr;
+	return CurrentGameFrame;
 }
 
 FChaosVDBinaryDataContainer& FChaosVDTraceProvider::FindOrAddUnprocessedData(const int32 DataID)
@@ -166,12 +179,6 @@ void FChaosVDTraceProvider::RegisterDataProcessor(TSharedPtr<IChaosVDDataProcess
 	RegisteredDataProcessors.Add(InDataProcessor->GetCompatibleTypeName(), InDataProcessor);
 }
 
-FRWLock& FChaosVDTraceProvider::GetDataLock()
-{
-	check(InternalRecording.IsValid());
-	return InternalRecording->GetRecordingDataLock();
-}
-
 void FChaosVDTraceProvider::RegisterDefaultDataProcessorsIfNeeded()
 {
 	if (bDefaultDataProcessorsRegistered)
@@ -197,3 +204,16 @@ void FChaosVDTraceProvider::RegisterDefaultDataProcessorsIfNeeded()
 
 	bDefaultDataProcessorsRegistered = true;
 }
+
+void FChaosVDTraceProvider::EnqueueGameFrameForProcessing(const TSharedPtr<FChaosVDGameFrameData>& FrameData)
+{
+	CurrentGameFrame = FrameData;
+	CurrentGameFrameQueue.Enqueue(FrameData);
+	CurrentGameFrameQueueSize++;
+}
+
+void FChaosVDTraceProvider::DeQueueGameFrameForProcessing(TSharedPtr<FChaosVDGameFrameData>& OutFrameData)
+{
+	CurrentGameFrameQueue.Dequeue(OutFrameData);
+	CurrentGameFrameQueueSize--;
+}
diff --git a/Engine/Plugins/Experimental/ChaosVD/Source/ChaosVD/Private/Trace/DataProcessors/ChaosVDConstraintDataProcessor.cpp b/Engine/Plugins/Experimental/ChaosVD/Source/ChaosVD/Private/Trace/DataProcessors/ChaosVDConstraintDataProcessor.cpp
--- a/Engine/Plugins/Experimental/ChaosVD/Source/ChaosVD/Private/Trace/DataProcessors/ChaosVDConstraintDataProcessor.cpp
+++ b/Engine/Plugins/Experimental/ChaosVD/Source/ChaosVD/Private/Trace/DataProcessors/ChaosVDConstraintDataProcessor.cpp
@@ -24,8 +24,7 @@ bool FChaosVDConstraintDataProcessor::ProcessRawData(const TArray<uint8>& InData
 	RecordedConstraint.Serialize(MemReader);
 
 	// This can be null if the recording started Mid-Frame. In this case we just discard the data for now
-	FWriteScopeLock WriteLock(ProviderSharedPtr->GetDataLock());
-	if (FChaosVDSolverFrameData* FrameData = ProviderSharedPtr->GetLastSolverFrame_AssumesLocked(RecordedConstraint.SolverID))
+	if (FChaosVDSolverFrameData* FrameData = ProviderSharedPtr->GetCurrentSolverFrame(RecordedConstraint.SolverID))
 	{
 		if (ensureMsgf(FrameData->SolverSteps.Num() > 0, TEXT("A MidPhase was traced without a valid step scope")))
 		{
diff --git a/Engine/Plugins/Experimental/ChaosVD/Source/ChaosVD/Private/Trace/DataProcessors/ChaosVDMidPhaseDataProcessor.cpp b/Engine/Plugins/Experimental/ChaosVD/Source/ChaosVD/Private/Trace/DataProcessors/ChaosVDMidPhaseDataProcessor.cpp
--- a/Engine/Plugins/Experimental/ChaosVD/Source/ChaosVD/Private/Trace/DataProcessors/ChaosVDMidPhaseDataProcessor.cpp
+++ b/Engine/Plugins/Experimental/ChaosVD/Source/ChaosVD/Private/Trace/DataProcessors/ChaosVDMidPhaseDataProcessor.cpp
@@ -23,9 +23,8 @@ bool FChaosVDMidPhaseDataProcessor::ProcessRawData(const TArray<uint8>& InData)
 
 	MidPhase->Serialize(MemReader);
 
-	FWriteScopeLock WriteLock(ProviderSharedPtr->GetDataLock());
 	// This can be null if the recording started Mid-Frame. In this case we just discard the data for now
-	if (FChaosVDSolverFrameData* FrameData = ProviderSharedPtr->GetLastSolverFrame_AssumesLocked(MidPhase->SolverID))
+	if (FChaosVDSolverFrameData* FrameData = ProviderSharedPtr->GetCurrentSolverFrame(MidPhase->SolverID))
 	{
 		if (ensureMsgf(FrameData->SolverSteps.Num() > 0, TEXT("A MidPhase was traced without a valid step scope")))
 		{
diff --git a/Engine/Plugins/Experimental/ChaosVD/Source/ChaosVD/Private/Trace/DataProcessors/ChaosVDTraceParticleDataProcessor.cpp b/Engine/Plugins/Experimental/ChaosVD/Source/ChaosVD/Private/Trace/DataProcessors/ChaosVDTraceParticleDataProcessor.cpp
--- a/Engine/Plugins/Experimental/ChaosVD/Source/ChaosVD/Private/Trace/DataProcessors/ChaosVDTraceParticleDataProcessor.cpp
+++ b/Engine/Plugins/Experimental/ChaosVD/Source/ChaosVD/Private/Trace/DataProcessors/ChaosVDTraceParticleDataProcessor.cpp
@@ -27,9 +27,8 @@ bool FChaosVDTraceParticleDataProcessor::ProcessRawData(const TArray<uint8>& InD
 	FChaosVDParticleDataWrapper ParticleData;
 	ParticleData.Serialize(Ar);
 
-	FWriteScopeLock WriteLock(ProviderSharedPtr->GetDataLock());
 	// This can be null if the recording started Mid-Frame. In this case we just discard the data for now
-	if (FChaosVDSolverFrameData* FrameData = ProviderSharedPtr->GetLastSolverFrame_AssumesLocked(ParticleData.SolverID))
+	if (FChaosVDSolverFrameData* FrameData = ProviderSharedPtr->GetCurrentSolverFrame(ParticleData.SolverID))
 	{
 		if (ensureMsgf(FrameData->SolverSteps.Num() > 0, TEXT("A particle was traced without a valid step scope")))
 		{
diff --git a/Engine/Plugins/Experimental/ChaosVD/Source/ChaosVD/Public/ChaosVDGeometryBuilder.h b/Engine/Plugins/Experimental/ChaosVD/Source/ChaosVD/Public/ChaosVDGeometryBuilder.h
--- a/Engine/Plugins/Experimental/ChaosVD/Source/ChaosVD/Public/ChaosVDGeometryBuilder.h
+++ b/Engine/Plugins/Experimental/ChaosVD/Source/ChaosVD/Public/ChaosVDGeometryBuilder.h
@@ -36,6 +36,8 @@ namespace UE
 	}
 }
 
+typedef TMap<uint32, TArray<TWeakObjectPtr<UMeshComponent>>> FChaosVDWaitListMeshMap;
+
 /*
  * Generates Dynamic mesh components and dynamic meshes based on Chaos implicit object data
  */
@@ -103,6 +105,13 @@ private:
 	 */
 	bool HasNegativeScale(const Chaos::FRigidTransform3& InTransform) const;
 
+	/**
+	 * Return true if we have cached geometry for the provided Geometry Key
+	 * @param GeometryKey Cache key for the geometry we are looking for
+	 */
+	bool HasGeometryInCache(uint32 GeometryKey);
+	bool HasGeometryInCache_AssumesLocked(uint32 GeometryKey) const;
+
 	/** Creates a Dynamic Mesh for the provided Implicit object and generator, and then caches it to be reused later
 	 * @param GeometryCacheKey Key to be used to find this geometry in the cache
 	 * @param MeshGenerator Generator class with the data and rules to create the mesh
@@ -173,23 +182,23 @@ private:
 	TMap<uint32, TObjectPtr<UStaticMesh>> StaticMeshCacheMap;
 
 	/** Map containing all the meshes component waiting for geometry, by geometry key*/
-	TMap<uint32, TArray<TWeakObjectPtr<UMeshComponent>>> MeshComponentsWaitingForGeometryByKey;
+	FChaosVDWaitListMeshMap MeshComponentsWaitingForGeometryByKey;
 
 	/** Set of all geometry keys of the Meshes that are being generated but not ready yet */
 	TSet<uint32> GeometryBeingGeneratedByKey;
 	
-	/** Used to locks Read or Writes to the Geometry cache and in flight job tracking containers */
-	FRWLock RWLock;
+	/** Used to lock Read or Writes to the Geometry cache and in flight job tracking containers */
+	FRWLock GeometryCacheRWLock;
+
+	/** Used to lock Read or Writes to the Geometry Waiting list */
+	FRWLock GeometryWaitListRWLock;
 
 	/** Handle to the ticker used to ticker the Geometry Builder in the game thread*/
 	FTSTicker::FDelegateHandle GameThreadTickDelegate;
 
-	/** Queue of geometry keys already generated and waiting to be applied */
-	TQueue<uint32, EQueueMode::Mpsc> GeometryReadyToApplyQueue;
-
 	FAsyncCompilationNotification GeometryGenerationNotification;
 
-	FThreadSafeBool bHasPendingJobs = false;
+	std::atomic<int32> MeshComponentsWaitingGeometryNum = 0;
 
 	friend class FGeometryGenerationTask;
 
@@ -259,7 +268,7 @@ void FChaosVDGeometryBuilder::CreateMeshComponentsFromImplicit(const Chaos::FImp
 	}
 
 	ComponentType* MeshComponent = nullptr;
-	MeshType* Mesh = nullptr;
+
 	switch (InnerType)
 	{
 		case ImplicitObjectType::Sphere:
@@ -271,14 +280,9 @@ void FChaosVDGeometryBuilder::CreateMeshComponentsFromImplicit(const Chaos::FImp
 			const uint32 GeometryKey = Sphere->GetTypeHash();
 				
 			MeshComponent = CreateMeshComponent<ComponentType>(Owner, Name, Transform, GeometryKey);
+			RegisterMeshComponentWaitingForGeometry(GeometryKey, MeshComponent, DesiredLODCount);
 
-			Mesh = GetCachedMeshForImplicit<MeshType>(GeometryKey);
-
-			if (Mesh)
-			{
-				ApplyMeshToComponentFromKey(MeshComponent, GeometryKey);
-			}
-			else
+			if (!HasGeometryInCache(GeometryKey))
 			{
 				TSharedPtr<UE::Geometry::FSphereGenerator> SphereGen = MakeShared<UE::Geometry::FSphereGenerator>();
 				SphereGen->Radius = Sphere->GetRadius();
@@ -286,7 +290,6 @@ void FChaosVDGeometryBuilder::CreateMeshComponentsFromImplicit(const Chaos::FImp
 				SphereGen->NumPhi = 50;
 				SphereGen->bPolygroupPerQuad = false;
 
-				RegisterMeshComponentWaitingForGeometry(GeometryKey, MeshComponent, DesiredLODCount);
 				DispatchCreateAndCacheMeshForImplicitAsync<MeshType>(GeometryKey, SphereGen);
 			}
 
@@ -301,22 +304,15 @@ void FChaosVDGeometryBuilder::CreateMeshComponentsFromImplicit(const Chaos::FImp
 			const uint32 GeometryKey = Box->GetTypeHash();
 
 			MeshComponent = CreateMeshComponent<ComponentType>(Owner, Name, Transform, GeometryKey);
+			RegisterMeshComponentWaitingForGeometry(GeometryKey, MeshComponent, DesiredLODCount);
 
-			Mesh = GetCachedMeshForImplicit<MeshType>(GeometryKey);
-
-			if (Mesh)
-			{
-				ApplyMeshToComponentFromKey(MeshComponent, GeometryKey);
-			}
-			else
+			if (!HasGeometryInCache(GeometryKey))
 			{
 				TSharedPtr<UE::Geometry::FMinimalBoxMeshGenerator> BoxGen = MakeShared<UE::Geometry::FMinimalBoxMeshGenerator>();
 				UE::Geometry::FOrientedBox3d OrientedBox;
 				OrientedBox.Frame = UE::Geometry::FFrame3d(Box->Center());
 				OrientedBox.Extents = Box->Extents() * 0.5;
 				BoxGen->Box = OrientedBox;
-
-				RegisterMeshComponentWaitingForGeometry(GeometryKey, MeshComponent, DesiredLODCount);
 				DispatchCreateAndCacheMeshForImplicitAsync<MeshType>(GeometryKey, BoxGen);
 			}
 
@@ -343,13 +339,9 @@ void FChaosVDGeometryBuilder::CreateMeshComponentsFromImplicit(const Chaos::FImp
 			MeshComponent->SetRelativeLocation(FinalLocation);
 			MeshComponent->SetRelativeScale3D(Transform.GetScale3D());
 
-			Mesh = GetCachedMeshForImplicit<MeshType>(GeometryKey);
+			RegisterMeshComponentWaitingForGeometry(GeometryKey, MeshComponent, DesiredLODCount);
 
-			if (Mesh)
-			{
-				ApplyMeshToComponentFromKey(MeshComponent, GeometryKey);
-			}
-			else
+			if (!HasGeometryInCache(GeometryKey))
 			{
 				TSharedPtr<UE::Geometry::FCapsuleGenerator> CapsuleGenerator = MakeShared<UE::Geometry::FCapsuleGenerator>();
 				CapsuleGenerator->Radius = FMath::Max(FMathf::ZeroTolerance, Capsule->GetRadius());
@@ -357,8 +349,7 @@ void FChaosVDGeometryBuilder::CreateMeshComponentsFromImplicit(const Chaos::FImp
 				CapsuleGenerator->NumHemisphereArcSteps = 12;
 				CapsuleGenerator->NumCircleSteps = 12;
 				CapsuleGenerator->bPolygroupPerQuad = false;
-
-				RegisterMeshComponentWaitingForGeometry(GeometryKey, MeshComponent, DesiredLODCount);
+	
 				DispatchCreateAndCacheMeshForImplicitAsync<MeshType>(GeometryKey, CapsuleGenerator);
 			}
 			break;
@@ -375,25 +366,19 @@ void FChaosVDGeometryBuilder::CreateMeshComponentsFromImplicit(const Chaos::FImp
 			{
 				const FString Name = FString::Format(TEXT("{0} - {1}"), {TEXT("Convex"), FString::FromInt(Index)});
 
-				
 				uint32 DataComponentKey = InImplicitObject->GetTypeHash();
 				MeshComponent = CreateMeshComponent<ComponentType>(Owner, Name, Transform, DataComponentKey);
 
 				// For the Cache key, we need the hash of the geometry itself without scale or transformations, because these will be applied to the
 				// mesh component. For example, we don't want to generate two meshes for a box, because in one instance was scaled. We only one mesh for one box, and scale the component as needed
 				const uint32 GeometryKey = Convex->GetTypeHash();
-				Mesh = GetCachedMeshForImplicit<MeshType>(GeometryKey);
-	
-				if (Mesh)
-				{
-					ApplyMeshToComponentFromKey(MeshComponent, GeometryKey);
-				}
-				else
+				RegisterMeshComponentWaitingForGeometry(GeometryKey, MeshComponent, DesiredLODCount);
+		
+				if (!HasGeometryInCache(GeometryKey))
 				{
 					TSharedPtr<FChaosVDConvexMeshGenerator> ConvexMeshGen = MakeShared<FChaosVDConvexMeshGenerator>();
 					ConvexMeshGen->GenerateFromConvex(*Convex);
 
-					RegisterMeshComponentWaitingForGeometry(GeometryKey, MeshComponent, DesiredLODCount);
 					DispatchCreateAndCacheMeshForImplicitAsync<MeshType>(GeometryKey ,ConvexMeshGen);
 				}
 			}
@@ -417,19 +402,14 @@ void FChaosVDGeometryBuilder::CreateMeshComponentsFromImplicit(const Chaos::FImp
 				// For the Cache key, we need the hash of the geometry itself without scale or transformations, because these will be applied to the
 				// mesh component. For example, we don't want to generate two meshes for a box, because in one instance was scaled. We only one mesh for one box, and scale the component as needed
 				const uint32 GeometryKey = TriangleMesh->GetTypeHash();
-				Mesh = GetCachedMeshForImplicit<MeshType>(GeometryKey);
+				RegisterMeshComponentWaitingForGeometry(GeometryKey, MeshComponent, DesiredLODCount);
 	
-				if (Mesh)
-				{
-					ApplyMeshToComponentFromKey(MeshComponent, GeometryKey);
-				}
-				else
+				if (!HasGeometryInCache(GeometryKey))
 				{
 					TSharedPtr<FChaosVDTriMeshGenerator> TriMeshGen = MakeShared<FChaosVDTriMeshGenerator>();
 					TriMeshGen->bReverseOrientation = true;
 					TriMeshGen->GenerateFromTriMesh(*TriangleMesh);
 
-					RegisterMeshComponentWaitingForGeometry(GeometryKey, MeshComponent, DesiredLODCount);
 					DispatchCreateAndCacheMeshForImplicitAsync<MeshType>(GeometryKey, TriMeshGen);
 				}
 			}
@@ -449,19 +429,14 @@ void FChaosVDGeometryBuilder::CreateMeshComponentsFromImplicit(const Chaos::FImp
 				// For the Cache key, we need the hash of the geometry itself without scale or transformations, because these will be applied to the
 				// mesh component. For example, we don't want to generate two meshes for a box, because in one instance was scaled. We only one mesh for one box, and scale the component as needed
 				const uint32 GeometryKey = HeightField->GetTypeHash();
-				Mesh = GetCachedMeshForImplicit<MeshType>(GeometryKey);
-	
-				if (Mesh)
-				{
-					ApplyMeshToComponentFromKey(MeshComponent, GeometryKey);
-				}
-				else
+				RegisterMeshComponentWaitingForGeometry(GeometryKey, MeshComponent, DesiredLODCount);
+
+				if (!HasGeometryInCache(GeometryKey))
 				{
 					TSharedPtr<FChaosVDHeightFieldMeshGenerator> HeightFieldMeshGen = MakeShared<FChaosVDHeightFieldMeshGenerator>();
 					HeightFieldMeshGen->bReverseOrientation = false;
 					HeightFieldMeshGen->GenerateFromHeightField(*HeightField);
 
-					RegisterMeshComponentWaitingForGeometry(GeometryKey, MeshComponent, DesiredLODCount);
 					DispatchCreateAndCacheMeshForImplicitAsync<MeshType>(GeometryKey, HeightFieldMeshGen, DesiredLODCount);
 				}
 
@@ -540,7 +515,7 @@ template <typename MeshType>
 void FChaosVDGeometryBuilder::DispatchCreateAndCacheMeshForImplicitAsync(const uint32 GeometryKey, TSharedPtr<UE::Geometry::FMeshShapeGenerator> MeshGenerator, const int32 LODsToGenerateNum)
 {
 	{
-		FReadScopeLock ReadLock(RWLock);
+		FReadScopeLock ReadLock(GeometryCacheRWLock);
 		if (GeometryBeingGeneratedByKey.Contains(GeometryKey))
 		{
 			return;
@@ -548,7 +523,7 @@ void FChaosVDGeometryBuilder::DispatchCreateAndCacheMeshForImplicitAsync(const u
 	}
 
 	{
-		FWriteScopeLock WriteLock(RWLock);
+		FWriteScopeLock WriteLock(GeometryCacheRWLock);
 		GeometryBeingGeneratedByKey.Add(GeometryKey);
 	}
 	
@@ -616,18 +591,14 @@ void FGeometryGenerationTask::GenerateGeometry()
 		{
 			ensureMsgf(LODsToGenerateNum == 0, TEXT("LOD Generation is only suppoted with static meshes | [%d] LODs were requested for a dynamic mesh when 0 is expected"), LODsToGenerateNum);
 			BuilderPtr->CreateAndCacheDynamicMesh(GeometryKey, *MeshGenerator.Get());
-
-			BuilderPtr->GeometryReadyToApplyQueue.Enqueue(GeometryKey);
 		}
 		else if constexpr (std::is_same_v<MeshType, UStaticMesh>)
 		{
 			BuilderPtr->CreateAndCacheStaticMesh(GeometryKey, *MeshGenerator.Get(), LODsToGenerateNum);
-
-			BuilderPtr->GeometryReadyToApplyQueue.Enqueue(GeometryKey);
 		}
 
 		{
-			FWriteScopeLock WriteLock(BuilderPtr->RWLock);
+			FWriteScopeLock WriteLock(BuilderPtr->GeometryCacheRWLock);
 			BuilderPtr->GeometryBeingGeneratedByKey.Remove(GeometryKey);
 		}
 	}
diff --git a/Engine/Plugins/Experimental/ChaosVD/Source/ChaosVD/Public/ChaosVDParticleActor.h b/Engine/Plugins/Experimental/ChaosVD/Source/ChaosVD/Public/ChaosVDParticleActor.h
--- a/Engine/Plugins/Experimental/ChaosVD/Source/ChaosVD/Public/ChaosVDParticleActor.h
+++ b/Engine/Plugins/Experimental/ChaosVD/Source/ChaosVD/Public/ChaosVDParticleActor.h
@@ -66,6 +66,14 @@ public:
 
 	void UpdateGeometryComponentsVisibility();
 
+	/** Changes the active state of this CVD Particle Actor */
+	void SetIsActive(bool bNewActive);
+
+	/** Returns true if this particle actor is active - Inactive Particle actors are still in the world but with outdated data
+	 * and hidden from the viewport and outliner. They represent particles that were destroyed.
+	 */
+	bool IsActive() const { return bIsActive; }
+
 protected:
 
 	void UpdateShapeDataComponents();
@@ -93,4 +101,6 @@ protected:
 	FDelegateHandle GeometryUpdatedDelegate;
 
 	TMap<FStringView, TUniquePtr<FChaosVDDataVisualizerBase>> CVDVisualizers;
+
+	bool bIsActive = false;
 };
diff --git a/Engine/Plugins/Experimental/ChaosVD/Source/ChaosVD/Public/ChaosVDPlaybackController.h b/Engine/Plugins/Experimental/ChaosVD/Source/ChaosVD/Public/ChaosVDPlaybackController.h
--- a/Engine/Plugins/Experimental/ChaosVD/Source/ChaosVD/Public/ChaosVDPlaybackController.h
+++ b/Engine/Plugins/Experimental/ChaosVD/Source/ChaosVD/Public/ChaosVDPlaybackController.h
@@ -46,6 +46,12 @@ struct FChaosVDQueuedTrackInfoUpdate
 	FGuid InstigatorID;
 };
 
+struct FChaosVDGeometryDataUpdate
+{
+	Chaos::FConstImplicitObjectPtr NewGeometry;
+	uint32 GeometryID;
+};
+
 /** Flags used to control how the unload of a recording is performed */
 enum class EChaosVDUnloadRecordingFlags : uint8
 {
@@ -223,20 +229,23 @@ protected:
 	void UpdateSolverTracksData();
 
 	/** Updates the controlled scene with the loaded data at specified game frame */
-	void GoToRecordedGameFrame_AssumesLocked(int32 FrameNumber, FGuid InstigatorID);
+	void GoToRecordedGameFrame_AssumesLocked(int32 FrameNumber, FGuid InstigatorID, int32 Attempts = 0);
 
 	/** Updates the controlled scene with the loaded data at specified solver frame and solver step */
-	void GoToRecordedSolverStep_AssumesLocked(int32 InTrackID, int32 FrameNumber, int32 Step, FGuid InstigatorID);
+	void GoToRecordedSolverStep_AssumesLocked(int32 InTrackID, int32 FrameNumber, int32 Step, FGuid InstigatorID, int32 Attempts = 0);
 
 	/** Handles any data changes on the loaded recording - Usually called during Trace analysis */
 	void HandleCurrentRecordingUpdated();
 
 	/** Finds the closest Key frame to the provided frame number, and plays all the following frames until the specified frame number (no inclusive) */
-	void PlayFromClosestKeyFrame_AssumesLocked(int32 InTrackID, int32 FrameNumber, FChaosVDScene& InSceneToControl);
+	void PlayFromClosestKeyFrame_AssumesLocked(int32 InTrackID, int32 FrameNumber, FChaosVDScene& InSceneToControl) const;
 
 	/** Add the provided track info update to the queue. The update will be broadcast in the game thread */
 	void EnqueueTrackInfoUpdate(const FChaosVDTrackInfo& InTrackInfo, FGuid InstigatorID);
-	
+
+	/** Add the provided Geometry info data to the queue. The update will be broadcast in the game thread */
+	void EnqueueGeometryDataUpdate(const Chaos::FConstImplicitObjectPtr& NewGeometry, const uint32 GeometryID);
+
 	/** Map containing all track info, by track type*/
 	TMap<EChaosVDTrackType, TrackInfoByIDMap> TrackInfoPerType;
 
@@ -258,6 +267,9 @@ protected:
 	/** Queue with a copy of all Track Info Updates that needs to be done in the Game thread */
 	TQueue<FChaosVDQueuedTrackInfoUpdate, EQueueMode::Mpsc> TrackInfoUpdateGTQueue;
 
+	/** Queue with a all the new geometry data that needs to be processed in the Game thread */
+	TQueue<FChaosVDGeometryDataUpdate, EQueueMode::Mpsc> GeometryDataUpdateGTQueue;
+
 	bool bPlayedFirstFrame = false;
 
 	int32 MaxFramesLaggingBehindDuringLiveSession = 50;
diff --git a/Engine/Plugins/Experimental/ChaosVD/Source/ChaosVD/Public/ChaosVDScene.h b/Engine/Plugins/Experimental/ChaosVD/Source/ChaosVD/Public/ChaosVDScene.h
--- a/Engine/Plugins/Experimental/ChaosVD/Source/ChaosVD/Public/ChaosVDScene.h
+++ b/Engine/Plugins/Experimental/ChaosVD/Source/ChaosVD/Public/ChaosVDScene.h
@@ -24,6 +24,7 @@ struct FTypedElementHandle;
 typedef TMap<int32, AChaosVDParticleActor*> FChaosVDParticlesByIDMap;
 
 DECLARE_MULTICAST_DELEGATE(FChaosVDSceneUpdatedDelegate)
+DECLARE_MULTICAST_DELEGATE_OneParam(FChaosVDActorActiveStateUpdateDelegate, AChaosVDParticleActor*)
 DECLARE_MULTICAST_DELEGATE_OneParam(FChaosVDOnObjectSelectedDelegate, UObject*)
 
 /** Recreates a UWorld from a recorded Chaos VD Frame */
@@ -94,6 +95,8 @@ public:
 
 	AActor* GetSkySphereActor() const { return SkySphere; }
 
+	FChaosVDActorActiveStateUpdateDelegate& OnActorActiveStateChanged() { return ParticleActorUpdateDelegate; }
+
 	TSharedPtr<FChaosVDRecording> LoadedRecording;
 
 private:
@@ -154,4 +157,6 @@ private:
 	mutable AActor* SkySphere = nullptr;
 
 	bool bIsInitialized = false;
+
+	FChaosVDActorActiveStateUpdateDelegate ParticleActorUpdateDelegate;
 };
diff --git a/Engine/Plugins/Experimental/ChaosVD/Source/ChaosVD/Public/Trace/ChaosVDTraceProvider.h b/Engine/Plugins/Experimental/ChaosVD/Source/ChaosVD/Public/Trace/ChaosVDTraceProvider.h
--- a/Engine/Plugins/Experimental/ChaosVD/Source/ChaosVD/Public/Trace/ChaosVDTraceProvider.h
+++ b/Engine/Plugins/Experimental/ChaosVD/Source/ChaosVD/Public/Trace/ChaosVDTraceProvider.h
@@ -1,15 +1,17 @@
 // Copyright Epic Games, Inc. All Rights Reserved.
 
 #pragma once
-#include "TraceServices/Model/AnalysisSession.h"
-#include "Templates/SharedPointer.h"
-#include "Chaos/ParticleHandleFwd.h"
-#include "Delegates/Delegate.h"
-#include "Containers/UnrealString.h"
-#include "Containers/Array.h"
 
 #include "Chaos/ChaosArchive.h"
+#include "Chaos/ParticleHandleFwd.h"
+#include "ChaosVDRecording.h"
+#include "Containers/Array.h"
+#include "Containers/Queue.h"
+#include "Containers/UnrealString.h"
 #include "DataProcessors/IChaosVDDataProcessor.h"
+#include "Delegates/Delegate.h"
+#include "Templates/SharedPointer.h"
+#include "TraceServices/Model/AnalysisSession.h"
 
 struct FChaosVDGameFrameData;
 class FChaosVDEngine;
@@ -52,12 +54,13 @@ public:
 
 	void CreateRecordingInstanceForSession(const FString& InSessionName);
 	void DeleteRecordingInstanceForSession();
-	void AddSolverFrame(const int32 InSolverGUID, FChaosVDSolverFrameData&& FrameData);
-	void AddGameFrame(FChaosVDGameFrameData&& FrameData);
-	FChaosVDSolverFrameData* GetSolverFrame_AssumesLocked(const int32 InSolverGUID, const int32 FrameNumber) const;
-	FChaosVDSolverFrameData* GetLastSolverFrame_AssumesLocked(const int32 InSolverGUID) const;
+	void StartSolverFrame(const int32 InSolverGUID, FChaosVDSolverFrameData&& FrameData);
+	void CommitProcessedGameFramesToRecording();
+	void StartGameFrame(const TSharedPtr<FChaosVDGameFrameData>& InFrameData);
+
+	FChaosVDSolverFrameData* GetCurrentSolverFrame(const int32 InSolverGUID);
 	
-	FChaosVDGameFrameData* GetLastGameFrame_AssumesLocked() const;
+	TWeakPtr<FChaosVDGameFrameData> GetCurrentGameFrame();
 
 	FChaosVDBinaryDataContainer& FindOrAddUnprocessedData(const int32 DataID);
 
@@ -67,11 +70,12 @@ public:
 
 	void RegisterDataProcessor(TSharedPtr<IChaosVDDataProcessor> InDataProcessor);
 
-	FRWLock& GetDataLock();
-
 private:
 
 	void RegisterDefaultDataProcessorsIfNeeded();
+
+	void EnqueueGameFrameForProcessing(const TSharedPtr<FChaosVDGameFrameData>& FrameData);
+	void DeQueueGameFrameForProcessing(TSharedPtr<FChaosVDGameFrameData>& OutFrameData);
 	
 	TraceServices::IAnalysisSession& Session;
 
@@ -81,5 +85,13 @@ private:
 
 	TMap<FStringView, TSharedPtr<IChaosVDDataProcessor>> RegisteredDataProcessors;
 
+	TMap<int32, FChaosVDSolverFrameData> CurrentSolverFramesByID;
+
+	TQueue<TSharedPtr<FChaosVDGameFrameData>> CurrentGameFrameQueue;
+
+	TWeakPtr<FChaosVDGameFrameData> CurrentGameFrame = nullptr;
+
+	int32 CurrentGameFrameQueueSize = 0;
+
 	bool bDefaultDataProcessorsRegistered = false;
 };
diff --git a/Engine/Source/Runtime/Experimental/ChaosVDData/Private/ChaosVDRecording.cpp b/Engine/Source/Runtime/Experimental/ChaosVDData/Private/ChaosVDRecording.cpp
--- a/Engine/Source/Runtime/Experimental/ChaosVDData/Private/ChaosVDRecording.cpp
+++ b/Engine/Source/Runtime/Experimental/ChaosVDData/Private/ChaosVDRecording.cpp
@@ -192,9 +192,9 @@ void FChaosVDRecording::AddFrameForSolver(const int32 SolverID, FChaosVDSolverFr
 			SolverFrames->Add(MoveTemp(InFrameData));	
 		}
 		else
-		{	
+		{
 			FrameNumber = 0;
-			RecordedFramesDataPerSolver.Add(SolverID, { MoveTemp(InFrameData) });
+			RecordedFramesDataPerSolver.Add(SolverID, {}).Emplace(MoveTemp(InFrameData));
 		}
 	}
 
@@ -206,10 +206,10 @@ void FChaosVDRecording::AddFrameForSolver(const int32 SolverID, FChaosVDSolverFr
 	OnRecordingUpdated().Broadcast();
 }
 
-void FChaosVDRecording::AddGameFrameData(FChaosVDGameFrameData&& InFrameData)
+void FChaosVDRecording::AddGameFrameData(const FChaosVDGameFrameData& InFrameData)
 {
 	FWriteScopeLock WriteLock(RecordingDataLock);
-	GameFrames.Add(MoveTemp(InFrameData));
+	GameFrames.Add(InFrameData);
 }
 
 FChaosVDGameFrameData* FChaosVDRecording::GetGameFrameDataAtCycle_AssumesLocked(uint64 Cycle)
@@ -254,8 +254,17 @@ void FChaosVDRecording::GetAvailableSolverIDsAtGameFrameNumber_AssumesLocked(int
 		return;
 	}
 	
-	FChaosVDGameFrameData& FrameData = GameFrames[FrameNumber];
+	GetAvailableSolverIDsAtGameFrame_AssumesLocked(GameFrames[FrameNumber], OutSolversID);
+}
+
+void FChaosVDRecording::GetAvailableSolverIDsAtGameFrame(const FChaosVDGameFrameData& GameFrameData, TArray<int32>& OutSolversID)
+{
+	FReadScopeLock ReadLock(RecordingDataLock);
+	GetAvailableSolverIDsAtGameFrame_AssumesLocked(GameFrameData, OutSolversID);
+}
 
+void FChaosVDRecording::GetAvailableSolverIDsAtGameFrame_AssumesLocked(const FChaosVDGameFrameData& GameFrameData, TArray<int32>& OutSolversID)
+{
 	OutSolversID.Reserve(RecordedFramesDataPerSolver.Num());
 
 	for (const TPair<int32, TArray<FChaosVDSolverFrameData>>& SolverFramesWithIDPair : RecordedFramesDataPerSolver)
@@ -265,13 +274,13 @@ void FChaosVDRecording::GetAvailableSolverIDsAtGameFrameNumber_AssumesLocked(int
 			continue;
 		}
 
-		if (SolverFramesWithIDPair.Value.Num() == 1 && SolverFramesWithIDPair.Value[0].FrameCycle < FrameData.FirstCycle)
+		if (SolverFramesWithIDPair.Value.Num() == 1 && SolverFramesWithIDPair.Value[0].FrameCycle < GameFrameData.FirstCycle)
 		{
 			OutSolversID.Add(SolverFramesWithIDPair.Key);
 		}
 		else
 		{
-			if (FrameData.FirstCycle > SolverFramesWithIDPair.Value[0].FrameCycle && FrameData.FirstCycle < SolverFramesWithIDPair.Value.Last().FrameCycle)
+			if (GameFrameData.FirstCycle > SolverFramesWithIDPair.Value[0].FrameCycle && GameFrameData.FirstCycle < SolverFramesWithIDPair.Value.Last().FrameCycle)
 			{
 				OutSolversID.Add(SolverFramesWithIDPair.Key);
 			}
@@ -279,6 +288,58 @@ void FChaosVDRecording::GetAvailableSolverIDsAtGameFrameNumber_AssumesLocked(int
 	}
 }
 
+void FChaosVDRecording::CollapseSolverFramesRange_AssumesLocked(int32 SolverID, int32 StartFrame, int32 EndFrame, FChaosVDSolverFrameData& OutCollapsedFrameData)
+{
+	TMap<int32, FChaosVDParticleDataWrapper> TempParticleDataByID;
+	
+	for (int32 CurrentFrameNumber = StartFrame; CurrentFrameNumber <= EndFrame; CurrentFrameNumber++)
+	{
+		if (const FChaosVDSolverFrameData* SolverFrameData = GetSolverFrameData_AssumesLocked(SolverID, CurrentFrameNumber))
+		{
+			OutCollapsedFrameData.ParticlesDestroyedIDs.Append(SolverFrameData->ParticlesDestroyedIDs);
+
+			// Only evaluate the last step as it contains all the particles that changed
+			const int32 LastStepNumber = SolverFrameData->SolverSteps.Num() - 1;
+			if (SolverFrameData->SolverSteps.IsValidIndex(LastStepNumber))
+			{
+				for (const FChaosVDStepData& StepData : SolverFrameData->SolverSteps)
+				{
+					for (const FChaosVDParticleDataWrapper& ParticleData : StepData.RecordedParticlesData)
+					{
+						if (FChaosVDParticleDataWrapper* FoundParticleData = TempParticleDataByID.Find(ParticleData.ParticleIndex))
+						{
+							(*FoundParticleData) = ParticleData;
+						}
+						else
+						{
+							TempParticleDataByID.Add(ParticleData.ParticleIndex, ParticleData);
+						}
+					}
+				}
+			}
+
+			// If this is the end frame, Copy all the "metadata" for the generated frame, and generate the Solver step
+			if (CurrentFrameNumber == EndFrame)
+			{
+				OutCollapsedFrameData.EndTime = SolverFrameData->EndTime;
+				OutCollapsedFrameData.StartTime = SolverFrameData->StartTime;
+				OutCollapsedFrameData.FrameCycle = SolverFrameData->FrameCycle;
+				OutCollapsedFrameData.bIsKeyFrame = true;
+				OutCollapsedFrameData.SolverID = SolverFrameData->SolverID;
+				OutCollapsedFrameData.SimulationTransform = SolverFrameData->SimulationTransform;
+				OutCollapsedFrameData.DebugName = SolverFrameData->DebugName;
+
+				FChaosVDStepData CollapsedStepData;
+				CollapsedStepData.StepName = TEXT("GeneratedStep");
+
+				TempParticleDataByID.GenerateValueArray(CollapsedStepData.RecordedParticlesData);
+				
+				OutCollapsedFrameData.SolverSteps.Add(MoveTemp(CollapsedStepData));
+			}
+		}
+	}
+}
+
 void FChaosVDRecording::AddImplicitObject(const uint32 ID, const Chaos::FImplicitObjectPtr& InImplicitObject)
 {
 	FWriteScopeLock WriteLock(RecordingDataLock);
diff --git a/Engine/Source/Runtime/Experimental/ChaosVDData/Public/ChaosVDRecording.h b/Engine/Source/Runtime/Experimental/ChaosVDData/Public/ChaosVDRecording.h
--- a/Engine/Source/Runtime/Experimental/ChaosVDData/Public/ChaosVDRecording.h
+++ b/Engine/Source/Runtime/Experimental/ChaosVDData/Public/ChaosVDRecording.h
@@ -174,7 +174,7 @@ struct CHAOSVDDATA_API FChaosVDRecording
 	 * Adds a Game Frame Data entry. Creates a solver entry if it does not exist 
 	 * @param InFrameData Reference to the frame data we want to add
 	 */
-	void AddGameFrameData(FChaosVDGameFrameData&& InFrameData);
+	void AddGameFrameData(const FChaosVDGameFrameData& InFrameData);
 
 	/** Called each time the recording changes - Mainly when a new frame is added from the Trace analysis */
 	FChaosVDRecordingUpdated& OnRecordingUpdated() { return RecordingUpdatedDelegate; };
@@ -213,6 +213,11 @@ struct CHAOSVDDATA_API FChaosVDRecording
      */
 	void GetAvailableSolverIDsAtGameFrameNumber(int32 FrameNumber, TArray<int32>& OutSolversID);
 	void GetAvailableSolverIDsAtGameFrameNumber_AssumesLocked(int32 FrameNumber, TArray<int32>& OutSolversID);
+	void GetAvailableSolverIDsAtGameFrame(const FChaosVDGameFrameData& GameFrameData, TArray<int32>& OutSolversID);
+	void GetAvailableSolverIDsAtGameFrame_AssumesLocked(const FChaosVDGameFrameData& GameFrameData, TArray<int32>& OutSolversID);
+
+	/** Collapses the most important frame data from a range of solver frames into a single solver frame data */
+	void CollapseSolverFramesRange_AssumesLocked(int32 SolverID, int32 StartFrame, int32 EndFrame, FChaosVDSolverFrameData& OutCollapsedFrameData);
 
 	/** Returns a reference to the GeometryID-ImplicitObject map of this recording */
 	const TMap<uint32, Chaos::FConstImplicitObjectPtr>& GetGeometryMap() const { return ImplicitObjects; };
diff --git a/Engine/Source/Runtime/Experimental/ChaosVisualDebugger/Private/DataWrappers/ChaosVDCollisionDataWrappers.cpp b/Engine/Source/Runtime/Experimental/ChaosVisualDebugger/Private/DataWrappers/ChaosVDCollisionDataWrappers.cpp
--- a/Engine/Source/Runtime/Experimental/ChaosVisualDebugger/Private/DataWrappers/ChaosVDCollisionDataWrappers.cpp
+++ b/Engine/Source/Runtime/Experimental/ChaosVisualDebugger/Private/DataWrappers/ChaosVDCollisionDataWrappers.cpp
@@ -105,3 +105,11 @@ bool FChaosVDShapeCollisionData::Serialize(FArchive& Ar)
 
 	return true;
 }
+
+bool FChaosVDShapeCollisionData::operator==(const FChaosVDShapeCollisionData& Other) const
+{
+	return CollisionTraceType == Other.CollisionTraceType
+			&& bSimCollision == Other.bSimCollision
+			&& bQueryCollision == Other.bQueryCollision
+			&& bIsProbe == Other.bIsProbe;
+}
diff --git a/Engine/Source/Runtime/Experimental/ChaosVisualDebugger/Public/DataWrappers/ChaosVDCollisionDataWrappers.h b/Engine/Source/Runtime/Experimental/ChaosVisualDebugger/Public/DataWrappers/ChaosVDCollisionDataWrappers.h
--- a/Engine/Source/Runtime/Experimental/ChaosVisualDebugger/Public/DataWrappers/ChaosVDCollisionDataWrappers.h
+++ b/Engine/Source/Runtime/Experimental/ChaosVisualDebugger/Public/DataWrappers/ChaosVDCollisionDataWrappers.h
@@ -341,6 +341,8 @@ struct CHAOSVDRUNTIME_API FChaosVDShapeCollisionData
 	bool bIsValid = false;
 
 	bool Serialize(FArchive& Ar);
+
+	bool operator==(const FChaosVDShapeCollisionData& Other) const;
 };
 
 template<>
