[ChaosVD] Improvements to how collision data can be inspected + performance and safety fixes.

- Implemented the Tollkit host interface to SChaosVDMainTab, so it can be passed to the Editor Mode Tools.
- Removed Code duplication in SChaosVDMainTab and added a way to get specific tabs if they are loaded.
- Added a way for CVD Tabs to generate an error widget that can be displayed if their setup failed.
- SChaosVDMainTab is now passed as a weak ptr instead of a raw ptr to any CVD tab.
- Removed collision data from the particle actors themselves and created a new Collision Data component per represented solver. This removes a lot of overhead caused by copying collision data to each particle, every frame. Also most that data was duplicated.
- Create a new SolverInfoActor. This is used by FChaosVDScene to centralize all the relevant info per represented solver.
- Re-Implemented collision data visualization as an UE Component Visualizer (and the ground work to be able to add more component visualizers if desired). Particle Data visualization has been kept with the CVD version of the visualizer (which is more limited but lighter).
- Added a way to debug draw (visualize) collision data without needing to have a particle actor selected.
- Created a new "Collision Data Inspector". As the collision data is not part of the CVD particle actor anymore, it cannot be inspected on its details panel. This new collision data inspector offers a dropdown list will all the collision data for the selected object.
- Added a way to click a visualized contact in the viewport and bring its details to the new Collision Data Inspector view.
- Added a way to see all the collision data for a particle in the new collision data inspector by clicking a button on the Particle's details panel.
- Downgraded two ensures to warning logs. There can happen during live debugging and it is not an issue.
- Updated SChaosVDNameListPicker so it auto-disables itself when the list is empty. Also it clearly states it is empty. This widget is used in a few places of the CVD UI.
- Disable navigation relevancy on components generated by CVD. The navigation system is already disabled for the CVD world, but the components having navigation relevancy enabled by default was having an impact because of the checks done when the actors are moved.

diff --git a/Engine/Plugins/Experimental/ChaosVD/Source/ChaosVD/Private/Actors/ChaosVDSolverInfoActor.cpp b/Engine/Plugins/Experimental/ChaosVD/Source/ChaosVD/Private/Actors/ChaosVDSolverInfoActor.cpp
new file mode 100644
--- /dev/null
+++ b/Engine/Plugins/Experimental/ChaosVD/Source/ChaosVD/Private/Actors/ChaosVDSolverInfoActor.cpp
@@ -0,0 +1,82 @@
+// Copyright Epic Games, Inc. All Rights Reserved.
+
+#include "Actors/ChaosVDSolverInfoActor.h"
+
+#include "ChaosVDParticleActor.h"
+#include "ChaosVDScene.h"
+#include "Components/ChaosVDSolverCollisionDataComponent.h"
+#include "Elements/Framework/TypedElementSelectionSet.h"
+#include "Engine/World.h"
+
+AChaosVDSolverInfoActor::AChaosVDSolverInfoActor(const FObjectInitializer& ObjectInitializer) : Super(ObjectInitializer)
+{
+	CollisionDataComponent = CreateDefaultSubobject<UChaosVDSolverCollisionDataComponent>(TEXT("SolverCollisionDataComponent"));
+}
+
+void AChaosVDSolverInfoActor::SetScene(TWeakPtr<FChaosVDScene> InScene)
+{
+	FChaosVDSceneObjectBase::SetScene(InScene);
+
+	if (TSharedPtr<FChaosVDScene> ScenePtr = InScene.Pin())
+	{
+		RegisterSelectionSetObject(ScenePtr->GetElementSelectionSet());
+	}
+}
+
+void AChaosVDSolverInfoActor::RegisterParticleActor(int32 ParticleID, AChaosVDParticleActor* ParticleActor)
+{
+	if (!SolverParticlesByID.Contains(ParticleID))
+	{
+		SolverParticlesByID.Add(ParticleID, ParticleActor);
+	}
+}
+
+AChaosVDParticleActor* AChaosVDSolverInfoActor::GetParticleActor(int32 ParticleID)
+{
+	AChaosVDParticleActor** FoundParticleActor = SolverParticlesByID.Find(ParticleID);
+
+	return FoundParticleActor ? *FoundParticleActor : nullptr;
+}
+
+void AChaosVDSolverInfoActor::HandleVisibilitySettingsUpdated()
+{
+	for (const TPair<int32, AChaosVDParticleActor*>& ParticleWithIDPair : SolverParticlesByID)
+	{
+		ParticleWithIDPair.Value->UpdateGeometryComponentsVisibility();
+	}
+}
+
+void AChaosVDSolverInfoActor::BeginDestroy()
+{
+	Super::BeginDestroy();
+
+	UWorld* World = GetWorld();
+	if (!World)
+	{
+		return;
+	}
+	
+	for (const TPair<int32, AChaosVDParticleActor*>& ParticleVDInstanceWithID : SolverParticlesByID)
+	{
+		World->DestroyActor(ParticleVDInstanceWithID.Value);
+	}
+}
+
+void AChaosVDSolverInfoActor::HandlePostSelectionChange(const UTypedElementSelectionSet* ChangesSelectionSet)
+{
+	SelectedParticlesID.Reset();
+
+	TArray<AChaosVDParticleActor*> SelectedParticles = ChangesSelectionSet->GetSelectedObjects<AChaosVDParticleActor>();
+
+	//TODO: Support multi-selection
+	if (SelectedParticles.Num() > 0)
+	{
+		if (AChaosVDParticleActor* SelectedParticle = SelectedParticles[0])
+		{
+			if (const FChaosVDParticleDataWrapper* ParticleData = SelectedParticle->GetParticleData())
+			{
+				SelectedParticlesID.Add(ParticleData->ParticleIndex);
+			}
+		}
+	}
+}
diff --git a/Engine/Plugins/Experimental/ChaosVD/Source/ChaosVD/Private/ChaosVDCollisionDataDetailsTab.cpp b/Engine/Plugins/Experimental/ChaosVD/Source/ChaosVD/Private/ChaosVDCollisionDataDetailsTab.cpp
new file mode 100644
--- /dev/null
+++ b/Engine/Plugins/Experimental/ChaosVD/Source/ChaosVD/Private/ChaosVDCollisionDataDetailsTab.cpp
@@ -0,0 +1,35 @@
+// Copyright Epic Games, Inc. All Rights Reserved.
+
+#include "ChaosVDCollisionDataDetailsTab.h"
+#include "PropertyEditorModule.h"
+#include "Modules/ModuleManager.h"
+
+#include "IStructureDetailsView.h"
+#include "Widgets/SChaosVDCollisionDataInspector.h"
+
+#define LOCTEXT_NAMESPACE "ChaosVisualDebugger"
+
+FChaosVDCollisionDataDetailsTab::~FChaosVDCollisionDataDetailsTab()
+{
+}
+
+TSharedRef<SDockTab> FChaosVDCollisionDataDetailsTab::HandleTabSpawned(const FSpawnTabArgs& Args)
+{
+
+	TSharedRef<SDockTab> DetailsPanelTab =
+		SNew(SDockTab)
+		.TabRole(ETabRole::MajorTab)
+		.Label(LOCTEXT("CollisionInspectorTab", "Collision Data Inspector"))
+		.ToolTipText(LOCTEXT("CollisionInspectorTabTip", "See the details of the any collision data for the selected object"));
+
+	DetailsPanelTab->SetContent
+	(
+		SAssignNew(CollisionDataInspector, SChaosVDCollisionDataInspector, GetChaosVDScene())
+	);
+
+	OnTabSpawned().Broadcast(DetailsPanelTab);
+
+	return DetailsPanelTab;
+}
+
+#undef LOCTEXT_NAMESPACE
diff --git a/Engine/Plugins/Experimental/ChaosVD/Source/ChaosVD/Private/ChaosVDEditorVisualizationSettingsTab.cpp b/Engine/Plugins/Experimental/ChaosVD/Source/ChaosVD/Private/ChaosVDEditorVisualizationSettingsTab.cpp
--- a/Engine/Plugins/Experimental/ChaosVD/Source/ChaosVD/Private/ChaosVDEditorVisualizationSettingsTab.cpp
+++ b/Engine/Plugins/Experimental/ChaosVD/Source/ChaosVD/Private/ChaosVDEditorVisualizationSettingsTab.cpp
@@ -7,6 +7,7 @@
 #include "ChaosVDPlaybackController.h"
 #include "ChaosVDPlaybackViewportTab.h"
 #include "ChaosVDStyle.h"
+#include "ChaosVDTabsIDs.h"
 #include "DetailsViewArgs.h"
 #include "Framework/Docking/TabManager.h"
 #include "IDetailsView.h"
@@ -22,8 +23,20 @@
 TSharedRef<SDockTab> FChaosVDEditorVisualizationSettingsTab::HandleTabSpawned(const FSpawnTabArgs& Args)
 {
 
+	TSharedRef<SDockTab> DetailsPanelTab = SNew(SDockTab)
+											.TabRole(ETabRole::MajorTab)
+											.Label(LOCTEXT("ChaosVDEditorSettings", "Chaos VD Settings"))
+											.ToolTipText(LOCTEXT("ChaosVDEditorSettingsTip", "See the available settings for the editor"));
+	
+	TSharedPtr<SChaosVDMainTab> CVDMainTabPtr = OwningTabWidget.Pin();
+	if (!CVDMainTabPtr.IsValid())
+	{
+		DetailsPanelTab->SetContent(GenerateErrorWidget());
+		return DetailsPanelTab;
+	}
+	
 	TSharedPtr<SChaosVDPlaybackViewport> ViewportWidget;
-	if (TSharedPtr<FChaosVDPlaybackViewportTab> ViewportTabSharedPtr = OwningTabWidget->GetPlaybackViewportTab().Pin())
+	if (TSharedPtr<FChaosVDPlaybackViewportTab> ViewportTabSharedPtr = CVDMainTabPtr->GetTabSpawnerInstance<FChaosVDPlaybackViewportTab>(FChaosVDTabID::PlaybackViewport).Pin())
 	{
 		ViewportWidget = ViewportTabSharedPtr->GetPlaybackViewportWidget().Pin();
 	}
@@ -34,15 +47,8 @@ TSharedRef<SDockTab> FChaosVDEditorVisualizationSettingsTab::HandleTabSpawned(co
 	DetailsViewArgs.NameAreaSettings = FDetailsViewArgs::HideNameArea;
 
 	TSharedRef<IDetailsView> DetailsPanel = PropertyEditorModule.CreateDetailView(DetailsViewArgs);
-
 	DetailsPanel->SetObject(GetMutableDefault<UChaosVDEditorSettings>());
 
-	TSharedRef<SDockTab> DetailsPanelTab =
-		SNew(SDockTab)
-		.TabRole(ETabRole::MajorTab)
-		.Label(LOCTEXT("ChaosVDEditorSettings", "Chaos VD Settings"))
-		.ToolTipText(LOCTEXT("ChaosVDEditorSettingsTip", "See the available settings for the editor"));
-
 	if (ensure(ViewportWidget.IsValid()))
 	{
 		DetailsPanelTab->SetContent
@@ -50,25 +56,19 @@ TSharedRef<SDockTab> FChaosVDEditorVisualizationSettingsTab::HandleTabSpawned(co
 			SNew(SVerticalBox)
 			+SVerticalBox::Slot()
 			[
-				SNew(SChaosVDVisualizationControls, OwningTabWidget->GetChaosVDEngineInstance()->GetPlaybackController(), ViewportWidget)
+				SNew(SChaosVDVisualizationControls,CVDMainTabPtr->GetChaosVDEngineInstance()->GetPlaybackController(), ViewportWidget)
 			]	
 		);
 	}
 	else
 	{
-		DetailsPanelTab->SetContent
-		(
-			SNew(SVerticalBox)
-			+SVerticalBox::Slot()
-			[
-				SNew(STextBlock)
-				.Text(LOCTEXT("ChaosVDEditorSettingsLoadError", "Failed to load Visualization Controls"))
-			]	
-		);
+		DetailsPanelTab->SetContent(GenerateErrorWidget());
 	}
 
 	DetailsPanelTab->SetTabIcon(FChaosVDStyle::Get().GetBrush("TabIconDetailsPanel"));
 
+	OnTabSpawned().Broadcast(DetailsPanelTab);
+
 	return DetailsPanelTab;
 }
 
diff --git a/Engine/Plugins/Experimental/ChaosVD/Source/ChaosVD/Private/ChaosVDObjectDetailsTab.cpp b/Engine/Plugins/Experimental/ChaosVD/Source/ChaosVD/Private/ChaosVDObjectDetailsTab.cpp
--- a/Engine/Plugins/Experimental/ChaosVD/Source/ChaosVD/Private/ChaosVDObjectDetailsTab.cpp
+++ b/Engine/Plugins/Experimental/ChaosVD/Source/ChaosVD/Private/ChaosVDObjectDetailsTab.cpp
@@ -2,21 +2,28 @@
 
 #include "ChaosVDObjectDetailsTab.h"
 
+#include "ChaosVDCollisionDataDetailsTab.h"
+#include "ChaosVDParticleActor.h"
 #include "ChaosVDScene.h"
 #include "ChaosVDStyle.h"
+#include "ChaosVDTabsIDs.h"
 #include "Editor.h"
 #include "Elements/Framework/TypedElementSelectionSet.h"
 #include "Modules/ModuleManager.h"
 #include "PropertyEditorModule.h"
 #include "Templates/SharedPointer.h"
 #include "Widgets/Docking/SDockTab.h"
+#include "Widgets/Input/SButton.h"
+#include "Widgets/SChaosVDCollisionDataInspector.h"
 #include "Widgets/SChaosVDDetailsView.h"
+#include "Widgets/SChaosVDMainTab.h"
+#include "Widgets/SToolTip.h"
+#include "Widgets/Text/STextBlock.h"
 
 #define LOCTEXT_NAMESPACE "ChaosVisualDebugger"
 
 class SSubobjectEditor;
 
-
 TSharedRef<SDockTab> FChaosVDObjectDetailsTab::HandleTabSpawned(const FSpawnTabArgs& Args)
 {
 	TSharedPtr<FChaosVDScene> ScenePtr = GetChaosVDScene().Pin();
@@ -32,11 +39,22 @@ TSharedRef<SDockTab> FChaosVDObjectDetailsTab::HandleTabSpawned(const FSpawnTabA
 
 	DetailsPanelTab->SetContent
 	(
-		SAssignNew(DetailsPanelView, SChaosVDDetailsView)
+		SNew(SVerticalBox)
+		+SVerticalBox::Slot()
+		[
+			SAssignNew(DetailsPanelView, SChaosVDDetailsView)
+		]
+		+SVerticalBox::Slot()
+		.AutoHeight()
+		[
+			GenerateShowCollisionDataButton().ToSharedRef()
+		]
 	);
 
 	DetailsPanelTab->SetTabIcon(FChaosVDStyle::Get().GetBrush("TabIconDetailsPanel"));
-
+	
+	OnTabSpawned().Broadcast(DetailsPanelTab);
+	
 	return DetailsPanelTab;
 }
 
@@ -49,8 +67,98 @@ void FChaosVDObjectDetailsTab::HandlePostSelectionChange(const UTypedElementSele
 		// We don't support multi selection yet
 		ensure(SelectedActors.Num() == 1);
 
-		DetailsPanelView->SetSelectedObject(SelectedActors[0]);
+		CurrentSelectedObject = SelectedActors[0];
+		DetailsPanelView->SetSelectedObject(CurrentSelectedObject.Get());
+	}
+	else
+	{
+		CurrentSelectedObject = nullptr;
+	}
+}
+
+EVisibility FChaosVDObjectDetailsTab::GetCollisionDataButtonVisibility() const
+{
+	EVisibility DesiredVisibility = EVisibility::Collapsed;
+	if (CurrentSelectedObject.IsValid())
+	{
+		if (Cast<IChaosVDCollisionDataProviderInterface>(CurrentSelectedObject))
+		{
+			DesiredVisibility = EVisibility::Visible;
+		}
+	}
+	return DesiredVisibility;
+}
+
+bool FChaosVDObjectDetailsTab::GetCollisionDataButtonEnabled() const
+{
+	if (IChaosVDCollisionDataProviderInterface* CollisionDataProvider = Cast<IChaosVDCollisionDataProviderInterface>(CurrentSelectedObject.Get()))
+	{
+		return CollisionDataProvider->HasCollisionData();
+	}
+
+	return false;
+}
+
+TSharedPtr<SWidget> FChaosVDObjectDetailsTab::GenerateShowCollisionDataButton()
+{
+	TSharedPtr<SWidget> ShowCollisionButton = SNew(SHorizontalBox)
+	.Visibility_Raw(this, &FChaosVDObjectDetailsTab::GetCollisionDataButtonVisibility)
+	+SHorizontalBox::Slot()
+	.VAlign(VAlign_Center)
+	.HAlign(HAlign_Center)
+	.Padding(12.0f, 7.0f, 12.0f, 7.0f)
+	.FillWidth(1.0f)
+	[
+		SNew(SButton)
+		.ToolTip(SNew(SToolTip).Text(LOCTEXT("OpenCollisionDataDesc", "Click here to open the collision data for this particle on the collision data inspector.")))
+		.IsEnabled_Raw(this, &FChaosVDObjectDetailsTab::GetCollisionDataButtonEnabled)
+		.ContentPadding(FMargin(0, 5.f, 0, 4.f))
+		.OnClicked_Raw(this, &FChaosVDObjectDetailsTab::ShowCollisionDataForSelectedObject)
+		.Content()
+		[
+			SNew(SHorizontalBox)
+			+SHorizontalBox::Slot()
+			.Padding(FMargin(3, 0, 0, 0))
+			.VAlign(VAlign_Center)
+			.AutoWidth()
+			[
+				SNew(STextBlock)
+				.TextStyle(FAppStyle::Get(), "SmallButtonText")
+				.Text(LOCTEXT("ShowCollisionDataOnInspector","Show Collison Data in Inspector"))
+			]
+		]
+	];
+
+	return ShowCollisionButton;
+}
+
+FReply FChaosVDObjectDetailsTab::ShowCollisionDataForSelectedObject()
+{
+	IChaosVDCollisionDataProviderInterface* CollisionDataProvider = Cast<IChaosVDCollisionDataProviderInterface>(CurrentSelectedObject.Get());
+	if (!CollisionDataProvider)
+	{
+		return FReply::Handled();
+	}
+	
+	TSharedPtr<SChaosVDMainTab> OwningTabPtr = OwningTabWidget.Pin();
+	if (!OwningTabPtr.IsValid())
+	{
+		return FReply::Handled();
+	}
+
+	if (TSharedPtr<FChaosVDCollisionDataDetailsTab> CollisionDataTab = OwningTabPtr->GetTabSpawnerInstance<FChaosVDCollisionDataDetailsTab>(FChaosVDTabID::CollisionDataDetails).Pin())
+	{
+		if (TSharedPtr<SChaosVDCollisionDataInspector> CollisionInspector = CollisionDataTab->GetCollisionInspectorInstance().Pin())
+		{
+			if (TSharedPtr<FTabManager> TabManager = OwningTabPtr->GetTabManager())
+			{
+				TabManager->TryInvokeTab(FChaosVDTabID::CollisionDataDetails);
+				CollisionInspector->SetCollisionDataProviderObjectToInspect(CollisionDataProvider);
+			}
+		}
 	}
+	
+	return FReply::Handled();
 }
 
 #undef LOCTEXT_NAMESPACE
diff --git a/Engine/Plugins/Experimental/ChaosVD/Source/ChaosVD/Private/ChaosVDOutputLogTab.cpp b/Engine/Plugins/Experimental/ChaosVD/Source/ChaosVD/Private/ChaosVDOutputLogTab.cpp
--- a/Engine/Plugins/Experimental/ChaosVD/Source/ChaosVD/Private/ChaosVDOutputLogTab.cpp
+++ b/Engine/Plugins/Experimental/ChaosVD/Source/ChaosVD/Private/ChaosVDOutputLogTab.cpp
@@ -27,6 +27,8 @@ TSharedRef<SDockTab> FChaosVDOutputLogTab::HandleTabSpawned(const FSpawnTabArgs&
 		FOutputLogModule::Get().MakeOutputLogWidget(Params)
 	);
 
+	OnTabSpawned().Broadcast(OutputLogTab);
+
 	return OutputLogTab;
 }
 
diff --git a/Engine/Plugins/Experimental/ChaosVD/Source/ChaosVD/Private/ChaosVDParticleActor.cpp b/Engine/Plugins/Experimental/ChaosVD/Source/ChaosVD/Private/ChaosVDParticleActor.cpp
--- a/Engine/Plugins/Experimental/ChaosVD/Source/ChaosVD/Private/ChaosVDParticleActor.cpp
+++ b/Engine/Plugins/Experimental/ChaosVD/Source/ChaosVD/Private/ChaosVDParticleActor.cpp
@@ -5,7 +5,9 @@
 #include "ChaosVDGeometryBuilder.h"
 #include "ChaosVDModule.h"
 #include "ChaosVDScene.h"
+#include "Actors/ChaosVDSolverInfoActor.h"
 #include "Components/ChaosVDInstancedStaticMeshComponent.h"
+#include "Components/ChaosVDSolverCollisionDataComponent.h"
 #include "Components/ChaosVDStaticMeshComponent.h"
 #include "Components/MeshComponent.h"
 #include "Components/SceneComponent.h"
@@ -14,10 +16,15 @@
 #include "Engine/CollisionProfile.h"
 #include "Engine/StaticMesh.h"
 #include "Visualizers/ChaosVDParticleDataVisualizer.h"
+#include "Visualizers/ChaosVDSolverCollisionDataComponentVisualizer.h"
+
 
 AChaosVDParticleActor::AChaosVDParticleActor(const FObjectInitializer& ObjectInitializer) : Super(ObjectInitializer)
 {
 	RootComponent = CreateDefaultSubobject<USceneComponent>(TEXT("SceneComponent0"));
+	RootComponent->SetCanEverAffectNavigation(false);
+	RootComponent->bNavigationRelevant = false;
+
 	CreateVisualizers();
 }
 
@@ -62,30 +69,6 @@ void AChaosVDParticleActor::UpdateFromRecordedParticleData(const FChaosVDParticl
 	}
 }
 
-void AChaosVDParticleActor::UpdateCollisionData(const TArray<TSharedPtr<FChaosVDParticlePairMidPhase>>& InRecordedMidPhases)
-{
-	// TODO: We should store a ptr to the data and in our custom details panel draw it
-	ParticleDataViewer.ParticleMidPhases.Reserve(InRecordedMidPhases.Num());
-	ParticleDataViewer.ParticleMidPhases.Reset(InRecordedMidPhases.Num());
-
-	for (const TSharedPtr<FChaosVDParticlePairMidPhase>& MidPhase : InRecordedMidPhases)
-	{
-		ParticleDataViewer.ParticleMidPhases.Emplace(*MidPhase.Get());
-	}
-}
-
-void AChaosVDParticleActor::UpdateCollisionData(const TArray<FChaosVDConstraint>& InRecordedConstraints)
-{
-	// TODO: We should store a ptr to the data and in our custom details panel draw it
-	ParticleDataViewer.ParticleConstraints.Reserve(InRecordedConstraints.Num());
-	ParticleDataViewer.ParticleConstraints.Reset(InRecordedConstraints.Num());
-
-	for (const FChaosVDConstraint& Constraint : InRecordedConstraints)
-	{
-		ParticleDataViewer.ParticleConstraints.Emplace(Constraint);
-	}
-}
-
 void AChaosVDParticleActor::UpdateGeometry(const Chaos::FConstImplicitObjectPtr& InImplicitObject, EChaosVDActorGeometryUpdateFlags OptionsFlags)
 {
 	if (!InImplicitObject.IsValid())
@@ -113,7 +96,7 @@ void AChaosVDParticleActor::UpdateGeometry(const Chaos::FConstImplicitObjectPtr&
 		return;
 	}
 
-	if (const TSharedPtr<FChaosVDScene>& ScenePtr = OwningScene.Pin())
+	if (const TSharedPtr<FChaosVDScene>& ScenePtr = SceneWeakPtr.Pin())
 	{
 		if (const TSharedPtr<FChaosVDGeometryBuilder>& GeometryGenerator = ScenePtr->GetGeometryGenerator())
 		{
@@ -155,7 +138,7 @@ void AChaosVDParticleActor::UpdateGeometry(const Chaos::FConstImplicitObjectPtr&
 
 void AChaosVDParticleActor::UpdateGeometry(uint32 NewGeometryHash, EChaosVDActorGeometryUpdateFlags OptionsFlags)
 {
-	if (const TSharedPtr<FChaosVDScene>& ScenePtr = OwningScene.Pin())
+	if (const TSharedPtr<FChaosVDScene>& ScenePtr = SceneWeakPtr.Pin())
 	{
 		if (const Chaos::FConstImplicitObjectPtr& Geometry = ScenePtr->GetUpdatedGeometry(NewGeometryHash))
 		{
@@ -164,11 +147,11 @@ void AChaosVDParticleActor::UpdateGeometry(uint32 NewGeometryHash, EChaosVDActor
 	}
 }
 
-void AChaosVDParticleActor::SetScene(const TSharedPtr<FChaosVDScene>& InScene)
+void AChaosVDParticleActor::SetScene(TWeakPtr<FChaosVDScene> InScene)
 {
-	OwningScene = InScene;
+	FChaosVDSceneObjectBase::SetScene(InScene);
 
-	if (const TSharedPtr<FChaosVDScene>& ScenePtr = OwningScene.Pin())
+	if (const TSharedPtr<FChaosVDScene>& ScenePtr = SceneWeakPtr.Pin())
 	{
 		GeometryUpdatedDelegate = ScenePtr->OnNewGeometryAvailable().AddWeakLambda(this, [this](const Chaos::FConstImplicitObjectPtr& ImplicitObject, const uint32 ID)
 		{
@@ -182,7 +165,7 @@ void AChaosVDParticleActor::SetScene(const TSharedPtr<FChaosVDScene>& InScene)
 
 void AChaosVDParticleActor::BeginDestroy()
 {
-	if (const TSharedPtr<FChaosVDScene>& ScenePtr = OwningScene.Pin())
+	if (const TSharedPtr<FChaosVDScene>& ScenePtr = SceneWeakPtr.Pin())
 	{
 		ScenePtr->OnNewGeometryAvailable().Remove(GeometryUpdatedDelegate);
 	}
@@ -193,14 +176,13 @@ void AChaosVDParticleActor::BeginDestroy()
 void AChaosVDParticleActor::GetVisualizationContext(FChaosVDVisualizationContext& OutVisualizationContext)
 {
 	OutVisualizationContext.SpaceTransform = CachedSimulationTransform;
-	OutVisualizationContext.CVDScene = OwningScene;
+	OutVisualizationContext.CVDScene = SceneWeakPtr;
 	OutVisualizationContext.SolverID = ParticleDataViewer.SolverID;
 }
 
 void AChaosVDParticleActor::CreateVisualizers()
 {
 	CVDVisualizers.Add(FChaosVDParticleDataVisualizer::VisualizerID,MakeUnique<FChaosVDParticleDataVisualizer>(*this));
-	CVDVisualizers.Add(FChaosVDCollisionDataVisualizer::VisualizerID, MakeUnique<FChaosVDCollisionDataVisualizer>(*this));
 }
 
 #if WITH_EDITOR
@@ -209,7 +191,7 @@ bool AChaosVDParticleActor::IsSelectedInEditor() const
 {
 	// The implementation of this method in UObject, used a global edit callback,
 	// but as we don't use the global editor selection system, we need to re-route it.
-	if (TSharedPtr<FChaosVDScene> ScenePtr = OwningScene.Pin())
+	if (TSharedPtr<FChaosVDScene> ScenePtr = SceneWeakPtr.Pin())
 	{
 		return ScenePtr->IsObjectSelected(this);
 	}
@@ -231,23 +213,67 @@ void AChaosVDParticleActor::PostEditChangeProperty(FPropertyChangedEvent& Proper
 		}
 	}
 
-	if (PropertyChangedEvent.GetMemberPropertyName() == GET_MEMBER_NAME_STRING_CHECKED(AChaosVDParticleActor, LocalCollisionDataVisualizationFlags))
+	if (PropertyChangedEvent.GetMemberPropertyName() == GET_MEMBER_NAME_STRING_CHECKED(AChaosVDParticleActor, LocalParticleDataVisualizationFlags))
 	{
-		if (TUniquePtr<FChaosVDDataVisualizerBase>* CollisionVisualizer = CVDVisualizers.Find(FChaosVDCollisionDataVisualizer::VisualizerID))
+		if (TUniquePtr<FChaosVDDataVisualizerBase>* ParticleDataVisualizer = CVDVisualizers.Find(FChaosVDParticleDataVisualizer::VisualizerID))
 		{
-			CollisionVisualizer->Get()->UpdateVisualizationFlags(LocalCollisionDataVisualizationFlags);
+			ParticleDataVisualizer->Get()->UpdateVisualizationFlags(LocalParticleDataVisualizationFlags);
 		}
 	}
+}
 
-	if (PropertyChangedEvent.GetMemberPropertyName() == GET_MEMBER_NAME_STRING_CHECKED(AChaosVDParticleActor, LocalParticleDataVisualizationFlags))
+void AChaosVDParticleActor::GetCollisionData(TArray<TSharedPtr<FChaosVDCollisionDataFinder>>& OutCollisionDataFound)
+{
+	if (const TArray<TSharedPtr<FChaosVDParticlePairMidPhase>>* MidPhases = GetCollisionMidPhasesArray())
 	{
-		if (TUniquePtr<FChaosVDDataVisualizerBase>* ParticleDataVisualizer = CVDVisualizers.Find(FChaosVDParticleDataVisualizer::VisualizerID))
+		OutCollisionDataFound.Reserve(MidPhases->Num());
+
+		for (const TSharedPtr<FChaosVDParticlePairMidPhase>& MidPhasePtr : *MidPhases)
 		{
-			ParticleDataVisualizer->Get()->UpdateVisualizationFlags(LocalParticleDataVisualizationFlags);
+			TSharedPtr<FChaosVDCollisionDataFinder> FinderData = MakeShared<FChaosVDCollisionDataFinder>();
+			FinderData->OwningMidPhase = MidPhasePtr;
+			FinderData->OwningConstraint = MidPhasePtr->Constraints.Num() > 0 ? &MidPhasePtr->Constraints[0] : nullptr;
+			FinderData->ContactIndex = INDEX_NONE;
+
+			OutCollisionDataFound.Add(FinderData);
 		}
 	}
 }
 
+bool AChaosVDParticleActor::HasCollisionData()
+{
+	if (const TArray<TSharedPtr<FChaosVDParticlePairMidPhase>>* MidPhases = GetCollisionMidPhasesArray())
+	{
+		return MidPhases->Num() > 0;
+	}
+
+	return false;
+}
+
+FName AChaosVDParticleActor::GetName()
+{
+	return GetFName();
+}
+
+const TArray<TSharedPtr<FChaosVDParticlePairMidPhase>>* AChaosVDParticleActor::GetCollisionMidPhasesArray() const
+{
+	const TSharedPtr<FChaosVDScene> ScenePtr = SceneWeakPtr.Pin();
+	if (!ScenePtr.IsValid())
+	{
+		return nullptr;
+	}
+
+	if (AChaosVDSolverInfoActor* SolverInfoActor = ScenePtr->GetSolverInfoActor(ParticleDataViewer.SolverID))
+	{
+		if (const UChaosVDSolverCollisionDataComponent* CollisionDataComponent = SolverInfoActor->GetCollisionDataComponent())
+		{
+			return CollisionDataComponent->GetMidPhasesForParticle(ParticleDataViewer.ParticleIndex, EChaosVDGetCollisionDataOptions::Any);
+		}
+	}
+
+	return nullptr;
+}
+
 void AChaosVDParticleActor::UpdateShapeDataComponents()
 {
 	for (TWeakObjectPtr<UMeshComponent> MeshComponent : MeshComponents)
@@ -295,7 +321,7 @@ void AChaosVDParticleActor::SetIsActive(bool bNewActive)
 #endif
 		bIsActive = bNewActive;
 
-		if (const TSharedPtr<FChaosVDScene> ScenePtr = OwningScene.Pin())
+		if (const TSharedPtr<FChaosVDScene> ScenePtr = SceneWeakPtr.Pin())
 		{
 			ScenePtr->OnActorActiveStateChanged().Broadcast(this);
 		}
diff --git a/Engine/Plugins/Experimental/ChaosVD/Source/ChaosVD/Private/ChaosVDPlaybackController.cpp b/Engine/Plugins/Experimental/ChaosVD/Source/ChaosVD/Private/ChaosVDPlaybackController.cpp
--- a/Engine/Plugins/Experimental/ChaosVD/Source/ChaosVD/Private/ChaosVDPlaybackController.cpp
+++ b/Engine/Plugins/Experimental/ChaosVD/Source/ChaosVD/Private/ChaosVDPlaybackController.cpp
@@ -103,8 +103,11 @@ void FChaosVDPlaybackController::UnloadCurrentRecording(EChaosVDUnloadRecordingF
 void FChaosVDPlaybackController::PlayFromClosestKeyFrame_AssumesLocked(const int32 InTrackID, const int32 FrameNumber, FChaosVDScene& InSceneToControl) const
 {
 	const int32 KeyFrameNumber = LoadedRecording->FindFirstSolverKeyFrameNumberFromFrame_AssumesLocked(InTrackID, FrameNumber);
-	if (!ensure(KeyFrameNumber >= 0))
+	if (KeyFrameNumber < 0)
 	{
+		// This can happen during live debugging as we miss some of the events at the beginning.
+		// Loading a trace file that was recorded as part of a live session, will have the same issue.
+		UE_LOG(LogChaosVDEditor, Warning, TEXT("[%s] Failed to find a keyframe close to frame [%d] of track [%d]"), ANSI_TO_TCHAR(__FUNCTION__), FrameNumber, InTrackID);
 		return;
 	}
 
diff --git a/Engine/Plugins/Experimental/ChaosVD/Source/ChaosVD/Private/ChaosVDPlaybackViewportClient.cpp b/Engine/Plugins/Experimental/ChaosVD/Source/ChaosVD/Private/ChaosVDPlaybackViewportClient.cpp
--- a/Engine/Plugins/Experimental/ChaosVD/Source/ChaosVD/Private/ChaosVDPlaybackViewportClient.cpp
+++ b/Engine/Plugins/Experimental/ChaosVD/Source/ChaosVD/Private/ChaosVDPlaybackViewportClient.cpp
@@ -3,10 +3,14 @@
 #include "ChaosVDPlaybackViewportClient.h"
 
 #include "ChaosVDEditorSettings.h"
+#include "ChaosVDModule.h"
 #include "ChaosVDParticleActor.h"
 #include "ChaosVDPlaybackController.h"
 #include "ChaosVDScene.h"
 #include "ChaosVDSkySphereInterface.h"
+#include "ChaosVDCollisionDataDetailsTab.h"
+#include "ChaosVDTabsIDs.h"
+#include "ComponentVisualizer.h"
 #include "EditorModeManager.h"
 #include "Elements/Framework/TypedElementSelectionSet.h"
 #include "Engine/DirectionalLight.h"
@@ -14,7 +18,12 @@
 #include "SEditorViewport.h"
 #include "Selection.h"
 #include "UnrealWidget.h"
+#include "Actors/ChaosVDSolverInfoActor.h"
+#include "Components/ChaosVDSolverCollisionDataComponent.h"
 #include "Visualizers/ChaosVDDebugDrawUtils.h"
+#include "Visualizers/ChaosVDSolverCollisionDataComponentVisualizer.h"
+#include "Widgets/SChaosVDCollisionDataInspector.h"
+#include "Widgets/SChaosVDMainTab.h"
 
 FChaosVDPlaybackViewportClient::FChaosVDPlaybackViewportClient(const TSharedPtr<FEditorModeTools>& InModeTools) : FEditorViewportClient(InModeTools.Get()), CVDWorld(nullptr)
 {
@@ -45,22 +54,48 @@ FChaosVDPlaybackViewportClient::~FChaosVDPlaybackViewportClient()
 void FChaosVDPlaybackViewportClient::ProcessClick(FSceneView& View, HHitProxy* HitProxy, FKey Key, EInputEvent Event, uint32 HitX, uint32 HitY)
 {
 	FEditorViewportClient::ProcessClick(View, HitProxy, Key, Event, HitX, HitY);
-	
+
 	if (HitProxy == nullptr)
 	{
 		return;
 	}
+	
+	const TSharedPtr<SChaosVDMainTab> MainTabToolkitHost = ModeTools.IsValid() ? StaticCastSharedPtr<SChaosVDMainTab>(ModeTools->GetToolkitHost()) : nullptr;
+	if (!MainTabToolkitHost.IsValid())
+	{
+		return;
+	}
 
 	if (TSharedPtr<FChaosVDScene> ScenePtr = CVDScene.Pin())
 	{
 		if (HitProxy->IsA(HActor::StaticGetType()))
 		{
-			HActor* ActorHitProxy = static_cast<HActor*>(HitProxy);
+			const HActor* ActorHitProxy = static_cast<HActor*>(HitProxy);
 			if (AActor* ClickedActor = ActorHitProxy->Actor)
 			{
 				ScenePtr->SetSelectedObject(ClickedActor);
 			}
 		}
+		else
+		{
+			if (HChaosVDContactPointProxy* ContactProxy = HitProxyCast<HChaosVDContactPointProxy>(HitProxy))
+			{
+				if (ContactProxy->ContactFinder.OwningMidPhase.Pin())
+				{
+					if (TSharedPtr<FChaosVDCollisionDataDetailsTab> CollisionDataDetailsTab = MainTabToolkitHost->GetTabSpawnerInstance<FChaosVDCollisionDataDetailsTab>(FChaosVDTabID::CollisionDataDetails).Pin())
+					{
+						if (TSharedPtr<SChaosVDCollisionDataInspector> CollisionInspector = CollisionDataDetailsTab->GetCollisionInspectorInstance().Pin())
+						{
+							if (TSharedPtr<FTabManager> TabManager = MainTabToolkitHost->GetTabManager())
+							{
+								TabManager->TryInvokeTab(FChaosVDTabID::CollisionDataDetails);
+								CollisionInspector->SetSingleContactDataToInspect(ContactProxy->ContactFinder);
+							}
+						}
+					}
+				}
+			}
+		}
 	}
 }
 
@@ -245,8 +280,27 @@ void FChaosVDPlaybackViewportClient::Draw(const FSceneView* View, FPrimitiveDraw
 {
 	FEditorViewportClient::Draw(View, PDI);
 
+	const TSharedPtr<SChaosVDMainTab> MainTabToolkitHost = ModeTools.IsValid() ? StaticCastSharedPtr<SChaosVDMainTab>(ModeTools->GetToolkitHost()) : nullptr;
+	if (!MainTabToolkitHost.IsValid())
+	{
+		return;
+	}
+
 	if (TSharedPtr<FChaosVDScene> ScenePtr = CVDScene.Pin())
-	{	
+	{
+	
+		for (const TPair<int32, AChaosVDSolverInfoActor*>& SolverInfoWithID : ScenePtr->GetSolverInfoActorsMap())
+		{
+			UChaosVDSolverCollisionDataComponent* CollisionDataComponent = SolverInfoWithID.Value ? SolverInfoWithID.Value->GetCollisionDataComponent() : nullptr;
+			if (CollisionDataComponent)
+			{
+				if (TSharedPtr<FComponentVisualizer> Visualizer = MainTabToolkitHost->FindComponentVisualizer(CollisionDataComponent->StaticClass()))
+				{
+					Visualizer->DrawVisualization(CollisionDataComponent, View, PDI);
+				}
+			}
+		}
+		
 		TArray<AActor*> SelectedActors = ScenePtr->GetElementSelectionSet()->GetSelectedObjects<AActor>();
 
 		for (AActor* SelectedActor : SelectedActors)
diff --git a/Engine/Plugins/Experimental/ChaosVD/Source/ChaosVD/Private/ChaosVDPlaybackViewportTab.cpp b/Engine/Plugins/Experimental/ChaosVD/Source/ChaosVD/Private/ChaosVDPlaybackViewportTab.cpp
--- a/Engine/Plugins/Experimental/ChaosVD/Source/ChaosVD/Private/ChaosVDPlaybackViewportTab.cpp
+++ b/Engine/Plugins/Experimental/ChaosVD/Source/ChaosVD/Private/ChaosVDPlaybackViewportTab.cpp
@@ -6,6 +6,7 @@
 #include "ChaosVDStyle.h"
 #include "Widgets/SChaosVDPlaybackViewport.h"
 #include "ChaosVDTabsIDs.h"
+#include "EditorModeManager.h"
 #include "Widgets/Docking/SDockTab.h"
 #include "Widgets/SChaosVDMainTab.h"
 
@@ -19,12 +20,18 @@ TSharedRef<SDockTab> FChaosVDPlaybackViewportTab::HandleTabSpawned(const FSpawnT
 		.Label(LOCTEXT("ViewportTabLabel", "Playback Viewport"))
 		.ToolTipText(LOCTEXT("ViewportTabToolTip", "The Chaos Visual debugger Viewport is under development"));
 	
-	ViewportTab->SetContent
-	(
-		//TODO: Handle Null cases to not crash the Editor
-		SAssignNew(PlaybackViewportWidget, SChaosVDPlaybackViewport, GetChaosVDScene(), OwningTabWidget->GetChaosVDEngineInstance()->GetPlaybackController())
-	);
-
+	if (TSharedPtr<SChaosVDMainTab> MainTabPtr = OwningTabWidget.Pin())
+	{
+		ViewportTab->SetContent
+		(
+			SAssignNew(PlaybackViewportWidget, SChaosVDPlaybackViewport, GetChaosVDScene(), MainTabPtr->GetChaosVDEngineInstance()->GetPlaybackController(), MainTabPtr->GetEditorModeManager().AsShared())
+		);
+	}
+	else
+	{
+		ViewportTab->SetContent(GenerateErrorWidget());
+	}
+	
 	ViewportTab->SetTabIcon(FChaosVDStyle::Get().GetBrush("TabIconPlaybackViewport"));
 
 	return ViewportTab;
diff --git a/Engine/Plugins/Experimental/ChaosVD/Source/ChaosVD/Private/ChaosVDScene.cpp b/Engine/Plugins/Experimental/ChaosVD/Source/ChaosVD/Private/ChaosVDScene.cpp
--- a/Engine/Plugins/Experimental/ChaosVD/Source/ChaosVD/Private/ChaosVDScene.cpp
+++ b/Engine/Plugins/Experimental/ChaosVD/Source/ChaosVD/Private/ChaosVDScene.cpp
@@ -23,6 +23,8 @@
 #include "Selection.h"
 #include "UObject/Package.h"
 #include "WorldPersistentFolders.h"
+#include "Actors/ChaosVDSolverInfoActor.h"
+#include "Components/ChaosVDSolverCollisionDataComponent.h"
 
 #define LOCTEXT_NAMESPACE "ChaosVisualDebugger"
 
@@ -84,6 +86,8 @@ void FChaosVDScene::DeInitialize()
 
 	if (PhysicsVDWorld)
 	{
+		PhysicsVDWorld->RemoveOnActorDestroyededHandler(ActorDestroyedHandle);
+
 		PhysicsVDWorld->DestroyWorld(true);
 		GEngine->DestroyWorldContext(PhysicsVDWorld);
 
@@ -107,7 +111,12 @@ void FChaosVDScene::AddReferencedObjects(FReferenceCollector& Collector)
 
 void FChaosVDScene::UpdateFromRecordedStepData(const int32 SolverID, const FString& SolverName, const FChaosVDStepData& InRecordedStepData, const FChaosVDSolverFrameData& InFrameData)
 {
-	FChaosVDParticlesByIDMap& SolverParticlesByID = ParticlesBySolverID.FindChecked(SolverID);
+	AChaosVDSolverInfoActor* SolverSceneData = SolverDataContainerBySolverID.FindChecked(SolverID);
+
+	if (!SolverSceneData)
+	{
+		return;
+	}
 	
 	TSet<int32> ParticlesIDsInRecordedStepData;
 	ParticlesIDsInRecordedStepData.Reserve(InRecordedStepData.RecordedParticlesData.Num());
@@ -130,23 +139,15 @@ void FChaosVDScene::UpdateFromRecordedStepData(const int32 SolverID, const FStri
 				continue;
 			}
 
-			if (AChaosVDParticleActor** ExistingParticleVDInstancePtrPtr = SolverParticlesByID.Find(ParticleVDInstanceID))
+			if (AChaosVDParticleActor* ExistingParticleVDInstancePtr = SolverSceneData->GetParticleActor(ParticleVDInstanceID))
 			{
-				if (AChaosVDParticleActor* ExistingParticleVDInstancePtr = *ExistingParticleVDInstancePtrPtr)
+				// We have new data for this particle, so re-activate the existing actor
+				if (!ExistingParticleVDInstancePtr->IsActive())
 				{
-					// We have new data for this particle, so re-activate the existing actor
-					if (!ExistingParticleVDInstancePtr->IsActive())
-					{
-						ExistingParticleVDInstancePtr->SetIsActive(true);
-					}
-
-					ExistingParticleVDInstancePtr->UpdateFromRecordedParticleData(Particle, InFrameData.SimulationTransform);
-				}
-				else
-				{
-					//TODO: Handle this error
-					ensure(false);
+					ExistingParticleVDInstancePtr->SetIsActive(true);
 				}
+
+				ExistingParticleVDInstancePtr->UpdateFromRecordedParticleData(Particle, InFrameData.SimulationTransform);
 			}
 			else
 			{
@@ -158,7 +159,7 @@ void FChaosVDScene::UpdateFromRecordedStepData(const int32 SolverID, const FStri
 					NewParticleVDInstance->SetFolderPath(FolderPath);
 
 					// TODO: Precalculate the max num of entries we would see in the loaded file, and use that number to pre-allocate this map
-					SolverParticlesByID.Add(ParticleVDInstanceID, NewParticleVDInstance);
+					SolverSceneData->RegisterParticleActor(ParticleVDInstanceID, NewParticleVDInstance);
 				}
 				else
 				{
@@ -172,8 +173,10 @@ void FChaosVDScene::UpdateFromRecordedStepData(const int32 SolverID, const FStri
 	}
 
 	UpdateParticlesCollisionData(InRecordedStepData, SolverID);
+	
+	const TMap<int32, AChaosVDParticleActor*>& AllSolverParticlesByID = SolverSceneData->GetAllParticleActorsByIDMap();
 
-	for (const TPair<int32, AChaosVDParticleActor*>& ParticleActorWithID : SolverParticlesByID)
+	for (const TPair<int32, AChaosVDParticleActor*>& ParticleActorWithID : AllSolverParticlesByID)
 	{
 		// If we are playing back a keyframe, the scene should only contain what it is in the recorded data
 		const bool bShouldDestroyParticleAnyway = InFrameData.bIsKeyFrame && !ParticlesIDsInRecordedStepData.Contains(ParticleActorWithID.Key);
@@ -203,46 +206,11 @@ void FChaosVDScene::UpdateFromRecordedStepData(const int32 SolverID, const FStri
 
 void FChaosVDScene::UpdateParticlesCollisionData(const FChaosVDStepData& InRecordedStepData, int32 SolverID)
 {
-	constexpr float AmountOfWork = 1.0f;
-	
-	FChaosVDParticlesByIDMap& SolverParticlesByID = ParticlesBySolverID.FindChecked(SolverID);
+	if (AChaosVDSolverInfoActor* SolverDataInfoContainer = SolverDataContainerBySolverID.FindChecked(SolverID))
 	{
-		FScopedSlowTask UpdatingSceneSlowTask(AmountOfWork, LOCTEXT("ProcessingMidphaseData", "Processing MidPhase Data..."));
-		UpdatingSceneSlowTask.MakeDialogDelayed(ChaosVDSceneUIOptions::DelayToShowProgressDialogThreshold,ChaosVDSceneUIOptions::bShowCancelButton, ChaosVDSceneUIOptions::bAllowInPIE);
-		
-		const float PercentagePerElement = 1.0f / InRecordedStepData.RecordedMidPhasesByParticleID.Num();
-		
-		for (const TPair<int32, TArray<TSharedPtr<FChaosVDParticlePairMidPhase>>>& ParticleIDMidPhasePair : InRecordedStepData.RecordedMidPhasesByParticleID)
+		if (UChaosVDSolverCollisionDataComponent* CollisionDataContainer = SolverDataInfoContainer->GetCollisionDataComponent())
 		{
-			if (AChaosVDParticleActor** ExistingParticleVDInstancePtrPtr = SolverParticlesByID.Find(ParticleIDMidPhasePair.Key))
-			{
-				if (AChaosVDParticleActor* ExistingParticleVDInstancePtr = *ExistingParticleVDInstancePtrPtr)
-				{
-					ExistingParticleVDInstancePtr->UpdateCollisionData(ParticleIDMidPhasePair.Value);
-				}
-			}
-
-			UpdatingSceneSlowTask.EnterProgressFrame(PercentagePerElement);
-		}
-	}
-
-	{
-		FScopedSlowTask UpdatingSceneSlowTask(AmountOfWork, LOCTEXT("ProcessingConstraingtData", "Processing Constraint Data..."));
-		UpdatingSceneSlowTask.MakeDialogDelayed(ChaosVDSceneUIOptions::DelayToShowProgressDialogThreshold, ChaosVDSceneUIOptions::bShowCancelButton, ChaosVDSceneUIOptions::bAllowInPIE);
-		
-		const float PercentagePerElement = 1.0f / InRecordedStepData.RecordedMidPhasesByParticleID.Num();
-		
-		for (const TPair<int32, TArray<FChaosVDConstraint>>& ParticleIDMidPhasePair : InRecordedStepData.RecordedConstraintsByParticleID)
-		{
-			if (AChaosVDParticleActor** ExistingParticleVDInstancePtrPtr = SolverParticlesByID.Find(ParticleIDMidPhasePair.Key))
-			{
-				if (AChaosVDParticleActor* ExistingParticleVDInstancePtr = *ExistingParticleVDInstancePtrPtr)
-				{
-					ExistingParticleVDInstancePtr->UpdateCollisionData(ParticleIDMidPhasePair.Value);
-				}
-			}
-
-			UpdatingSceneSlowTask.EnterProgressFrame(PercentagePerElement);
+			CollisionDataContainer->UpdateCollisionData(InRecordedStepData.RecordedMidPhases);
 		}
 	}
 }
@@ -266,31 +234,31 @@ void FChaosVDScene::HandleEnterNewGameFrame(int32 FrameNumber, const TArray<int3
 	for (int32 SolverID : AvailableSolversIds)
 	{
 		AvailableSolversSet.Add(SolverID);
-		if (!ParticlesBySolverID.Contains(SolverID))
+
+		if (!SolverDataContainerBySolverID.Contains(SolverID))
 		{
-			ParticlesBySolverID.Add(SolverID);
+			AChaosVDSolverInfoActor* CollisionDataContainer = PhysicsVDWorld->SpawnActor<AChaosVDSolverInfoActor>();
+			check(CollisionDataContainer);
+			CollisionDataContainer->SetSolverID(SolverID);
+			CollisionDataContainer->SetScene(AsWeak());
+			SolverDataContainerBySolverID.Add(SolverID, CollisionDataContainer);
 		}
 	}
 
 	int32 AmountRemoved = 0;
-	for (TMap<int32, FChaosVDParticlesByIDMap>::TIterator RemoveIterator = ParticlesBySolverID.CreateIterator(); RemoveIterator; ++RemoveIterator)
+
+	for (TMap<int32, AChaosVDSolverInfoActor*>::TIterator RemoveIterator = SolverDataContainerBySolverID.CreateIterator(); RemoveIterator; ++RemoveIterator)
 	{
 		if (!AvailableSolversSet.Contains(RemoveIterator.Key()))
 		{
 			UE_LOG(LogChaosVDEditor, Log, TEXT("[%s] Removing Solver [%d] as it is no longer present in the recording"), ANSI_TO_TCHAR(__FUNCTION__), RemoveIterator.Key());
-	
-			for (const TPair<int32, AChaosVDParticleActor*>& ParticleVDInstanceWithID : RemoveIterator.Value())
-			{
-				if (IsObjectSelected(ParticleVDInstanceWithID.Value))
-				{
-					ClearSelectionAndNotify();
-				}
 
-				PhysicsVDWorld->DestroyActor(ParticleVDInstanceWithID.Value);
+			if (AChaosVDSolverInfoActor* SolverInfoActor = RemoveIterator.Value())
+			{
+				PhysicsVDWorld->DestroyActor(SolverInfoActor);
 			}
 
 			RemoveIterator.RemoveCurrent();
-
 			AmountRemoved++;
 		}
 	}
@@ -307,16 +275,13 @@ void FChaosVDScene::CleanUpScene()
 
 	if (PhysicsVDWorld)
 	{
-		for (const TPair<int32, FChaosVDParticlesByIDMap>& SolverParticleVDInstanceWithID : ParticlesBySolverID)
+		for (const TPair<int32, AChaosVDSolverInfoActor*>& SolverDataInfoWithID : SolverDataContainerBySolverID)
 		{
-			for (const TPair<int32, AChaosVDParticleActor*>& ParticleVDInstanceWithID : SolverParticleVDInstanceWithID.Value)
-			{
-				PhysicsVDWorld->DestroyActor(ParticleVDInstanceWithID.Value);
-			}
+			PhysicsVDWorld->DestroyActor(SolverDataInfoWithID.Value);
 		}
 	}
 
-	ParticlesBySolverID.Reset();
+	SolverDataContainerBySolverID.Reset();
 }
 
 Chaos::FConstImplicitObjectPtr FChaosVDScene::GetUpdatedGeometry(int32 GeometryID) const
@@ -334,8 +299,22 @@ Chaos::FConstImplicitObjectPtr FChaosVDScene::GetUpdatedGeometry(int32 GeometryI
 
 AChaosVDParticleActor* FChaosVDScene::GetParticleActor(int32 SolverID, int32 ParticleID)
 {
-	AChaosVDParticleActor** ParticleActorPtrPtr = ParticlesBySolverID[SolverID].Find(ParticleID);
-	return ParticleActorPtrPtr ? *ParticleActorPtrPtr : nullptr;
+	if (AChaosVDSolverInfoActor** SolverDataInfo = SolverDataContainerBySolverID.Find(SolverID))
+	{
+		return (*SolverDataInfo)->GetParticleActor(ParticleID);
+	}
+
+	return nullptr;
+}
+
+AChaosVDSolverInfoActor* FChaosVDScene::GetSolverInfoActor(int32 SolverID)
+{
+	if (AChaosVDSolverInfoActor** SolverDataInfo = SolverDataContainerBySolverID.Find(SolverID))
+	{
+		return *SolverDataInfo;
+	}
+
+	return nullptr;
 }
 
 AChaosVDParticleActor* FChaosVDScene::SpawnParticleFromRecordedData(const FChaosVDParticleDataWrapper& InParticleData, const FChaosVDSolverFrameData& InFrameData)
@@ -413,7 +392,7 @@ void FChaosVDScene::CreateBaseLights(UWorld* TargetWorld) const
 	}
 }
 
-UWorld* FChaosVDScene::CreatePhysicsVDWorld() const
+UWorld* FChaosVDScene::CreatePhysicsVDWorld()
 {
 	const FName UniqueWorldName = FName(FGuid::NewGuid().ToString());
 	UWorld* NewWorld = NewWorld = NewObject<UWorld>( GetTransientPackage(), UniqueWorldName );
@@ -434,6 +413,9 @@ UWorld* FChaosVDScene::CreatePhysicsVDWorld() const
 	);
 
 	CreateBaseLights(NewWorld);
+
+	ActorDestroyedHandle = NewWorld->AddOnActorDestroyedHandler(FOnActorDestroyed::FDelegate::CreateRaw(this, &FChaosVDScene::HandleActorDestroyed));
+	
 	return NewWorld;
 }
 
@@ -495,11 +477,11 @@ void FChaosVDScene::ClearSelectionAndNotify()
 
 void FChaosVDScene::HandleVisibilitySettingsChanged(UChaosVDEditorSettings* SettingsObject)
 {
-	for (const TPair<int32, FChaosVDParticlesByIDMap>& ParticlesBySolver : ParticlesBySolverID)
+	for (const TPair<int32, AChaosVDSolverInfoActor*>& SolverDataInfoWithID : SolverDataContainerBySolverID)
 	{
-		for (const TPair<int32, AChaosVDParticleActor*>& ParticleWithIDPair : ParticlesBySolver.Value)
+		if (AChaosVDSolverInfoActor* SolverDataInfo = SolverDataInfoWithID.Value)
 		{
-			ParticleWithIDPair.Value->UpdateGeometryComponentsVisibility();
+			SolverDataInfo->HandleVisibilitySettingsUpdated();
 		}
 	}
 }
@@ -532,6 +514,14 @@ void FChaosVDScene::DeInitializeSelectionSets()
 	SelectionSet->OnChanged().RemoveAll(this);
 }
 
+void FChaosVDScene::HandleActorDestroyed(AActor* ActorDestroyed)
+{
+	if (IsObjectSelected(ActorDestroyed))
+	{
+		ClearSelectionAndNotify();
+	}
+}
+
 void FChaosVDScene::SetSelectedObject(UObject* SelectedObject)
 {
 	if (!SelectionSet)
diff --git a/Engine/Plugins/Experimental/ChaosVD/Source/ChaosVD/Private/ChaosVDSolversTracksTab.cpp b/Engine/Plugins/Experimental/ChaosVD/Source/ChaosVD/Private/ChaosVDSolversTracksTab.cpp
--- a/Engine/Plugins/Experimental/ChaosVD/Source/ChaosVD/Private/ChaosVDSolversTracksTab.cpp
+++ b/Engine/Plugins/Experimental/ChaosVD/Source/ChaosVD/Private/ChaosVDSolversTracksTab.cpp
@@ -18,14 +18,22 @@ TSharedRef<SDockTab> FChaosVDSolversTracksTab::HandleTabSpawned(const FSpawnTabA
 	.Label(LOCTEXT("SolverTracksTabLabel", "Available Solvers"))
 	.ToolTipText(LOCTEXT("SolverTracksTabToolTip", "Playback controls for the available solvers on the current Frame"));
 
-	ViewportTab->SetContent
-	(
-		// TODO: Handle Null cases to not crash the Editor
-		SAssignNew(SolverTracksWidget, SChaosVDSolverTracks, OwningTabWidget->GetChaosVDEngineInstance()->GetPlaybackController())
-	);
+	if (const TSharedPtr<SChaosVDMainTab> MainTabPtr = OwningTabWidget.Pin())
+	{
+		ViewportTab->SetContent
+		(
+			SAssignNew(SolverTracksWidget, SChaosVDSolverTracks, MainTabPtr->GetChaosVDEngineInstance()->GetPlaybackController())
+		);
+	}
+	else
+	{
+		ViewportTab->SetContent(GenerateErrorWidget());
+	}
 
 	ViewportTab->SetTabIcon(FChaosVDStyle::Get().GetBrush("TabIconPlaybackViewport"));
 
+	OnTabSpawned().Broadcast(ViewportTab);
+
 	return ViewportTab;
 }
 
diff --git a/Engine/Plugins/Experimental/ChaosVD/Source/ChaosVD/Private/ChaosVDTabSpawnerBase.cpp b/Engine/Plugins/Experimental/ChaosVD/Source/ChaosVD/Private/ChaosVDTabSpawnerBase.cpp
--- a/Engine/Plugins/Experimental/ChaosVD/Source/ChaosVD/Private/ChaosVDTabSpawnerBase.cpp
+++ b/Engine/Plugins/Experimental/ChaosVD/Source/ChaosVD/Private/ChaosVDTabSpawnerBase.cpp
@@ -5,13 +5,31 @@
 #include "ChaosVDEngine.h"
 #include "ChaosVDScene.h"
 #include "Widgets/SChaosVDMainTab.h"
+#include "Widgets/Text/STextBlock.h"
 
-FChaosVDTabSpawnerBase::FChaosVDTabSpawnerBase(const FName& InTabID, TSharedPtr<FTabManager> InTabManager, SChaosVDMainTab* InOwningTabWidget)
+#define LOCTEXT_NAMESPACE "ChaosVisualDebugger"
+
+FChaosVDTabSpawnerBase::FChaosVDTabSpawnerBase(const FName& InTabID, TSharedPtr<FTabManager> InTabManager, TWeakPtr<SChaosVDMainTab> InOwningTabWidget)
 {
 	OwningTabWidget = InOwningTabWidget;
 	InTabManager->RegisterTabSpawner(InTabID, FOnSpawnTab::CreateRaw(this, &FChaosVDTabSpawnerBase::HandleTabSpawned));
 }
 
+void FChaosVDTabSpawnerBase::HandleTabClosed(const TSharedRef<SDockTab>& InTabClosed)
+{
+	OnTabDestroyed().Broadcast(InTabClosed);
+}
+
+TSharedRef<SWidget> FChaosVDTabSpawnerBase::GenerateErrorWidget()
+{
+	return SNew(SVerticalBox)
+		+SVerticalBox::Slot()
+		[
+			SNew(STextBlock)
+			.Text(LOCTEXT("ChaosVDEditorTabSpawner", "Failed to generate Tab content"))
+		];
+}
+
 UWorld* FChaosVDTabSpawnerBase::GetChaosVDWorld() const
 {
 	if (TSharedPtr<FChaosVDScene> ScenePtr = GetChaosVDScene().Pin())
@@ -24,9 +42,10 @@ UWorld* FChaosVDTabSpawnerBase::GetChaosVDWorld() const
 
 TWeakPtr<FChaosVDScene> FChaosVDTabSpawnerBase::GetChaosVDScene() const
 {
-	if (ensure(OwningTabWidget))
+	const TSharedPtr<SChaosVDMainTab> MainTabPtr = OwningTabWidget.Pin();
+	if (ensure(MainTabPtr))
 	{
-		const TSharedRef<FChaosVDEngine> ChaosVDEngine = OwningTabWidget->GetChaosVDEngineInstance();
+		const TSharedRef<FChaosVDEngine> ChaosVDEngine = MainTabPtr->GetChaosVDEngineInstance();
 		if (TSharedPtr<FChaosVDScene> ScenePtr = ChaosVDEngine->GetCurrentScene())
 		{
 			return ScenePtr;
@@ -34,3 +53,5 @@ TWeakPtr<FChaosVDScene> FChaosVDTabSpawnerBase::GetChaosVDScene() const
 	}
 	return nullptr;
 }
+
+#undef LOCTEXT_NAMESPACE
diff --git a/Engine/Plugins/Experimental/ChaosVD/Source/ChaosVD/Private/ChaosVDTabsIDs.cpp b/Engine/Plugins/Experimental/ChaosVD/Source/ChaosVD/Private/ChaosVDTabsIDs.cpp
--- a/Engine/Plugins/Experimental/ChaosVD/Source/ChaosVD/Private/ChaosVDTabsIDs.cpp
+++ b/Engine/Plugins/Experimental/ChaosVD/Source/ChaosVD/Private/ChaosVDTabsIDs.cpp
@@ -10,4 +10,5 @@ const FName FChaosVDTabID::DetailsPanel = TEXT("DetailsPanel");
 const FName FChaosVDTabID::OutputLog = TEXT("ChaosVDOutputLog");
 const FName FChaosVDTabID::SolversTrack = TEXT("SolversTrack");
 const FName FChaosVDTabID::CVDEditorSettings = TEXT("CVDEditorSettings");
-const FName FChaosVDTabID::StatusBar = TEXT("StatusBar");
\ No newline at end of file
+const FName FChaosVDTabID::StatusBar = TEXT("StatusBar");
+const FName FChaosVDTabID::CollisionDataDetails = TEXT("CollisionDataDetails");
\ No newline at end of file
diff --git a/Engine/Plugins/Experimental/ChaosVD/Source/ChaosVD/Private/ChaosVDWorldOutlinerTab.cpp b/Engine/Plugins/Experimental/ChaosVD/Source/ChaosVD/Private/ChaosVDWorldOutlinerTab.cpp
--- a/Engine/Plugins/Experimental/ChaosVD/Source/ChaosVD/Private/ChaosVDWorldOutlinerTab.cpp
+++ b/Engine/Plugins/Experimental/ChaosVD/Source/ChaosVD/Private/ChaosVDWorldOutlinerTab.cpp
@@ -58,6 +58,8 @@ TSharedRef<SDockTab> FChaosVDWorldOutlinerTab::HandleTabSpawned(const FSpawnTabA
 
 	OutlinerTab->SetTabIcon(FChaosVDStyle::Get().GetBrush("TabIconWorldOutliner"));
 
+	OnTabSpawned().Broadcast(OutlinerTab);
+
 	return OutlinerTab;
 }
 
diff --git a/Engine/Plugins/Experimental/ChaosVD/Source/ChaosVD/Private/Components/ChaosVDSolverCollisionDataComponent.cpp b/Engine/Plugins/Experimental/ChaosVD/Source/ChaosVD/Private/Components/ChaosVDSolverCollisionDataComponent.cpp
new file mode 100644
--- /dev/null
+++ b/Engine/Plugins/Experimental/ChaosVD/Source/ChaosVD/Private/Components/ChaosVDSolverCollisionDataComponent.cpp
@@ -0,0 +1,59 @@
+// Copyright Epic Games, Inc. All Rights Reserved.
+
+#include "Components/ChaosVDSolverCollisionDataComponent.h"
+
+#include "DataWrappers/ChaosVDCollisionDataWrappers.h"
+
+UChaosVDSolverCollisionDataComponent::UChaosVDSolverCollisionDataComponent()
+{
+	SetCanEverAffectNavigation(false);
+	bNavigationRelevant = false;
+}
+
+void UChaosVDSolverCollisionDataComponent::UpdateCollisionData(const TArray<TSharedPtr<FChaosVDParticlePairMidPhase>>& InMidPhaseData)
+{
+	ClearCollisionData();
+
+	AllMidPhases.Reserve(InMidPhaseData.Num());
+	AllMidPhases = InMidPhaseData;
+
+	MidPhasesByParticleID0.Reserve(InMidPhaseData.Num());
+	MidPhasesByParticleID1.Reserve(InMidPhaseData.Num());
+
+	for (const TSharedPtr<FChaosVDParticlePairMidPhase>& ParticleIDMidPhase : InMidPhaseData)
+	{
+		AddCollisionDataToParticleIDMap(MidPhasesByParticleID0, ParticleIDMidPhase, ParticleIDMidPhase->Particle0Idx);
+		AddCollisionDataToParticleIDMap(MidPhasesByParticleID1, ParticleIDMidPhase, ParticleIDMidPhase->Particle1Idx);
+
+		ConstraintsByParticleID0.Reserve(ParticleIDMidPhase->Constraints.Num());
+		ConstraintsByParticleID1.Reserve(ParticleIDMidPhase->Constraints.Num());
+		for (FChaosVDConstraint& Constraint : ParticleIDMidPhase->Constraints)
+		{
+			AddCollisionDataToParticleIDMap<FChaosVDConstraintByParticleMap, FChaosVDConstraint*>(ConstraintsByParticleID0, &Constraint, Constraint.Particle0Index);
+			AddCollisionDataToParticleIDMap<FChaosVDConstraintByParticleMap, FChaosVDConstraint*>(ConstraintsByParticleID1, &Constraint, Constraint.Particle1Index);
+		}
+	}
+}
+
+const TArray<TSharedPtr<FChaosVDParticlePairMidPhase>>* UChaosVDSolverCollisionDataComponent::GetMidPhasesForParticle(int32 ParticleID, EChaosVDGetCollisionDataOptions Options) const
+{
+	return GetCollisionDataFromMap<FChaosVDMidPhaseByParticleMap, TSharedPtr<FChaosVDParticlePairMidPhase>>(MidPhasesByParticleID0, MidPhasesByParticleID1, ParticleID, Options);
+}
+
+const TArray<FChaosVDConstraint*>* UChaosVDSolverCollisionDataComponent::GetConstraintsForParticle(int32 ParticleID, EChaosVDGetCollisionDataOptions Options) const
+{
+	return GetCollisionDataFromMap<FChaosVDConstraintByParticleMap, FChaosVDConstraint*>(ConstraintsByParticleID0, ConstraintsByParticleID1, ParticleID, Options);
+}
+
+void UChaosVDSolverCollisionDataComponent::DrawVisualization(const FSceneView* View, FPrimitiveDrawInterface* PDI)
+{
+}
+
+void UChaosVDSolverCollisionDataComponent::ClearCollisionData()
+{
+	AllMidPhases.Reset();
+	MidPhasesByParticleID0.Reset();
+	MidPhasesByParticleID1.Reset();
+	ConstraintsByParticleID0.Reset();
+	ConstraintsByParticleID1.Reset();
+}
\ No newline at end of file
diff --git a/Engine/Plugins/Experimental/ChaosVD/Source/ChaosVD/Private/DetailsCustomizations/ChaosVDParticleDataWrapperCustomization.cpp b/Engine/Plugins/Experimental/ChaosVD/Source/ChaosVD/Private/DetailsCustomizations/ChaosVDParticleDataWrapperCustomization.cpp
--- a/Engine/Plugins/Experimental/ChaosVD/Source/ChaosVD/Private/DetailsCustomizations/ChaosVDParticleDataWrapperCustomization.cpp
+++ b/Engine/Plugins/Experimental/ChaosVD/Source/ChaosVD/Private/DetailsCustomizations/ChaosVDParticleDataWrapperCustomization.cpp
@@ -29,7 +29,7 @@ void FChaosVDParticleDataWrapperCustomization::CustomizeChildren(TSharedRef<IPro
 	ParticleDataViewersNames.Add(GET_MEMBER_NAME_CHECKED(FChaosVDParticleDataWrapper, ParticleDynamics));
 	ParticleDataViewersNames.Add(GET_MEMBER_NAME_CHECKED(FChaosVDParticleDataWrapper, ParticleDynamicsMisc));
 	ParticleDataViewersNames.Add(GET_MEMBER_NAME_CHECKED(FChaosVDParticleDataWrapper, ParticleMassProps));
-
+	
 	for (uint32 ChildIndex = 0; ChildIndex < NumChildren; ++ChildIndex)
 	{
 		TSharedRef<IPropertyHandle> ChildHandle = StructPropertyHandle->GetChildHandle(ChildIndex).ToSharedRef();
@@ -51,9 +51,5 @@ void FChaosVDParticleDataWrapperCustomization::CustomizeChildren(TSharedRef<IPro
 				}
 			}
 		}
-		else
-		{
-			StructBuilder.AddProperty(ChildHandle);
-		}
 	}
 }
diff --git a/Engine/Plugins/Experimental/ChaosVD/Source/ChaosVD/Private/Trace/ChaosVDTraceAnalyzer.cpp b/Engine/Plugins/Experimental/ChaosVD/Source/ChaosVD/Private/Trace/ChaosVDTraceAnalyzer.cpp
--- a/Engine/Plugins/Experimental/ChaosVD/Source/ChaosVD/Private/Trace/ChaosVDTraceAnalyzer.cpp
+++ b/Engine/Plugins/Experimental/ChaosVD/Source/ChaosVD/Private/Trace/ChaosVDTraceAnalyzer.cpp
@@ -177,9 +177,11 @@ bool FChaosVDTraceAnalyzer::OnEvent(uint16 RouteId, EStyle Style, const FOnEvent
 		{
 			const int32 DataID = EventData.GetValue<int32>("DataID");
 
-			if (!ensure(ChaosVDTraceProvider->ProcessBinaryData(DataID)))
+			if (!ChaosVDTraceProvider->ProcessBinaryData(DataID))
 			{
-				UE_LOG(LogChaosVDEditor, Error, TEXT("[%s] FailedToProcess Binary Data with ID [%d]"), ANSI_TO_TCHAR(__FUNCTION__), DataID);
+				// This can happen during live debugging as we miss some of the events at the beginning.
+				// Loading a trace file that was recorded as part of a live session, will have the same issue.
+				UE_LOG(LogChaosVDEditor, Warning, TEXT("[%s] FailedToProcess Binary Data with ID [%d]"), ANSI_TO_TCHAR(__FUNCTION__), DataID);
 			}
 
 			break;
diff --git a/Engine/Plugins/Experimental/ChaosVD/Source/ChaosVD/Private/Visualizers/ChaosVDCollisionDataVisualizer.cpp b/Engine/Plugins/Experimental/ChaosVD/Source/ChaosVD/Private/Visualizers/ChaosVDCollisionDataVisualizer.cpp
deleted file mode 100644
diff --git a/Engine/Plugins/Experimental/ChaosVD/Source/ChaosVD/Private/Visualizers/ChaosVDParticleDataVisualizer.cpp b/Engine/Plugins/Experimental/ChaosVD/Source/ChaosVD/Private/Visualizers/ChaosVDParticleDataVisualizer.cpp
--- a/Engine/Plugins/Experimental/ChaosVD/Source/ChaosVD/Private/Visualizers/ChaosVDParticleDataVisualizer.cpp
+++ b/Engine/Plugins/Experimental/ChaosVD/Source/ChaosVD/Private/Visualizers/ChaosVDParticleDataVisualizer.cpp
@@ -9,24 +9,18 @@
 
 void FChaosVDParticleDataVisualizer::DrawVisualization(const FSceneView* View, FPrimitiveDrawInterface* PDI)
 {
-	uint32 VisualizationFlagsToUse = 0;
-	if (static_cast<EChaosVDParticleDataVisualizationFlags>(LocalVisualizationFlags) == EChaosVDParticleDataVisualizationFlags::None)
+	EChaosVDParticleDataVisualizationFlags VisualizationFlagsToUse = static_cast<EChaosVDParticleDataVisualizationFlags>(LocalVisualizationFlags);
+	if (VisualizationFlagsToUse == EChaosVDParticleDataVisualizationFlags::None)
 	{
 		if (const UChaosVDEditorSettings* EditorSettings = GetDefault<UChaosVDEditorSettings>())
 		{
-			VisualizationFlagsToUse = EditorSettings->GlobalParticleDataVisualizationFlags;
+			VisualizationFlagsToUse = static_cast<EChaosVDParticleDataVisualizationFlags>(EditorSettings->GlobalParticleDataVisualizationFlags);
 		}
 		else
 		{
 			UE_LOG(LogChaosVDEditor, Warning, TEXT("[%s] Failed to retrive global visualization setting. Falling back to local settings"), ANSI_TO_TCHAR(__FUNCTION__));
-
-			VisualizationFlagsToUse = LocalVisualizationFlags;
 		}	
 	}
-	else
-	{
-		VisualizationFlagsToUse = LocalVisualizationFlags;
-	}
 	
 	FChaosVDVisualizationContext VisualizationContext;
 	DataProvider.GetVisualizationContext(VisualizationContext);
@@ -43,12 +37,12 @@ void FChaosVDParticleDataVisualizer::DrawVisualization(const FSceneView* View, F
 
 	if (ParticleDataViewer->ParticleVelocities.HasValidData())
 	{
-		if (EnumHasAnyFlags(static_cast<EChaosVDParticleDataVisualizationFlags>(VisualizationFlagsToUse), EChaosVDParticleDataVisualizationFlags::Velocity))
+		if (EnumHasAnyFlags(VisualizationFlagsToUse, EChaosVDParticleDataVisualizationFlags::Velocity))
 		{
 			FChaosVDDebugDrawUtils::DrawArrowVector(PDI, OwnerLocation, OwnerLocation + ParticleDataViewer->ParticleVelocities.MV, TEXT("Velocity"), FColor::Green);
 		}
 
-		if (EnumHasAnyFlags(static_cast<EChaosVDParticleDataVisualizationFlags>(VisualizationFlagsToUse), EChaosVDParticleDataVisualizationFlags::AngularVelocity))
+		if (EnumHasAnyFlags(VisualizationFlagsToUse, EChaosVDParticleDataVisualizationFlags::AngularVelocity))
 		{
 			FChaosVDDebugDrawUtils::DrawArrowVector(PDI, OwnerLocation, OwnerLocation + ParticleDataViewer->ParticleVelocities.MW, TEXT("Angular Velocity"), FColor::Blue);
 		}
@@ -56,22 +50,22 @@ void FChaosVDParticleDataVisualizer::DrawVisualization(const FSceneView* View, F
 
 	if (ParticleDataViewer->ParticleDynamics.HasValidData())
 	{
-		if (EnumHasAnyFlags(static_cast<EChaosVDParticleDataVisualizationFlags>(VisualizationFlagsToUse), EChaosVDParticleDataVisualizationFlags::Acceleration))
+		if (EnumHasAnyFlags(VisualizationFlagsToUse, EChaosVDParticleDataVisualizationFlags::Acceleration))
 		{
 			FChaosVDDebugDrawUtils::DrawArrowVector(PDI, OwnerLocation, OwnerLocation + ParticleDataViewer->ParticleDynamics.MAcceleration, TEXT("Acceleration"), FColor::Orange);
 		}
 		
-		if (EnumHasAnyFlags(static_cast<EChaosVDParticleDataVisualizationFlags>(VisualizationFlagsToUse), EChaosVDParticleDataVisualizationFlags::AngularAcceleration))
+		if (EnumHasAnyFlags(VisualizationFlagsToUse, EChaosVDParticleDataVisualizationFlags::AngularAcceleration))
 		{
 			FChaosVDDebugDrawUtils::DrawArrowVector(PDI, OwnerLocation, OwnerLocation + ParticleDataViewer->ParticleDynamics.MAngularAcceleration, TEXT("Angular Acceleration"), FColor::Purple);
 		}
 
-		if (EnumHasAnyFlags(static_cast<EChaosVDParticleDataVisualizationFlags>(VisualizationFlagsToUse), EChaosVDParticleDataVisualizationFlags::LinearImpulse))
+		if (EnumHasAnyFlags(VisualizationFlagsToUse, EChaosVDParticleDataVisualizationFlags::LinearImpulse))
 		{
 			FChaosVDDebugDrawUtils::DrawArrowVector(PDI, OwnerLocation, OwnerLocation + ParticleDataViewer->ParticleDynamics.MLinearImpulseVelocity, TEXT("Linear Implulse Velocity"), FColor::Purple);
 		}
 
-		if (EnumHasAnyFlags(static_cast<EChaosVDParticleDataVisualizationFlags>(VisualizationFlagsToUse), EChaosVDParticleDataVisualizationFlags::AngularImpulse))
+		if (EnumHasAnyFlags(VisualizationFlagsToUse, EChaosVDParticleDataVisualizationFlags::AngularImpulse))
 		{
 			FChaosVDDebugDrawUtils::DrawArrowVector(PDI, OwnerLocation, OwnerLocation + ParticleDataViewer->ParticleDynamics.MAngularImpulseVelocity, TEXT("Angular Implulse Velocity"), FColor::Emerald);
 		}
diff --git a/Engine/Plugins/Experimental/ChaosVD/Source/ChaosVD/Private/Visualizers/ChaosVDSolverCollisionDataComponentVisualizer.cpp b/Engine/Plugins/Experimental/ChaosVD/Source/ChaosVD/Private/Visualizers/ChaosVDSolverCollisionDataComponentVisualizer.cpp
new file mode 100644
--- /dev/null
+++ b/Engine/Plugins/Experimental/ChaosVD/Source/ChaosVD/Private/Visualizers/ChaosVDSolverCollisionDataComponentVisualizer.cpp
@@ -0,0 +1,244 @@
+// Copyright Epic Games, Inc. All Rights Reserved.
+
+#include "ChaosVDSolverCollisionDataComponentVisualizer.h"
+
+#include "ChaosVDEditorSettings.h"
+#include "ChaosVDParticleActor.h"
+#include "ChaosVDScene.h"
+#include "Actors/ChaosVDSolverInfoActor.h"
+#include "Components/ChaosVDSolverCollisionDataComponent.h"
+#include "Visualizers/ChaosVDDebugDrawUtils.h"
+
+IMPLEMENT_HIT_PROXY(HChaosVDContactPointProxy, HHitProxy)
+
+FChaosVDSolverCollisionDataComponentVisualizer::FChaosVDSolverCollisionDataComponentVisualizer()
+{
+}
+
+FChaosVDSolverCollisionDataComponentVisualizer::~FChaosVDSolverCollisionDataComponentVisualizer()
+{
+}
+
+bool FChaosVDSolverCollisionDataComponentVisualizer::ShowWhenSelected()
+{
+	return false;
+}
+
+void FChaosVDSolverCollisionDataComponentVisualizer::DrawVisualization(const UActorComponent* Component, const FSceneView* View, FPrimitiveDrawInterface* PDI)
+{
+	const UChaosVDSolverCollisionDataComponent* CollisionDataComponent = Cast<UChaosVDSolverCollisionDataComponent>(Component);
+	if (!CollisionDataComponent)
+	{
+		return;
+	}
+
+	const AChaosVDSolverInfoActor* SolverInfoContainer = Cast<AChaosVDSolverInfoActor>(CollisionDataComponent->GetOwner());
+	if (!SolverInfoContainer)
+	{
+		return;
+	}
+
+	FChaosVDVisualizationContext VisualizationContext;
+	VisualizationContext.SolverID = SolverInfoContainer->GetSolverID();
+	VisualizationContext.CVDScene = SolverInfoContainer->GetScene();
+	VisualizationContext.SpaceTransform = SolverInfoContainer->GetSimulationTransform();
+
+	if (const UChaosVDEditorSettings* EditorSettings = GetDefault<UChaosVDEditorSettings>())
+	{
+		VisualizationContext.VisualizationFlags = EditorSettings->GlobalCollisionDataVisualizationFlags;
+	}
+
+	if (EnumHasAnyFlags(EChaosVDCollisionVisualizationFlags::EnableDrawForAllParticles, static_cast<EChaosVDCollisionVisualizationFlags>(VisualizationContext.VisualizationFlags)))
+	{
+		for (const TSharedPtr<FChaosVDParticlePairMidPhase>& MidPhase : CollisionDataComponent->GetMidPhases())
+		{
+			DrawnMidPhaseData(MidPhase, VisualizationContext, View, PDI);
+		}
+	}
+	else
+	{
+		const TArray<int32>& SelectedParticlesIDs = SolverInfoContainer->GetSelectedParticlesIDs();
+
+		for (const int32 SelectedParticleID : SelectedParticlesIDs)
+		{
+			if (const TArray<TSharedPtr<FChaosVDParticlePairMidPhase>>* ParticleMidPhases = CollisionDataComponent->GetMidPhasesForParticle(SelectedParticleID, EChaosVDGetCollisionDataOptions::Owner))
+			{
+				for (const TSharedPtr<FChaosVDParticlePairMidPhase>& MidPhase : *ParticleMidPhases)
+				{
+					DrawnMidPhaseData(MidPhase, VisualizationContext, View, PDI);
+				}
+			}
+		}		
+	}
+}
+
+void FChaosVDSolverCollisionDataComponentVisualizer::DrawnMidPhaseData(const TSharedPtr<FChaosVDParticlePairMidPhase>& MidPhase, const FChaosVDVisualizationContext& VisualizationContext, const FSceneView* View, FPrimitiveDrawInterface* PDI)
+{
+	TSharedPtr<FChaosVDScene> ScenePtr = VisualizationContext.CVDScene.Pin();
+	if (!ScenePtr.IsValid())
+	{
+		return;
+	}
+	
+	if (MidPhase.IsValid())
+	{
+
+		AChaosVDParticleActor* ParticleActor = ScenePtr->GetParticleActor(VisualizationContext.SolverID, MidPhase->Particle0Idx);
+		if (!ParticleActor)
+		{
+			return;
+		}
+		
+		const FChaosVDParticleDataWrapper* ParticleDataViewer = ParticleActor->GetParticleData();
+		if (!ensure(ParticleDataViewer))
+		{
+			return;
+		}
+
+		EChaosVDCollisionVisualizationFlags VisualizationFlags = static_cast<EChaosVDCollisionVisualizationFlags>(VisualizationContext.VisualizationFlags);
+		 
+		const FVector& OwnerActorLocation = ParticleDataViewer->ParticlePositionRotation.MX;
+
+		for (const FChaosVDConstraint& Constraint : MidPhase->Constraints)
+		{
+			int32 ContactIndex = INDEX_NONE;
+			for (const FChaosVDManifoldPoint& ManifoldPoint : Constraint.ManifoldPoints)
+			{
+				ContactIndex++;
+
+				FChaosVDCollisionDataFinder HitPRoxyDataFinder;
+				HitPRoxyDataFinder.OwningMidPhase = MidPhase;
+				HitPRoxyDataFinder.OwningConstraint = &Constraint;
+				HitPRoxyDataFinder.ContactIndex = ContactIndex;
+
+				constexpr int32 MinShapeContactAmount = 2;
+				if (Constraint.ShapeWorldTransforms.Num() < MinShapeContactAmount || ManifoldPoint.ContactPoint.ShapeContactPoints.Num() < MinShapeContactAmount)
+				{
+					continue;
+				}
+
+				const bool bIsProbe = Constraint.bIsProbe;
+				const bool bIsActive = ManifoldPoint.bIsValid && (!ManifoldPoint.NetPushOut.IsNearlyZero() || !ManifoldPoint.NetImpulse.IsNearlyZero() || (!Constraint.bUseManifold && !Constraint.AccumulatedImpulse.IsNearlyZero()));
+				if (!bIsActive && EnumHasAnyFlags(VisualizationFlags, EChaosVDCollisionVisualizationFlags::DrawOnlyActiveContacts))
+				{
+					continue;
+				}
+
+				const bool bPruned = ManifoldPoint.bDisabled;
+
+				AChaosVDParticleActor* CVDParticleActor1 = ScenePtr->GetParticleActor(VisualizationContext.SolverID, Constraint.Particle1Index);
+				AChaosVDParticleActor* CVDParticleActor0 = ScenePtr->GetParticleActor(VisualizationContext.SolverID, Constraint.Particle0Index);
+				if (CVDParticleActor1 == nullptr || CVDParticleActor0 == nullptr)
+				{
+					continue;
+				}
+
+				PDI->SetHitProxy(new HChaosVDContactPointProxy(HitPRoxyDataFinder));
+
+				const FTransform& WorldActorTransform1 = CVDParticleActor1->GetTransform();
+				const FTransform& WorldActorTransform0 = CVDParticleActor0->GetTransform();
+
+				constexpr int32 ContactPlaneOwner = 1;
+				constexpr int32 ContactPointOwner = 1 - ContactPlaneOwner;
+				const FTransform& PlaneTransform = Constraint.ImplicitTransforms[1] * WorldActorTransform1;
+				const FTransform& PointTransform = Constraint.ImplicitTransforms[0] * WorldActorTransform0;
+
+				const FVector PlaneNormal = PlaneTransform.TransformVectorNoScale(FVector(ManifoldPoint.ContactPoint.ShapeContactNormal));
+				const FVector PointLocation = PointTransform.TransformPosition(FVector(ManifoldPoint.ContactPoint.ShapeContactPoints[ContactPointOwner]));
+				const FVector PlaneLocation = PlaneTransform.TransformPosition(FVector(ManifoldPoint.ContactPoint.ShapeContactPoints[ContactPlaneOwner]));
+				const FVector PointPlaneLocation = PointLocation - FVector::DotProduct(PointLocation - PlaneLocation, PlaneNormal) * PlaneNormal;
+
+				// TODO: Extract these colors to a class or another place where ChaosDebugDraw and Chaos VD Visualizers can access
+				// The bulk of this implementation including the following colors come from ChaosDebugDraw Implementation.
+				// The implementation itself is likely to drift, but we should try to keep the colors the same between both debugging tools
+
+				// Dynamic friction, restitution = red
+				// Static friction, no restitution = green
+				// Inactive = gray
+				FColor DiscColor = FColor(200, 0, 0);
+				FColor PlaneNormalColor = FColor(200, 0, 0);
+				FColor EdgeNormalColor = FColor(200, 150, 0);
+				FColor ImpulseColor = FColor(0, 0, 200);
+				FColor PushOutImpusleColor = FColor(0, 200, 200);
+				if (ManifoldPoint.bInsideStaticFrictionCone)
+				{
+					DiscColor = FColor(150, 200, 0);
+				}
+				if (bIsProbe)
+				{
+					DiscColor = FColor(50, 180, 180);
+					PlaneNormalColor = FColor(50, 180, 180);
+					EdgeNormalColor = FColor(50, 180, 130);
+				}
+				else if (!bIsActive)
+				{
+					DiscColor = FColor(100, 100, 100);
+					PlaneNormalColor = FColor(100, 0, 0);
+					EdgeNormalColor = FColor(100, 80, 0);
+				}
+				if (bPruned)
+				{
+					PlaneNormalColor = FColor(200, 0, 200);
+					EdgeNormalColor = FColor(200, 0, 200);
+				}
+			
+				const FVector WorldPointLocation = VisualizationContext.SpaceTransform.TransformPosition(PointLocation);
+				const FVector WorldPlaneLocation = VisualizationContext.SpaceTransform.TransformPosition(PlaneLocation);
+				const FVector WorldPointPlaneLocation = VisualizationContext.SpaceTransform.TransformPosition(PointPlaneLocation);
+				const FVector WorldPlaneNormal = VisualizationContext.SpaceTransform.TransformVectorNoScale(PlaneNormal);
+
+				if (EnumHasAnyFlags(VisualizationFlags, EChaosVDCollisionVisualizationFlags::NetPushOut))
+				{
+					if (ManifoldPoint.bIsValid && !ManifoldPoint.NetPushOut.IsNearlyZero())
+					{
+						FChaosVDDebugDrawUtils::DrawArrowVector(PDI, WorldPointPlaneLocation, OwnerActorLocation + WorldPointPlaneLocation + VisualizationContext.SpaceTransform.TransformPosition(ManifoldPoint.NetPushOut), TEXT("Net Push out"), PushOutImpusleColor);
+					}
+				}
+
+				if (EnumHasAnyFlags(static_cast<EChaosVDCollisionVisualizationFlags>(VisualizationContext.VisualizationFlags), EChaosVDCollisionVisualizationFlags::NetImpulse))
+				{
+					if (ManifoldPoint.bIsValid && !ManifoldPoint.NetImpulse.IsNearlyZero())
+					{
+						FChaosVDDebugDrawUtils::DrawArrowVector(PDI, WorldPointPlaneLocation, OwnerActorLocation + WorldPointPlaneLocation + VisualizationContext.SpaceTransform.TransformPosition(ManifoldPoint.NetImpulse), TEXT("Net Impulse"), ImpulseColor);
+					}
+				}
+
+				if (EnumHasAnyFlags(VisualizationFlags, EChaosVDCollisionVisualizationFlags::ContactPoints))
+				{
+					FChaosVDDebugDrawUtils::DrawPoint(PDI, WorldPlaneLocation, TEXT("Contact : WorldPlaneLocation"), DiscColor);
+				}
+
+				if (EnumHasAnyFlags(VisualizationFlags, EChaosVDCollisionVisualizationFlags::ContactNormal))
+				{
+					FColor NormalColor = ((ManifoldPoint.ContactPoint.ContactType != EChaosVDContactPointType::EdgeEdge) ? PlaneNormalColor : EdgeNormalColor);
+					constexpr int32 Scale = 10.0f;
+					FChaosVDDebugDrawUtils::DrawArrowVector(PDI, WorldPlaneLocation, WorldPlaneLocation + WorldPlaneNormal * Scale, TEXT("Contact Normal"), NormalColor, 2.0f);
+				}
+
+				if (EnumHasAnyFlags(VisualizationFlags, EChaosVDCollisionVisualizationFlags::ContactPoints))
+				{
+					if (ManifoldPoint.ContactPoint.Phi < FLT_MAX)
+					{
+						FColor Color = FColor(128, 128, 0);
+						FChaosVDDebugDrawUtils::DrawPoint(PDI, WorldPlaneLocation - ManifoldPoint.ContactPoint.Phi * WorldPlaneNormal, TEXT("Contact : WorldPlaneLocation - Phi"), Color);
+
+						FChaosVDDebugDrawUtils::DrawPoint(PDI, WorldPlaneLocation - ManifoldPoint.ContactPoint.Phi * WorldPlaneNormal, TEXT("Contact : WorldPlaneLocation - Phi"), Color);
+					}
+					
+					// Manifold point
+					FChaosVDDebugDrawUtils::DrawPoint(PDI, WorldPointLocation , TEXT("Contact : Manifold Point"), DiscColor);
+				}
+
+				if (EnumHasAnyFlags(VisualizationFlags, EChaosVDCollisionVisualizationFlags::AccumulatedImpulse))
+				{
+					if (!Constraint.AccumulatedImpulse.IsNearlyZero())
+					{
+						FChaosVDDebugDrawUtils::DrawArrowVector(PDI, WorldActorTransform0.GetLocation(), WorldActorTransform0.GetLocation() + Constraint.AccumulatedImpulse, TEXT("AccumulatedImpulse"), FColor::White, 5.0f);
+					}
+				}
+
+				PDI->SetHitProxy(nullptr);
+			}
+		}
+	}
+}
diff --git a/Engine/Plugins/Experimental/ChaosVD/Source/ChaosVD/Private/Visualizers/ChaosVDSolverCollisionDataComponentVisualizer.h b/Engine/Plugins/Experimental/ChaosVD/Source/ChaosVD/Private/Visualizers/ChaosVDSolverCollisionDataComponentVisualizer.h
new file mode 100644
--- /dev/null
+++ b/Engine/Plugins/Experimental/ChaosVD/Source/ChaosVD/Private/Visualizers/ChaosVDSolverCollisionDataComponentVisualizer.h
@@ -0,0 +1,57 @@
+// Copyright Epic Games, Inc. All Rights Reserved.
+
+#pragma once
+#include "ComponentVisualizer.h"
+#include "HitProxies.h"
+#include "DataWrappers/ChaosVDCollisionDataWrappers.h"
+
+struct FChaosVDParticlePairMidPhase;
+struct FChaosVDVisualizationContext;
+class FChaosVDScene;
+
+struct FChaosVDCollisionDataFinder
+{
+	TWeakPtr<FChaosVDParticlePairMidPhase> OwningMidPhase;
+	const FChaosVDConstraint* OwningConstraint = nullptr;
+	int32 ContactIndex = INDEX_NONE;
+};
+
+struct HChaosVDContactPointProxy : public HHitProxy
+{
+	DECLARE_HIT_PROXY()
+	
+	HChaosVDContactPointProxy(const FChaosVDCollisionDataFinder& InContactFinderData) : HHitProxy(HPP_UI), ContactFinder(InContactFinderData)
+	{	
+	}
+
+	FChaosVDCollisionDataFinder ContactFinder;
+};
+
+UENUM(meta = (Bitflags, UseEnumValuesAsMaskValuesInEditor = "true"))
+enum class EChaosVDCollisionVisualizationFlags: uint32
+{
+	None					= 0 UMETA(Hidden),
+	ContactPoints			= 1 << 0,
+	ContactInfo				= 1 << 1 UMETA(Hidden), // Not used yet
+	NetPushOut				= 1 << 2,
+	NetImpulse				= 1 << 3,
+	ContactNormal			= 1 << 4,
+	AccumulatedImpulse		= 1 << 5,
+	DrawOnlyActiveContacts	= 1 << 6,
+	EnableDrawForAllParticles	= 1 << 7,
+};
+ENUM_CLASS_FLAGS(EChaosVDCollisionVisualizationFlags);
+
+class FChaosVDSolverCollisionDataComponentVisualizer : public FComponentVisualizer
+{
+public:
+	FChaosVDSolverCollisionDataComponentVisualizer();
+	~FChaosVDSolverCollisionDataComponentVisualizer();
+
+	virtual bool ShowWhenSelected() override;
+	virtual void DrawVisualization(const UActorComponent* Component, const FSceneView* View, FPrimitiveDrawInterface* PDI) override;
+
+protected:
+
+	void DrawnMidPhaseData(const TSharedPtr<FChaosVDParticlePairMidPhase>& MidPhase, const FChaosVDVisualizationContext& VisualizationContext, const FSceneView* View, FPrimitiveDrawInterface* PDI);
+};
diff --git a/Engine/Plugins/Experimental/ChaosVD/Source/ChaosVD/Private/Widgets/SChaosVDBrowseSessionsModal.cpp b/Engine/Plugins/Experimental/ChaosVD/Source/ChaosVD/Private/Widgets/SChaosVDBrowseSessionsModal.cpp
--- a/Engine/Plugins/Experimental/ChaosVD/Source/ChaosVD/Private/Widgets/SChaosVDBrowseSessionsModal.cpp
+++ b/Engine/Plugins/Experimental/ChaosVD/Source/ChaosVD/Private/Widgets/SChaosVDBrowseSessionsModal.cpp
@@ -107,6 +107,7 @@ void SChaosVDBrowseSessionsModal::Construct(const FArguments& InArgs)
 			[
 				SNew(SVerticalBox)
 				+SVerticalBox::Slot()
+				.Padding(0.0f, 0.0f, 0.0f, 5.0f)
 				.AutoHeight()
 				[
 					SNew(SHorizontalBox)
diff --git a/Engine/Plugins/Experimental/ChaosVD/Source/ChaosVD/Private/Widgets/SChaosVDCollisionDataInspector.cpp b/Engine/Plugins/Experimental/ChaosVD/Source/ChaosVD/Private/Widgets/SChaosVDCollisionDataInspector.cpp
new file mode 100644
--- /dev/null
+++ b/Engine/Plugins/Experimental/ChaosVD/Source/ChaosVD/Private/Widgets/SChaosVDCollisionDataInspector.cpp
@@ -0,0 +1,350 @@
+// Copyright Epic Games, Inc. All Rights Reserved.
+
+#include "Widgets/SChaosVDCollisionDataInspector.h"
+
+#include "ChaosVDParticleActor.h"
+#include "ChaosVDScene.h"
+#include "IStructureDetailsView.h"
+#include "Modules/ModuleManager.h"
+#include "PropertyEditorModule.h"
+#include "Visualizers/ChaosVDSolverCollisionDataComponentVisualizer.h"
+#include "Widgets/Input/SButton.h"
+#include "Widgets/Layout/SScrollBox.h"
+#include "Widgets/Layout/SUniformGridPanel.h"
+#include "Widgets/SChaosVDNameListPicker.h"
+#include "Widgets/SChaosVDWarningMessageBox.h"
+#include "Widgets/Text/STextBlock.h"
+
+#define LOCTEXT_NAMESPACE "ChaosVisualDebugger"
+
+SChaosVDCollisionDataInspector::~SChaosVDCollisionDataInspector()
+{
+	if (TSharedPtr<FChaosVDScene> ScenePtr = SceneWeakPtr.Pin())
+	{
+		ScenePtr->OnSceneUpdated().RemoveAll(this);
+	}
+}
+
+void SChaosVDCollisionDataInspector::Construct(const FArguments& InArgs, const TWeakPtr<FChaosVDScene>& InScenePtr)
+{
+	SceneWeakPtr = InScenePtr;
+
+	if (TSharedPtr<FChaosVDScene> ScenePtr = SceneWeakPtr.Pin())
+	{
+		ScenePtr->OnSceneUpdated().AddRaw(this, &SChaosVDCollisionDataInspector::HandleSceneUpdated);
+	}
+
+	ContactDataDetailsView = CreateCollisionDataDetailsView();
+	ConstraintDataDetailsView = CreateCollisionDataDetailsView();
+
+	ChildSlot
+	[
+		SNew(SVerticalBox)
+		+SVerticalBox::Slot()
+		.AutoHeight()
+		.Padding(FMargin(10.0f, 10.0f, 10.0f, 0.0f))
+		[
+			GenerateObjectNameRowWidget().ToSharedRef()
+		]
+		+SVerticalBox::Slot()
+		.Padding(FMargin(10.0f, 5.0f, 10.0f, 1.0f))
+		.AutoHeight()
+		[
+			SAssignNew(CollisionDataAvailableList, SChaosVDNameListPicker)
+			.OnNameSleceted_Raw(this, &SChaosVDCollisionDataInspector::HandleSessionNameSelected)
+		]
+		+SVerticalBox::Slot()
+		.Padding(10)
+		[
+			SNew(SScrollBox)
+			.Visibility_Raw(this, &SChaosVDCollisionDataInspector::GetDetailsSectionVisibility)
+			+SScrollBox::Slot()
+			[
+				SNew(SBox)
+				.Padding(4.0f,2.0f,2.0f,10.0f)
+				[
+					SNew(SChaosVDWarningMessageBox)
+					.WarningText(LOCTEXT("CollisionDataOutOfData", "Scene change detected!. Selected collision data is out of date..."))
+					.Visibility_Raw(this, &SChaosVDCollisionDataInspector::GetOutOfDateWarningVisibility)
+				]
+			]
+			+SScrollBox::Slot()
+			.Padding(0.0f,5.0f,0.0f,15.0f)
+			[
+				SNew(SUniformGridPanel)
+				.SlotPadding(FAppStyle::GetMargin("StandardDialog.SlotPadding"))
+				.MinDesiredSlotWidth(FAppStyle::GetFloat("StandardDialog.MinDesiredSlotWidth"))
+				.MinDesiredSlotHeight(FAppStyle::GetFloat("StandardDialog.MinDesiredSlotHeight"))
+				+SUniformGridPanel::Slot(0, 0)
+				[
+					SNew(SButton)
+					.HAlign(HAlign_Center)
+					.ContentPadding(FAppStyle::GetMargin("StandardDialog.ContentPadding"))
+					.Text(LOCTEXT("SelectParticle0", "Select Particle 0"))
+					.OnClicked(this, &SChaosVDCollisionDataInspector::SelectParticleForCurrentCollisionData, EChaosVDCollisionParticleSelector::Index_0)
+				]
+				+SUniformGridPanel::Slot(1, 0)
+				[
+					SNew(SButton)
+					.HAlign(HAlign_Center)
+					.ContentPadding(FAppStyle::GetMargin("StandardDialog.ContentPadding"))
+					.Text(LOCTEXT("SelectParticle1", "Select Particle 1"))
+					.OnClicked(this, &SChaosVDCollisionDataInspector::SelectParticleForCurrentCollisionData, EChaosVDCollisionParticleSelector::Index_1)
+				]				
+			]
+			+SScrollBox::Slot()
+			.Padding(15.0f,0.0f,15.0f,0.0f)
+			[
+				ContactDataDetailsView->GetWidget().ToSharedRef()
+			]
+			+SScrollBox::Slot()
+			.Padding(15.0f,10.0f,15.0f,5.0f)
+			[
+				ConstraintDataDetailsView->GetWidget().ToSharedRef()
+			]
+		]
+	];
+}
+
+void SChaosVDCollisionDataInspector::SetCollisionDataProviderObjectToInspect(IChaosVDCollisionDataProviderInterface* CollisionDataProvider)
+{
+	ClearInspector();
+
+	if (CollisionDataProvider == nullptr)
+	{
+		return;
+	}
+
+	CurrentObjectBeingInspectedName = CollisionDataProvider->GetName();
+
+	TArray<TSharedPtr<FName>> NewCollisionDataEntriesNameList;
+	TArray<TSharedPtr<FChaosVDCollisionDataFinder>> FoundCollisionData;
+
+	CollisionDataProvider->GetCollisionData(FoundCollisionData);
+
+	for (const TSharedPtr<FChaosVDCollisionDataFinder>& CollisionData : FoundCollisionData)
+	{
+		if (TSharedPtr<FName> CollisionDataEntryName = GenerateNameForCollisionDataItem(*CollisionData))
+		{
+			CollisionDataByNameMap.Add(*CollisionDataEntryName, CollisionData);
+			NewCollisionDataEntriesNameList.Add(CollisionDataEntryName);
+		}
+	}
+
+	CollisionDataAvailableList->UpdateNameList(MoveTemp(NewCollisionDataEntriesNameList));
+	bIsUpToDate = true;
+}
+
+void SChaosVDCollisionDataInspector::SetSingleContactDataToInspect(const FChaosVDCollisionDataFinder& InContactFinderData)
+{
+	ClearInspector();
+
+	TSharedPtr<FChaosVDCollisionDataFinder> ContactFinderDataPtr = MakeShared<FChaosVDCollisionDataFinder>();
+	*ContactFinderDataPtr = InContactFinderData;
+
+	CurrentSelectedName = GenerateNameForCollisionDataItem(*ContactFinderDataPtr);
+	
+	// TODO: We need a way to provide a name when we receive collision data directly.
+	// Currently this is only possible when clicking contacts in the viewport
+	CurrentObjectBeingInspectedName = TEXT("Contact");
+
+	CollisionDataByNameMap.Add(*CurrentSelectedName, ContactFinderDataPtr);
+
+	CollisionDataAvailableList->UpdateNameList({ CurrentSelectedName });
+	CollisionDataAvailableList->SelectName(CurrentSelectedName, ESelectInfo::OnMouseClick);
+}
+
+void SChaosVDCollisionDataInspector::HandleSceneUpdated()
+{
+	if (TSharedPtr<FChaosVDCollisionDataFinder> DataFinderPtr = GetCurrentDataBeingInspected())
+	{
+		if (DataFinderPtr.IsValid() && DataFinderPtr->OwningMidPhase.Pin())
+		{
+			bIsUpToDate = false;
+		}
+		else
+		{
+			ClearInspector();
+		}
+	}
+}
+
+void SChaosVDCollisionDataInspector::ClearInspector()
+{
+	ContactDataDetailsView->SetStructureData(nullptr);
+	ConstraintDataDetailsView->SetStructureData(nullptr);
+
+	CollisionDataByNameMap.Reset();
+	CollisionDataAvailableList->UpdateNameList({});
+
+	CurrentSelectedName = nullptr;
+
+	CurrentObjectBeingInspectedName = FName();
+
+	bIsUpToDate = true;
+}
+
+FText SChaosVDCollisionDataInspector::GetObjectBeingInspectedName() const
+{
+	return FText::AsCultureInvariant(CurrentObjectBeingInspectedName.ToString());
+}
+
+TSharedPtr<SWidget> SChaosVDCollisionDataInspector::GenerateObjectNameRowWidget()
+{
+	TSharedPtr<SWidget> RowWidget = SNew(SBorder)
+		.BorderImage(FAppStyle::Get().GetBrush("DetailsView.CategoryMiddle"))
+		.BorderBackgroundColor(FAppStyle::Get().GetSlateColor("Colors.Panel"))
+		.Padding(0)
+		[
+			SNew(SBox)
+			.Padding(2.0f)
+			[
+				SNew(SSplitter)
+				.Style(FAppStyle::Get(), "DetailsView.Splitter")
+				.PhysicalSplitterHandleSize(1.5f)
+				.HitDetectionSplitterHandleSize(5.0f)
+				+SSplitter::Slot()
+				.Value(0.2)
+				[
+					SNew(SBox)
+					.VAlign(VAlign_Center)
+					.Padding(1.0)
+					.VAlign(VAlign_Center)
+					[
+						SNew(STextBlock)
+						.Text(LOCTEXT("CollisionDataInspectorObject","Object"))
+						.Font(FAppStyle::GetFontStyle("PropertyWindow.NormalFont"))
+					]	
+				]
+				+SSplitter::Slot()
+				[
+					SNew(SBox)
+					.VAlign(VAlign_Center)
+					.Padding(1.0)
+					[
+						SNew(STextBlock)
+						.Text_Raw(this, &SChaosVDCollisionDataInspector::GetObjectBeingInspectedName)
+						.Font(FAppStyle::GetFontStyle("PropertyWindow.NormalFont"))
+					]
+				]
+			]
+		];
+
+	return RowWidget;
+}
+
+
+EVisibility SChaosVDCollisionDataInspector::GetOutOfDateWarningVisibility() const
+{
+	return bIsUpToDate ? EVisibility::Collapsed : EVisibility::Visible;
+}
+
+void SChaosVDCollisionDataInspector::HandleSessionNameSelected(TSharedPtr<FName> SelectedName)
+{
+	CurrentSelectedName = SelectedName;
+
+	if (!CurrentSelectedName.IsValid())
+	{
+		return;
+	}
+
+	if (TSharedPtr<FChaosVDCollisionDataFinder> DataFinderPtr = GetCurrentDataBeingInspected())
+	{
+		if (DataFinderPtr->OwningConstraint)
+		{
+			FChaosVDConstraint* MutableConstraint = const_cast<FChaosVDConstraint*>(DataFinderPtr->OwningConstraint);
+			TSharedPtr<FStructOnScope> ConstraintView = MakeShared<FStructOnScope>(FChaosVDConstraint::StaticStruct(), reinterpret_cast<uint8*>(MutableConstraint));
+			ConstraintDataDetailsView->SetStructureData(ConstraintView); 
+
+			if (DataFinderPtr->OwningConstraint->ManifoldPoints.IsValidIndex(DataFinderPtr->ContactIndex))
+			{
+				const FChaosVDManifoldPoint* ContactPointData = &DataFinderPtr->OwningConstraint->ManifoldPoints[DataFinderPtr->ContactIndex];
+				TSharedPtr<FStructOnScope> ContactViewView = MakeShared<FStructOnScope>(FChaosVDManifoldPoint::StaticStruct(), reinterpret_cast<uint8*>(const_cast<FChaosVDManifoldPoint*>(ContactPointData)));
+				ContactDataDetailsView->SetStructureData(ContactViewView);
+			}
+			else
+			{
+				ContactDataDetailsView->SetStructureData(nullptr); 
+			}
+		}
+	}
+}
+
+TSharedPtr<FName> SChaosVDCollisionDataInspector::GenerateNameForCollisionDataItem(const FChaosVDCollisionDataFinder& InContactFinderData)
+{
+	if (InContactFinderData.OwningMidPhase.Pin())
+	{
+		if (InContactFinderData.OwningConstraint)
+		{
+			const FText GeneratedName = FText::Format(LOCTEXT("CollisionIemDataTitle", "Particle Pair | Index0 [ID {0}] <-> Index1 [{1}]"), InContactFinderData.OwningConstraint->Particle0Index, InContactFinderData.OwningConstraint->Particle1Index);
+			return MakeShared<FName>(GeneratedName.ToString());
+		}
+	}
+
+	return nullptr;
+}
+
+TSharedPtr<IStructureDetailsView> SChaosVDCollisionDataInspector::CreateCollisionDataDetailsView()
+{
+	FPropertyEditorModule& PropertyEditorModule = FModuleManager::GetModuleChecked<FPropertyEditorModule>("PropertyEditor");
+
+	FStructureDetailsViewArgs StructDetailsViewArgs;
+	FDetailsViewArgs DetailsViewArgs;
+	DetailsViewArgs.bShowOptions = false;
+	DetailsViewArgs.bAllowFavoriteSystem = false;
+	DetailsViewArgs.bAllowSearch = false;
+	DetailsViewArgs.bShowScrollBar = false;
+
+	return PropertyEditorModule.CreateStructureDetailView(DetailsViewArgs,StructDetailsViewArgs, nullptr);
+}
+
+EVisibility SChaosVDCollisionDataInspector::GetDetailsSectionVisibility() const
+{
+	return CurrentSelectedName ? EVisibility::Visible : EVisibility::Collapsed;
+}
+
+FReply SChaosVDCollisionDataInspector::SelectParticleForCurrentCollisionData(EChaosVDCollisionParticleSelector ParticleSelector)
+{
+	TSharedPtr<FChaosVDScene> ScenePtr = SceneWeakPtr.Pin();
+	if (!ScenePtr.IsValid())
+	{
+		return FReply::Handled();
+	}
+	
+	if (TSharedPtr<FChaosVDCollisionDataFinder> DataFinderPtr = GetCurrentDataBeingInspected())
+	{
+		if (FChaosVDMidPhasePtr MidPhasePtr = DataFinderPtr->OwningMidPhase.Pin())
+		{
+			if (DataFinderPtr->OwningConstraint)
+			{
+				const int32 ParticleIndex = ParticleSelector == EChaosVDCollisionParticleSelector::Index_0 ? DataFinderPtr->OwningConstraint->Particle0Index : DataFinderPtr->OwningConstraint->Particle1Index;
+				if (AChaosVDParticleActor* ParticleActor = ScenePtr->GetParticleActor(MidPhasePtr->SolverID, ParticleIndex))
+				{
+					ScenePtr->SetSelectedObject(ParticleActor);
+				}
+			}
+		}
+	}
+
+	return FReply::Handled();
+}
+
+TSharedPtr<FChaosVDCollisionDataFinder> SChaosVDCollisionDataInspector::GetCurrentDataBeingInspected()
+{
+	if (!CurrentSelectedName.IsValid())
+	{
+		return nullptr;
+	}
+
+	if (TSharedPtr<FChaosVDCollisionDataFinder>* DataFinderPtrPtr = CollisionDataByNameMap.Find(*CurrentSelectedName))
+	{
+		if (TSharedPtr<FChaosVDCollisionDataFinder> DataFinderPtr = *DataFinderPtrPtr)
+		{
+			return DataFinderPtr;
+		}
+	}
+
+	return nullptr;
+}
+
+#undef LOCTEXT_NAMESPACE
diff --git a/Engine/Plugins/Experimental/ChaosVD/Source/ChaosVD/Private/Widgets/SChaosVDDetailsView.cpp b/Engine/Plugins/Experimental/ChaosVD/Source/ChaosVD/Private/Widgets/SChaosVDDetailsView.cpp
--- a/Engine/Plugins/Experimental/ChaosVD/Source/ChaosVD/Private/Widgets/SChaosVDDetailsView.cpp
+++ b/Engine/Plugins/Experimental/ChaosVD/Source/ChaosVD/Private/Widgets/SChaosVDDetailsView.cpp
@@ -19,6 +19,7 @@ void SChaosVDDetailsView::Construct(const FArguments& InArgs)
 	DetailsViewArgs.bCustomNameAreaLocation = true;
 	DetailsViewArgs.bCustomFilterAreaLocation = false;
 	DetailsViewArgs.bShowSectionSelector = false;
+	DetailsViewArgs.bShowScrollBar = false;
 
 	DetailsView = PropertyEditorModule.CreateDetailView(DetailsViewArgs);
 
@@ -41,14 +42,15 @@ void SChaosVDDetailsView::Construct(const FArguments& InArgs)
 		+SVerticalBox::Slot()
 		[
 			SNew(SSplitter)
-			.MinimumSlotHeight(80.0f)
+			.MinimumSlotHeight(40.0f)
 			.Orientation(Orient_Vertical)
 			.Style(FAppStyle::Get(), "SplitterDark")
 			.PhysicalSplitterHandleSize(2.0f)
 			+SSplitter::Slot()
 			.Value(0.2f)
 			[
-				SNew(SBox)
+				SNew(SVerticalBox)
+				+SVerticalBox::Slot()
 				.Padding(FMargin(2.0f, 0.0f, 2.0f, 0.0f))
 				[
 					SubobjectEditor.ToSharedRef()
@@ -56,7 +58,7 @@ void SChaosVDDetailsView::Construct(const FArguments& InArgs)
 			]
 			+SSplitter::Slot()
 			[
-				SNew( SVerticalBox )
+				SNew(SVerticalBox)
 				+ SVerticalBox::Slot()
 				[
 					DetailsView.ToSharedRef()
diff --git a/Engine/Plugins/Experimental/ChaosVD/Source/ChaosVD/Private/Widgets/SChaosVDMainTab.cpp b/Engine/Plugins/Experimental/ChaosVD/Source/ChaosVD/Private/Widgets/SChaosVDMainTab.cpp
--- a/Engine/Plugins/Experimental/ChaosVD/Source/ChaosVD/Private/Widgets/SChaosVDMainTab.cpp
+++ b/Engine/Plugins/Experimental/ChaosVD/Source/ChaosVD/Private/Widgets/SChaosVDMainTab.cpp
@@ -2,6 +2,7 @@
 
 #include "Widgets/SChaosVDMainTab.h"
 
+#include "ChaosVDEditorModeTools.h"
 #include "ChaosVDEditorVisualizationSettingsTab.h"
 #include "ChaosVDEngine.h"
 #include "ChaosVDModule.h"
@@ -9,7 +10,9 @@
 #include "ChaosVDOutputLogTab.h"
 #include "ChaosVDPlaybackController.h"
 #include "ChaosVDPlaybackViewportTab.h"
+#include "ChaosVDScene.h"
 #include "ChaosVDSolversTracksTab.h"
+#include "ChaosVDCollisionDataDetailsTab.h"
 #include "ChaosVDStyle.h"
 #include "ChaosVDTabsIDs.h"
 #include "ChaosVDWorldOutlinerTab.h"
@@ -22,9 +25,11 @@
 #include "IDesktopPlatform.h"
 #include "Misc/MessageDialog.h"
 #include "StatusBarSubsystem.h"
+#include "Components/ChaosVDSolverCollisionDataComponent.h"
 #include "Styling/StyleColors.h"
 #include "Styling/ToolBarStyle.h"
 #include "Trace/ChaosVDTraceManager.h"
+#include "Visualizers/ChaosVDSolverCollisionDataComponentVisualizer.h"
 #include "Widgets/SChaosBrowseTraceFileSourceModal.h"
 #include "Widgets/SChaosVDBrowseSessionsModal.h"
 #include "Widgets/SChaosVDRecordingControls.h"
@@ -35,21 +40,25 @@
 
 #define LOCTEXT_NAMESPACE "ChaosVisualDebugger"
 
-
 void SChaosVDMainTab::Construct(const FArguments& InArgs, TSharedPtr<FChaosVDEngine> InChaosVDEngine)
 {
+	EditorModeTools = MakeShared<FChaosVDEditorModeTools>(InChaosVDEngine->GetCurrentScene());
+
+	EditorModeTools->SetToolkitHost(StaticCastSharedRef<SChaosVDMainTab>(AsShared()));
 	ChaosVDEngine = InChaosVDEngine;
+	OwnerTab = InArgs._OwnerTab;
+
+	RegisterComponentVisualizer(UChaosVDSolverCollisionDataComponent::StaticClass()->GetFName(), MakeShared<FChaosVDSolverCollisionDataComponentVisualizer>());
 
 	TabManager = FGlobalTabmanager::Get()->NewTabManager(InArgs._OwnerTab.ToSharedRef()).ToSharedPtr();
 
-	// Create the UI Tabs Handlers
-	//TODO: Pass self as a weakPtr
-	WorldOutlinerTab = MakeShared<FChaosVDWorldOutlinerTab>(FChaosVDTabID::WorldOutliner, TabManager, this);
-	ObjectDetailsTab = MakeShared<FChaosVDObjectDetailsTab>(FChaosVDTabID::DetailsPanel, TabManager, this);
-	OutputLogTab = MakeShared<FChaosVDOutputLogTab>(FChaosVDTabID::OutputLog, TabManager, this);
-	PlaybackViewportTab = MakeShared<FChaosVDPlaybackViewportTab>(FChaosVDTabID::PlaybackViewport, TabManager, this);
-	SolversTracksTab = MakeShared<FChaosVDSolversTracksTab>(FChaosVDTabID::SolversTrack, TabManager, this);
-	EditorSettingsTab = MakeShared<FChaosVDEditorVisualizationSettingsTab>(FChaosVDTabID::CVDEditorSettings, TabManager, this);
+	RegisterTabSpawner<FChaosVDWorldOutlinerTab>(FChaosVDTabID::WorldOutliner);
+	RegisterTabSpawner<FChaosVDObjectDetailsTab>(FChaosVDTabID::DetailsPanel);
+	RegisterTabSpawner<FChaosVDOutputLogTab>(FChaosVDTabID::OutputLog);
+	RegisterTabSpawner<FChaosVDPlaybackViewportTab>(FChaosVDTabID::PlaybackViewport);
+	RegisterTabSpawner<FChaosVDSolversTracksTab>(FChaosVDTabID::SolversTrack);
+	RegisterTabSpawner<FChaosVDEditorVisualizationSettingsTab>(FChaosVDTabID::CVDEditorSettings);
+	RegisterTabSpawner<FChaosVDCollisionDataDetailsTab>(FChaosVDTabID::CollisionDataDetails);
 
 	StatusBarID = FName(FChaosVDTabID::StatusBar.ToString() + InChaosVDEngine->GetInstanceGuid().ToString());
 	
@@ -234,7 +243,79 @@ void SChaosVDMainTab::Construct(const FArguments& InArgs, TSharedPtr<FChaosVDEng
 		]
 	];
 
+	// Make sure these tabs are always focused at the start
 	TabManager->TryInvokeTab(FChaosVDTabID::SolversTrack);
+	TabManager->TryInvokeTab(FChaosVDTabID::DetailsPanel);
+}
+
+void SChaosVDMainTab::BringToFront()
+{
+	if (TabManager.IsValid())
+	{
+		if (const TSharedPtr<SDockTab> TabPtr = OwnerTab.Pin())
+		{
+			TabManager->DrawAttention(TabPtr.ToSharedRef());
+		}
+	}
+}
+
+void SChaosVDMainTab::OnToolkitHostingStarted(const TSharedRef<IToolkit>& Toolkit)
+{
+}
+
+void SChaosVDMainTab::OnToolkitHostingFinished(const TSharedRef<IToolkit>& Toolkit)
+{
+}
+
+UWorld* SChaosVDMainTab::GetWorld() const
+{
+	return GetChaosVDEngineInstance()->GetCurrentScene()->GetUnderlyingWorld();
+}
+
+FEditorModeTools& SChaosVDMainTab::GetEditorModeManager() const
+{
+	check(EditorModeTools.IsValid())
+	return *EditorModeTools;
+}
+
+TSharedPtr<FComponentVisualizer> SChaosVDMainTab::FindComponentVisualizer(UClass* ClassPtr)
+{
+	TSharedPtr<FComponentVisualizer> Visualizer;
+	while (!Visualizer.IsValid() && (ClassPtr != nullptr) && (ClassPtr != UActorComponent::StaticClass()))
+	{
+		Visualizer = FindComponentVisualizer(ClassPtr->GetFName());
+		ClassPtr = ClassPtr->GetSuperClass();
+	}
+
+	return Visualizer;
+}
+
+TSharedPtr<FComponentVisualizer> SChaosVDMainTab::FindComponentVisualizer(FName ClassName)
+{
+	TSharedPtr<FComponentVisualizer>* FoundVisualizer = ComponentVisualizersMap.Find(ClassName);
+
+	return FoundVisualizer ? *FoundVisualizer : nullptr;
+}
+
+void SChaosVDMainTab::RegisterComponentVisualizer(FName ClassName, const TSharedPtr<FComponentVisualizer>& Visualizer)
+{
+	if (!ComponentVisualizersMap.Contains(ClassName))
+	{
+		ComponentVisualizersMap.Add(ClassName, Visualizer);
+	}
+}
+
+void SChaosVDMainTab::HandleTabSpawned(TSharedRef<SDockTab> Tab, FName TabID)
+{
+	if (!ActiveTabsByID.Contains(TabID))
+	{
+		ActiveTabsByID.Add(TabID, Tab);
+	}
+}
+
+void SChaosVDMainTab::HandleTabDestroyed(TSharedRef<SDockTab> Tab, FName TabID)
+{
+	ActiveTabsByID.Remove(TabID);
 }
 
 TSharedRef<FTabManager::FLayout> SChaosVDMainTab::GenerateMainLayout()
@@ -281,6 +362,7 @@ TSharedRef<FTabManager::FLayout> SChaosVDMainTab::GenerateMainLayout()
 					FTabManager::NewStack()
 					->SetSizeCoefficient(0.3f)
 					->AddTab(FChaosVDTabID::DetailsPanel, ETabState::OpenedTab)
+					->AddTab(FChaosVDTabID::CollisionDataDetails, ETabState::OpenedTab)
 					->AddTab(FChaosVDTabID::CVDEditorSettings, ETabState::ClosedTab)
 				)
 			)
diff --git a/Engine/Plugins/Experimental/ChaosVD/Source/ChaosVD/Private/Widgets/SChaosVDNameListPicker.cpp b/Engine/Plugins/Experimental/ChaosVD/Source/ChaosVD/Private/Widgets/SChaosVDNameListPicker.cpp
--- a/Engine/Plugins/Experimental/ChaosVD/Source/ChaosVD/Private/Widgets/SChaosVDNameListPicker.cpp
+++ b/Engine/Plugins/Experimental/ChaosVD/Source/ChaosVD/Private/Widgets/SChaosVDNameListPicker.cpp
@@ -9,6 +9,8 @@
 #include "Widgets/SBoxPanel.h"
 #include "Widgets/SBoxPanel.h"
 
+#define LOCTEXT_NAMESPACE "ChaosVisualDebugger"
+
 void SChaosVDNameListPicker::Construct(const FArguments& InArgs)
 {
 	NameSelectedDelegate = InArgs._OnNameSleceted;
@@ -34,7 +36,6 @@ void SChaosVDNameListPicker::Construct(const FArguments& InArgs)
 		.Padding(2)
 		[
 			SNew(SBox)
-			.WidthOverride(250)
 			[				
 				SNew(SVerticalBox)
 				+SVerticalBox::Slot()
@@ -50,6 +51,7 @@ void SChaosVDNameListPicker::Construct(const FArguments& InArgs)
 	[
 		// Combo button that summons the dropdown menu
 		SAssignNew(PickerComboButton, SComboButton)
+			.IsEnabled_Raw(this, &SChaosVDNameListPicker::HasElements)
 			.ButtonContent()
 			[
 				DropDownButtonLabelTextBox.ToSharedRef()
@@ -62,16 +64,33 @@ void SChaosVDNameListPicker::Construct(const FArguments& InArgs)
 			.ContentPadding(2.0f)
 			.OnComboBoxOpened(this, &SChaosVDNameListPicker::OnListOpened)
 	];
-	
 }
 
 void SChaosVDNameListPicker::UpdateNameList(TArray<TSharedPtr<FName>>&& NewNameList)
 {
 	// TODO: It might be worth to check is the array is different before doing this, but it might be more expensive
 	CachedNameList = MoveTemp(NewNameList);
+
+	if(!HasElements())
+	{
+		CurrentSelectedName = nullptr;
+	}
+
 	NameListWidget->RebuildList();
 }
 
+void SChaosVDNameListPicker::SelectName(const TSharedPtr<FName>& NameToSelect, ESelectInfo::Type SelectionInfo)
+{
+	NameListWidget->SetSelection(NameToSelect, ESelectInfo::OnMouseClick);
+	NameListWidget->RequestListRefresh();
+	OnClickItem(NameToSelect);
+}
+
+bool SChaosVDNameListPicker::HasElements() const
+{
+	return CachedNameList.Num() > 0;
+}
+
 FText SChaosVDNameListPicker::GetCurrentDropDownButtonLabel() const
 {
 	if (CurrentSelectedName.IsValid())
@@ -80,7 +99,7 @@ FText SChaosVDNameListPicker::GetCurrentDropDownButtonLabel() const
 	}
 	else
 	{
-		return NSLOCTEXT("ChaosVisualDebugger", "NameList", "Pick an element...");
+		return HasElements() ? LOCTEXT("NameList", "Pick an element...") : LOCTEXT("EmptyList", "Empty...");
 	}
 }
 
@@ -123,3 +142,5 @@ TSharedRef<ITableRow> SChaosVDNameListPicker::MakeNameItemForList(TSharedPtr<FNa
 
 	return ListTableRow.ToSharedRef();
 }
+
+#undef LOCTEXT_NAMESPACE
diff --git a/Engine/Plugins/Experimental/ChaosVD/Source/ChaosVD/Private/Widgets/SChaosVDPlaybackViewport.cpp b/Engine/Plugins/Experimental/ChaosVD/Source/ChaosVD/Private/Widgets/SChaosVDPlaybackViewport.cpp
--- a/Engine/Plugins/Experimental/ChaosVD/Source/ChaosVD/Private/Widgets/SChaosVDPlaybackViewport.cpp
+++ b/Engine/Plugins/Experimental/ChaosVD/Source/ChaosVD/Private/Widgets/SChaosVDPlaybackViewport.cpp
@@ -23,11 +23,11 @@ SChaosVDPlaybackViewport::~SChaosVDPlaybackViewport()
 	PlaybackViewportClient.Reset();
 }
 
-void SChaosVDPlaybackViewport::Construct(const FArguments& InArgs, TWeakPtr<FChaosVDScene> InScene, TWeakPtr<FChaosVDPlaybackController> InPlaybackController)
+void SChaosVDPlaybackViewport::Construct(const FArguments& InArgs, TWeakPtr<FChaosVDScene> InScene, TWeakPtr<FChaosVDPlaybackController> InPlaybackController, TSharedPtr<FEditorModeTools> InEditorModeTools)
 {
 	Extender = MakeShared<FExtender>();
 
-	EditorModeTools = MakeShared<FChaosVDEditorModeTools>(InScene);
+	EditorModeTools = InEditorModeTools;
 	EditorModeTools->SetWidgetMode(UE::Widget::WM_Translate);
 	EditorModeTools->SetDefaultMode(UChaosVDEditorMode::EM_ChaosVisualDebugger);
 	EditorModeTools->ActivateDefaultMode();
diff --git a/Engine/Plugins/Experimental/ChaosVD/Source/ChaosVD/Private/Widgets/SChaosVDRecordingControls.cpp b/Engine/Plugins/Experimental/ChaosVD/Source/ChaosVD/Private/Widgets/SChaosVDRecordingControls.cpp
--- a/Engine/Plugins/Experimental/ChaosVD/Source/ChaosVD/Private/Widgets/SChaosVDRecordingControls.cpp
+++ b/Engine/Plugins/Experimental/ChaosVD/Source/ChaosVD/Private/Widgets/SChaosVDRecordingControls.cpp
@@ -23,7 +23,7 @@ void SChaosVDRecordingControls::Construct(const FArguments& InArgs, const TWeakP
 	if (TSharedPtr<SChaosVDMainTab> MainTabSharedPtr = InMainTabWeakPtr.Pin())
 	{
 		MainTabWeakPtr = MainTabSharedPtr;
-		StatusBarID = MainTabSharedPtr->GetStatusBarID();
+		StatusBarID = MainTabSharedPtr->GetStatusBarName();
 	}
 	else
 	{
diff --git a/Engine/Plugins/Experimental/ChaosVD/Source/ChaosVD/Private/Widgets/SChaosVDWarningMessageBox.cpp b/Engine/Plugins/Experimental/ChaosVD/Source/ChaosVD/Private/Widgets/SChaosVDWarningMessageBox.cpp
new file mode 100644
--- /dev/null
+++ b/Engine/Plugins/Experimental/ChaosVD/Source/ChaosVD/Private/Widgets/SChaosVDWarningMessageBox.cpp
@@ -0,0 +1,41 @@
+// Copyright Epic Games, Inc. All Rights Reserved.
+
+#include "Widgets/SChaosVDWarningMessageBox.h"
+
+#include "Components/HorizontalBox.h"
+#include "Styling/AppStyle.h"
+
+#include "Widgets/Images/SImage.h"
+#include "Widgets/Layout/SBorder.h"
+#include "Widgets/Text/STextBlock.h"
+
+void SChaosVDWarningMessageBox::Construct(const FArguments& InArgs)
+{
+	ChildSlot
+	[
+		SNew(SBorder)
+		.BorderImage(FAppStyle::Get().GetBrush("Brushes.Panel"))
+		[
+			SNew(SHorizontalBox)
+			+ SHorizontalBox::Slot()
+			.AutoWidth()
+			.HAlign(HAlign_Center)
+			.VAlign(VAlign_Center)
+			.Padding(2)
+			[
+				SNew(SImage)
+				.Image(FAppStyle::Get().GetBrush("Icons.Warning"))
+			]
+			+ SHorizontalBox::Slot()
+			.VAlign(VAlign_Center)
+			.Padding(2)
+			[
+				SNew(STextBlock)
+				.Justification(ETextJustify::Center)
+				.TextStyle(FAppStyle::Get(), "DetailsView.BPMessageTextStyle")
+				.Text(InArgs._WarningText)
+				.AutoWrapText(true)
+			]
+		]
+	];
+}
diff --git a/Engine/Plugins/Experimental/ChaosVD/Source/ChaosVD/Public/Actors/ChaosVDSolverInfoActor.h b/Engine/Plugins/Experimental/ChaosVD/Source/ChaosVD/Public/Actors/ChaosVDSolverInfoActor.h
new file mode 100644
--- /dev/null
+++ b/Engine/Plugins/Experimental/ChaosVD/Source/ChaosVD/Public/Actors/ChaosVDSolverInfoActor.h
@@ -0,0 +1,60 @@
+// Copyright Epic Games, Inc. All Rights Reserved.
+
+#pragma once
+
+#include "ChaosVDSceneObjectBase.h"
+#include "ChaosVDSceneSelectionObserver.h"
+#include "GameFramework/Actor.h"
+#include "ChaosVDSolverInfoActor.generated.h"
+
+class AChaosVDParticleActor;
+class UChaosVDSolverCollisionDataComponent;
+
+UCLASS()
+class AChaosVDSolverInfoActor : public AActor, public FChaosVDSceneObjectBase, public FChaosVDSceneSelectionObserver
+{
+	GENERATED_BODY()
+
+public:
+
+	AChaosVDSolverInfoActor(const FObjectInitializer& ObjectInitializer);
+
+	void SetSolverID(int32 InSolverID) { SolverID = InSolverID; }
+	int32 GetSolverID() const { return SolverID; }
+
+	virtual void SetScene(TWeakPtr<FChaosVDScene> InScene) override;
+
+	void SetSimulationTransform(const FTransform& InSimulationTransform) { SimulationTransform = InSimulationTransform; }
+	const FTransform& GetSimulationTransform() const { return SimulationTransform; }
+
+	UChaosVDSolverCollisionDataComponent* GetCollisionDataComponent() { return CollisionDataComponent; }
+
+	void RegisterParticleActor(int32 ParticleID, AChaosVDParticleActor* ParticleActor);
+
+	AChaosVDParticleActor* GetParticleActor(int32 ParticleID);
+	const TMap<int32, AChaosVDParticleActor*>& GetAllParticleActorsByIDMap() { return  SolverParticlesByID; }
+
+	const TArray<int32>& GetSelectedParticlesIDs() const { return SelectedParticlesID; }
+
+	void HandleVisibilitySettingsUpdated();
+
+	virtual void BeginDestroy() override;
+
+protected:
+
+	virtual void HandlePostSelectionChange(const UTypedElementSelectionSet* ChangesSelectionSet) override;
+
+	UPROPERTY(VisibleAnywhere, Category="Solver Data")
+	int32 SolverID = INDEX_NONE;
+
+	UPROPERTY(VisibleAnywhere, Category="Solver Data")
+	FTransform SimulationTransform;
+
+	UPROPERTY()
+	TObjectPtr<UChaosVDSolverCollisionDataComponent> CollisionDataComponent;
+
+	UPROPERTY()
+	TMap<int32, AChaosVDParticleActor*> SolverParticlesByID;
+
+	TArray<int32> SelectedParticlesID;
+};
diff --git a/Engine/Plugins/Experimental/ChaosVD/Source/ChaosVD/Public/ChaosVDCollisionDataDetailsTab.h b/Engine/Plugins/Experimental/ChaosVD/Source/ChaosVD/Public/ChaosVDCollisionDataDetailsTab.h
new file mode 100644
--- /dev/null
+++ b/Engine/Plugins/Experimental/ChaosVD/Source/ChaosVD/Public/ChaosVDCollisionDataDetailsTab.h
@@ -0,0 +1,28 @@
+// Copyright Epic Games, Inc. All Rights Reserved.
+
+#pragma once
+
+#include "ChaosVDTabSpawnerBase.h"
+#include "Templates/SharedPointer.h"
+
+class SChaosVDCollisionDataInspector;
+class IStructureDetailsView;
+class FStructOnScope;
+
+class FChaosVDCollisionDataDetailsTab : public FChaosVDTabSpawnerBase
+{
+public:
+	FChaosVDCollisionDataDetailsTab(const FName& InTabID, const TSharedPtr<FTabManager>& InTabManager, const TWeakPtr<SChaosVDMainTab>& InOwningTabWidget)
+		: FChaosVDTabSpawnerBase(InTabID, InTabManager, InOwningTabWidget)
+	{
+	}
+
+	virtual ~FChaosVDCollisionDataDetailsTab() override;
+	virtual TSharedRef<SDockTab> HandleTabSpawned(const FSpawnTabArgs& Args) override;
+
+	TWeakPtr<SChaosVDCollisionDataInspector> GetCollisionInspectorInstance() const { return CollisionDataInspector; };
+
+protected:
+
+	TSharedPtr<SChaosVDCollisionDataInspector> CollisionDataInspector;
+};
diff --git a/Engine/Plugins/Experimental/ChaosVD/Source/ChaosVD/Public/ChaosVDCollisionDataProviderInterface.h b/Engine/Plugins/Experimental/ChaosVD/Source/ChaosVD/Public/ChaosVDCollisionDataProviderInterface.h
new file mode 100644
--- /dev/null
+++ b/Engine/Plugins/Experimental/ChaosVD/Source/ChaosVD/Public/ChaosVDCollisionDataProviderInterface.h
@@ -0,0 +1,42 @@
+// Copyright Epic Games, Inc. All Rights Reserved.
+
+#pragma once
+
+#include "UObject/Interface.h"
+#include "ChaosVDCollisionDataProviderInterface.generated.h"
+
+struct FChaosVDCollisionDataFinder;
+
+// This class does not need to be modified.
+UINTERFACE()
+class UChaosVDCollisionDataProviderInterface : public UInterface
+{
+	GENERATED_BODY()
+};
+
+/**
+ * Interface for any object that is able to provide CVD Collision Data
+ */
+class IChaosVDCollisionDataProviderInterface
+{
+	GENERATED_BODY()
+
+public:
+	
+	/**
+	 * Gathers and populates the provided array with any existing collision data for this object (if any) 
+	 * @param OutCollisionDataFound Array to populate with collision data
+	 */
+	virtual void GetCollisionData(TArray<TSharedPtr<FChaosVDCollisionDataFinder>>& OutCollisionDataFound) PURE_VIRTUAL(IChaosVDCollisionDataProviderInterface::GetCollisionData)
+
+	/**
+	 * Checks if the object implementing the interface has any collision data available
+	 * @return True if any collision data is found 
+	 */
+	virtual bool HasCollisionData() PURE_VIRTUAL(IChaosVDCollisionDataProviderInterface::HasCollisionData, return false;)
+
+	/**
+	 * Returns the name of the object providing the collision data
+	 */
+	virtual FName GetName() PURE_VIRTUAL(IChaosVDCollisionDataProviderInterface::HasCollisionData, return FName();)
+};
diff --git a/Engine/Plugins/Experimental/ChaosVD/Source/ChaosVD/Public/ChaosVDEditorSettings.h b/Engine/Plugins/Experimental/ChaosVD/Source/ChaosVD/Public/ChaosVDEditorSettings.h
--- a/Engine/Plugins/Experimental/ChaosVD/Source/ChaosVD/Public/ChaosVDEditorSettings.h
+++ b/Engine/Plugins/Experimental/ChaosVD/Source/ChaosVD/Public/ChaosVDEditorSettings.h
@@ -7,6 +7,8 @@
 #include "Engine/StaticMesh.h"
 #include "UObject/Object.h"
 #include "UObject/SoftObjectPtr.h"
+#include "Visualizers/ChaosVDParticleDataVisualizer.h"
+#include "Visualizers/ChaosVDSolverCollisionDataComponentVisualizer.h"
 
 #include "ChaosVDEditorSettings.generated.h"
 
@@ -49,10 +51,10 @@ class UChaosVDEditorSettings : public UObject
 public:
 
 	UPROPERTY(EditAnywhere, Category = "Viewport Visualization Flags", meta = (Bitmask, BitmaskEnum = "/Script/ChaosVD.EChaosVDParticleDataVisualizationFlags"))
-	uint8 GlobalParticleDataVisualizationFlags = 0;
+	uint32 GlobalParticleDataVisualizationFlags = 0;
 	
 	UPROPERTY(EditAnywhere, Category = "Viewport Visualization Flags", meta = (Bitmask, BitmaskEnum = "/Script/ChaosVD.EChaosVDCollisionVisualizationFlags"))
-	uint8 GlobalCollisionDataVisualizationFlags = 0;
+	uint32 GlobalCollisionDataVisualizationFlags = 0;
 
 	UPROPERTY(EditAnywhere, Category = "Viewport Visualization Flags", meta = (Bitmask, BitmaskEnum = "/Script/ChaosVD.EChaosVDCollisionVisualizationFlags"))
 	bool bShowDebugText = false;
diff --git a/Engine/Plugins/Experimental/ChaosVD/Source/ChaosVD/Public/ChaosVDEditorVisualizationSettingsTab.h b/Engine/Plugins/Experimental/ChaosVD/Source/ChaosVD/Public/ChaosVDEditorVisualizationSettingsTab.h
--- a/Engine/Plugins/Experimental/ChaosVD/Source/ChaosVD/Public/ChaosVDEditorVisualizationSettingsTab.h
+++ b/Engine/Plugins/Experimental/ChaosVD/Source/ChaosVD/Public/ChaosVDEditorVisualizationSettingsTab.h
@@ -9,7 +9,7 @@ class SChaosVDNameListPicker;
 class FChaosVDEditorVisualizationSettingsTab : public FChaosVDTabSpawnerBase
 {
 public:
-	FChaosVDEditorVisualizationSettingsTab(const FName& InTabID, const TSharedPtr<FTabManager>& InTabManager, SChaosVDMainTab* InOwningTabWidget)
+	FChaosVDEditorVisualizationSettingsTab(const FName& InTabID, const TSharedPtr<FTabManager>& InTabManager, TWeakPtr<SChaosVDMainTab> InOwningTabWidget)
 		: FChaosVDTabSpawnerBase(InTabID, InTabManager, InOwningTabWidget)
 	{
 	}
diff --git a/Engine/Plugins/Experimental/ChaosVD/Source/ChaosVD/Public/ChaosVDObjectDetailsTab.h b/Engine/Plugins/Experimental/ChaosVD/Source/ChaosVD/Public/ChaosVDObjectDetailsTab.h
--- a/Engine/Plugins/Experimental/ChaosVD/Source/ChaosVD/Public/ChaosVDObjectDetailsTab.h
+++ b/Engine/Plugins/Experimental/ChaosVD/Source/ChaosVD/Public/ChaosVDObjectDetailsTab.h
@@ -21,7 +21,7 @@ class FChaosVDObjectDetailsTab : public FChaosVDTabSpawnerBase, public FChaosVDS
 {
 public:
 
-	FChaosVDObjectDetailsTab(const FName& InTabID, TSharedPtr<FTabManager> InTabManager, SChaosVDMainTab* InOwningTabWidget) : FChaosVDTabSpawnerBase(InTabID, InTabManager, InOwningTabWidget)
+	FChaosVDObjectDetailsTab(const FName& InTabID, TSharedPtr<FTabManager> InTabManager, TWeakPtr<SChaosVDMainTab> InOwningTabWidget) : FChaosVDTabSpawnerBase(InTabID, InTabManager, InOwningTabWidget)
 	{
 	}
 
@@ -31,6 +31,14 @@ protected:
 
 	virtual void HandlePostSelectionChange(const UTypedElementSelectionSet* ChangedSelectionSet) override;
 
+	EVisibility GetCollisionDataButtonVisibility() const;
+	bool GetCollisionDataButtonEnabled() const;
+
+	TSharedPtr<SWidget> GenerateShowCollisionDataButton();
+
+	FReply ShowCollisionDataForSelectedObject();
+
 	FDelegateHandle SelectionDelegateHandle;
 	TSharedPtr<SChaosVDDetailsView> DetailsPanelView;
+	TWeakObjectPtr<UObject> CurrentSelectedObject = nullptr;
 };
diff --git a/Engine/Plugins/Experimental/ChaosVD/Source/ChaosVD/Public/ChaosVDOutputLogTab.h b/Engine/Plugins/Experimental/ChaosVD/Source/ChaosVD/Public/ChaosVDOutputLogTab.h
--- a/Engine/Plugins/Experimental/ChaosVD/Source/ChaosVD/Public/ChaosVDOutputLogTab.h
+++ b/Engine/Plugins/Experimental/ChaosVD/Source/ChaosVD/Public/ChaosVDOutputLogTab.h
@@ -14,7 +14,7 @@ class FChaosVDOutputLogTab : public FChaosVDTabSpawnerBase, public TSharedFromTh
 {
 public:
 
-	FChaosVDOutputLogTab(const FName& InTabID, TSharedPtr<FTabManager> InTabManager, SChaosVDMainTab* InOwningTabWidget) : FChaosVDTabSpawnerBase(InTabID, InTabManager, InOwningTabWidget)
+	FChaosVDOutputLogTab(const FName& InTabID, TSharedPtr<FTabManager> InTabManager, TWeakPtr<SChaosVDMainTab> InOwningTabWidget) : FChaosVDTabSpawnerBase(InTabID, InTabManager, InOwningTabWidget)
 	{
 	}
 
diff --git a/Engine/Plugins/Experimental/ChaosVD/Source/ChaosVD/Public/ChaosVDParticleActor.h b/Engine/Plugins/Experimental/ChaosVD/Source/ChaosVD/Public/ChaosVDParticleActor.h
--- a/Engine/Plugins/Experimental/ChaosVD/Source/ChaosVD/Public/ChaosVDParticleActor.h
+++ b/Engine/Plugins/Experimental/ChaosVD/Source/ChaosVD/Public/ChaosVDParticleActor.h
@@ -1,17 +1,19 @@
 // Copyright Epic Games, Inc. All Rights Reserved.
 
 #pragma once
+#include "ChaosVDCollisionDataProviderInterface.h"
 #include "ChaosVDGeometryDataComponent.h"
+#include "ChaosVDSceneObjectBase.h"
 #include "Chaos/Core.h"
 #include "DataWrappers/ChaosVDParticleDataWrapper.h"
 #include "DataWrappers/ChaosVDCollisionDataWrappers.h"
 #include "GameFramework/Actor.h"
-#include "Visualizers/ChaosVDCollisionDataVisualizer.h"
-#include "Visualizers/ChaosVDParticleDataVisualizer.h"
 #include "Visualizers/ChaosVDDataVisualizerBase.h"
+#include "Visualizers/IChaosVDParticleVisualizationDataProvider.h"
 
 #include "ChaosVDParticleActor.generated.h"
 
+enum class EChaosVDParticleDataVisualizationFlags : uint32;
 class FChaosVDScene;
 struct FChaosVDParticleDebugData;
 class UMeshComponent;
@@ -33,22 +35,23 @@ enum class EChaosVDActorGeometryUpdateFlags : int32
 ENUM_CLASS_FLAGS(EChaosVDActorGeometryUpdateFlags)
 
 /** Actor used to represent a Chaos Particle in the Visual Debugger's world */
-UCLASS()
-class AChaosVDParticleActor : public AActor, public IChaosVDParticleVisualizationDataProvider, public IChaosVDVisualizerContainerInterface
+UCLASS(HideCategories=(Transform))
+class AChaosVDParticleActor : public AActor, public IChaosVDParticleVisualizationDataProvider, public IChaosVDVisualizerContainerInterface,
+								public FChaosVDSceneObjectBase, public IChaosVDCollisionDataProviderInterface
 {
+
 	GENERATED_BODY()
+
 public:
 	AChaosVDParticleActor(const FObjectInitializer& ObjectInitializer);
 
 	void UpdateFromRecordedParticleData(const FChaosVDParticleDataWrapper& InRecordedData, const Chaos::FRigidTransform3& SimulationTransform);
 
-	void UpdateCollisionData(const TArray<TSharedPtr<FChaosVDParticlePairMidPhase>>& InRecordedMidPhases);
-	void UpdateCollisionData(const TArray<FChaosVDConstraint>& InRecordedConstraints);
 	void UpdateGeometry(const Chaos::FConstImplicitObjectPtr& InImplicitObject, EChaosVDActorGeometryUpdateFlags OptionsFlags = EChaosVDActorGeometryUpdateFlags::None);
 
 	void UpdateGeometry(uint32 NewGeometryHash, EChaosVDActorGeometryUpdateFlags OptionsFlags = EChaosVDActorGeometryUpdateFlags::None);
 
-	void SetScene(const TSharedPtr<FChaosVDScene>& InScene);
+	virtual void SetScene(TWeakPtr<FChaosVDScene> InScene) override;
 
 	virtual void BeginDestroy() override;
 
@@ -74,15 +77,20 @@ public:
 	 */
 	bool IsActive() const { return bIsActive; }
 
+	//BEGIN IChaosVDCollisionDataProvider Interface
+	virtual void GetCollisionData(TArray<TSharedPtr<FChaosVDCollisionDataFinder>>& OutCollisionDataFound) override;
+	virtual bool HasCollisionData() override;
+	virtual FName GetName() override;
+	//END IChaosVDCollisionDataProvider Interface
+
 protected:
 
+	const TArray<TSharedPtr<FChaosVDParticlePairMidPhase>>* GetCollisionMidPhasesArray() const;
+
 	void UpdateShapeDataComponents();
 	
 	UPROPERTY(EditAnywhere, Category = "Viewport Visualization Flags", meta = (Bitmask, BitmaskEnum = "/Script/ChaosVD.EChaosVDParticleDataVisualizationFlags"))
 	uint8 LocalParticleDataVisualizationFlags;
-	
-	UPROPERTY(EditAnywhere, Category = "Viewport Visualization Flags", meta = (Bitmask, BitmaskEnum = "/Script/ChaosVD.EChaosVDCollisionVisualizationFlags"))
-	uint8 LocalCollisionDataVisualizationFlags;
 
 	UPROPERTY(EditAnywhere, Category = "Viewport Visualization Flags")
 	bool bShowDebugText = false;
@@ -94,8 +102,6 @@ protected:
 
 	bool bIsGeometryDataGenerationStarted = false;
 
-	TWeakPtr<FChaosVDScene> OwningScene;
-
 	TArray<TWeakObjectPtr<UMeshComponent>> MeshComponents;
 
 	FDelegateHandle GeometryUpdatedDelegate;
diff --git a/Engine/Plugins/Experimental/ChaosVD/Source/ChaosVD/Public/ChaosVDPlaybackViewportTab.h b/Engine/Plugins/Experimental/ChaosVD/Source/ChaosVD/Public/ChaosVDPlaybackViewportTab.h
--- a/Engine/Plugins/Experimental/ChaosVD/Source/ChaosVD/Public/ChaosVDPlaybackViewportTab.h
+++ b/Engine/Plugins/Experimental/ChaosVD/Source/ChaosVD/Public/ChaosVDPlaybackViewportTab.h
@@ -14,7 +14,7 @@ class SChaosVDPlaybackViewport;
 class FChaosVDPlaybackViewportTab : public FChaosVDTabSpawnerBase, public TSharedFromThis<FChaosVDPlaybackViewportTab>
 {
 public:
-	FChaosVDPlaybackViewportTab(const FName& InTabID, TSharedPtr<FTabManager> InTabManager, SChaosVDMainTab* InOwningTabWidget) : FChaosVDTabSpawnerBase(InTabID, InTabManager, InOwningTabWidget)
+	FChaosVDPlaybackViewportTab(const FName& InTabID, TSharedPtr<FTabManager> InTabManager, TWeakPtr<SChaosVDMainTab> InOwningTabWidget) : FChaosVDTabSpawnerBase(InTabID, InTabManager, InOwningTabWidget)
 	{	
 	}
 	
diff --git a/Engine/Plugins/Experimental/ChaosVD/Source/ChaosVD/Public/ChaosVDScene.h b/Engine/Plugins/Experimental/ChaosVD/Source/ChaosVD/Public/ChaosVDScene.h
--- a/Engine/Plugins/Experimental/ChaosVD/Source/ChaosVD/Public/ChaosVDScene.h
+++ b/Engine/Plugins/Experimental/ChaosVD/Source/ChaosVD/Public/ChaosVDScene.h
@@ -9,6 +9,8 @@
 #include "UObject/ObjectMacros.h"
 #include "UObject/ObjectPtr.h"
 
+class AChaosVDSolverInfoActor;
+class AChaosVDSceneCollisionContainer;
 class UChaosVDEditorSettings;
 class FChaosVDGeometryBuilder;
 class AChaosVDParticleActor;
@@ -21,7 +23,7 @@ class UWorld;
 
 struct FTypedElementHandle;
 
-typedef TMap<int32, AChaosVDParticleActor*> FChaosVDParticlesByIDMap;
+typedef TMap<int32, AChaosVDSolverInfoActor*> FChaosVDSolverInfoByIDMap;
 
 DECLARE_MULTICAST_DELEGATE(FChaosVDSceneUpdatedDelegate)
 DECLARE_MULTICAST_DELEGATE_OneParam(FChaosVDActorActiveStateUpdateDelegate, AChaosVDParticleActor*)
@@ -93,6 +95,9 @@ public:
 	 */
 	AChaosVDParticleActor* GetParticleActor(int32 SolverID, int32 ParticleID);
 
+	const FChaosVDSolverInfoByIDMap& GetSolverInfoActorsMap() { return SolverDataContainerBySolverID; }
+	AChaosVDSolverInfoActor* GetSolverInfoActor(int32 SolverID);
+
 	AActor* GetSkySphereActor() const { return SkySphere; }
 
 	FChaosVDActorActiveStateUpdateDelegate& OnActorActiveStateChanged() { return ParticleActorUpdateDelegate; }
@@ -109,10 +114,10 @@ private:
 	void CreateBaseLights(UWorld* TargetWorld) const;
 
 	/** Creates the instance of the World which will be used the recorded data*/
-	UWorld* CreatePhysicsVDWorld() const;
+	UWorld* CreatePhysicsVDWorld();
 
-	/** Map of ID-ChaosVDParticle Actor. Used to keep track of actor instances and be able to modify them as needed*/
-	TMap<int32, FChaosVDParticlesByIDMap> ParticlesBySolverID;
+	/** Map of SolverID-ChaosVDSolverInfo Actor. Used to keep track of active solvers representations and be able to modify them as needed*/
+	FChaosVDSolverInfoByIDMap SolverDataContainerBySolverID;
 
 	/** Returns the correct TypedElementHandle based on an object type so it can be used with the selection set object */
 	FTypedElementHandle GetSelectionHandleForObject(const UObject* Object) const;
@@ -130,6 +135,8 @@ private:
 	void InitializeSelectionSets();
 	void DeInitializeSelectionSets();
 
+	void HandleActorDestroyed(AActor* ActorDestroyed);
+
 	/** UWorld instance used to represent the recorded debug data */
 	TObjectPtr<UWorld> PhysicsVDWorld = nullptr;
 
@@ -159,4 +166,6 @@ private:
 	bool bIsInitialized = false;
 
 	FChaosVDActorActiveStateUpdateDelegate ParticleActorUpdateDelegate;
+
+	FDelegateHandle ActorDestroyedHandle;
 };
diff --git a/Engine/Plugins/Experimental/ChaosVD/Source/ChaosVD/Public/ChaosVDSceneObjectBase.h b/Engine/Plugins/Experimental/ChaosVD/Source/ChaosVD/Public/ChaosVDSceneObjectBase.h
new file mode 100644
--- /dev/null
+++ b/Engine/Plugins/Experimental/ChaosVD/Source/ChaosVD/Public/ChaosVDSceneObjectBase.h
@@ -0,0 +1,20 @@
+// Copyright Epic Games, Inc. All Rights Reserved.
+
+#pragma once
+#include "Templates/SharedPointer.h"
+
+class FChaosVDScene;
+
+/** Base Class for any object that is owned by a Chaos VD Scene */
+class FChaosVDSceneObjectBase
+{
+public:
+	virtual ~FChaosVDSceneObjectBase() = default;
+
+	virtual void SetScene(TWeakPtr<FChaosVDScene> InScene) { SceneWeakPtr = InScene; }
+	
+	TWeakPtr<FChaosVDScene> GetScene() const { return SceneWeakPtr; }
+	
+protected:
+	TWeakPtr<FChaosVDScene> SceneWeakPtr;
+};
diff --git a/Engine/Plugins/Experimental/ChaosVD/Source/ChaosVD/Public/ChaosVDSolversTracksTab.h b/Engine/Plugins/Experimental/ChaosVD/Source/ChaosVD/Public/ChaosVDSolversTracksTab.h
--- a/Engine/Plugins/Experimental/ChaosVD/Source/ChaosVD/Public/ChaosVDSolversTracksTab.h
+++ b/Engine/Plugins/Experimental/ChaosVD/Source/ChaosVD/Public/ChaosVDSolversTracksTab.h
@@ -11,7 +11,7 @@ class SChaosVDSolverTracks;
 class FChaosVDSolversTracksTab : public FChaosVDTabSpawnerBase, public TSharedFromThis<FChaosVDSolversTracksTab>
 {
 public:
-	FChaosVDSolversTracksTab(const FName& InTabID, TSharedPtr<FTabManager> InTabManager, SChaosVDMainTab* InOwningTabWidget) : FChaosVDTabSpawnerBase(InTabID, InTabManager, InOwningTabWidget)
+	FChaosVDSolversTracksTab(const FName& InTabID, TSharedPtr<FTabManager> InTabManager, TWeakPtr<SChaosVDMainTab> InOwningTabWidget) : FChaosVDTabSpawnerBase(InTabID, InTabManager, InOwningTabWidget)
 	{	
 	}
 	
diff --git a/Engine/Plugins/Experimental/ChaosVD/Source/ChaosVD/Public/ChaosVDTabSpawnerBase.h b/Engine/Plugins/Experimental/ChaosVD/Source/ChaosVD/Public/ChaosVDTabSpawnerBase.h
--- a/Engine/Plugins/Experimental/ChaosVD/Source/ChaosVD/Public/ChaosVDTabSpawnerBase.h
+++ b/Engine/Plugins/Experimental/ChaosVD/Source/ChaosVD/Public/ChaosVDTabSpawnerBase.h
@@ -2,8 +2,10 @@
 
 #pragma once
 
+#include "Delegates/DelegateCombinations.h"
 #include "Templates/SharedPointer.h"
 #include "UObject/NameTypes.h"
+#include "Widgets/Docking/SDockTab.h"
 
 class FChaosVDScene;
 class SChaosVDMainTab;
@@ -12,20 +14,32 @@ class FTabManager;
 class SDockTab;
 class FSpawnTabArgs;
 
+DECLARE_MULTICAST_DELEGATE_OneParam(FChaosVDTabSpawned, TSharedRef<SDockTab>)
+DECLARE_MULTICAST_DELEGATE_OneParam(FChaosVDTabDestroyed, TSharedRef<SDockTab>)
+
 /** Base class for any tab of the Chaos Visual Debugger tool*/
 class FChaosVDTabSpawnerBase
 {
 public:
 	virtual ~FChaosVDTabSpawnerBase() = default;
 
-	FChaosVDTabSpawnerBase(const FName& InTabID, TSharedPtr<FTabManager> InTabManager, SChaosVDMainTab* InOwningTabWidget);
+	FChaosVDTabSpawnerBase(const FName& InTabID, TSharedPtr<FTabManager> InTabManager, TWeakPtr<SChaosVDMainTab> InOwningTabWidget);
 
 	virtual TSharedRef<SDockTab> HandleTabSpawned(const FSpawnTabArgs& Args) = 0;
+	virtual void HandleTabClosed(const TSharedRef<SDockTab>& InTabClosed);
+
+	FChaosVDTabSpawned& OnTabSpawned() { return TabSpawnedDelegate; }
+	FChaosVDTabDestroyed& OnTabDestroyed() { return TabDestroyedDelegate; }
+
+	TSharedRef<SWidget> GenerateErrorWidget();
 
 protected:
 
 	/** Ptr to the main tab of the owning visual debugger tool instance */
-	SChaosVDMainTab* OwningTabWidget;
+	TWeakPtr<SChaosVDMainTab> OwningTabWidget;
+
+	FChaosVDTabSpawned TabSpawnedDelegate;
+	FChaosVDTabDestroyed TabDestroyedDelegate;
 
 	UWorld* GetChaosVDWorld() const;
 
diff --git a/Engine/Plugins/Experimental/ChaosVD/Source/ChaosVD/Public/ChaosVDTabsIDs.h b/Engine/Plugins/Experimental/ChaosVD/Source/ChaosVD/Public/ChaosVDTabsIDs.h
--- a/Engine/Plugins/Experimental/ChaosVD/Source/ChaosVD/Public/ChaosVDTabsIDs.h
+++ b/Engine/Plugins/Experimental/ChaosVD/Source/ChaosVD/Public/ChaosVDTabsIDs.h
@@ -16,4 +16,5 @@ public:
 	static const FName SolversTrack;
 	static const FName CVDEditorSettings;
 	static const FName StatusBar;
+	static const FName CollisionDataDetails;
 };
diff --git a/Engine/Plugins/Experimental/ChaosVD/Source/ChaosVD/Public/ChaosVDWorldOutlinerTab.h b/Engine/Plugins/Experimental/ChaosVD/Source/ChaosVD/Public/ChaosVDWorldOutlinerTab.h
--- a/Engine/Plugins/Experimental/ChaosVD/Source/ChaosVD/Public/ChaosVDWorldOutlinerTab.h
+++ b/Engine/Plugins/Experimental/ChaosVD/Source/ChaosVD/Public/ChaosVDWorldOutlinerTab.h
@@ -26,7 +26,7 @@ DECLARE_DELEGATE_OneParam(FChaosVDSelectionChanged, AActor* SelectedActor)
 class FChaosVDWorldOutlinerTab : public FChaosVDTabSpawnerBase, public TSharedFromThis<FChaosVDWorldOutlinerTab>
 {
 public:
-	FChaosVDWorldOutlinerTab(const FName& InTabID, TSharedPtr<FTabManager> InTabManager, SChaosVDMainTab* InOwningTabWidget) : FChaosVDTabSpawnerBase(InTabID, InTabManager, InOwningTabWidget)
+	FChaosVDWorldOutlinerTab(const FName& InTabID, TSharedPtr<FTabManager> InTabManager, TWeakPtr<SChaosVDMainTab> InOwningTabWidget) : FChaosVDTabSpawnerBase(InTabID, InTabManager, InOwningTabWidget)
 	{
 	}
 
diff --git a/Engine/Plugins/Experimental/ChaosVD/Source/ChaosVD/Public/Components/ChaosVDInstancedStaticMeshComponent.h b/Engine/Plugins/Experimental/ChaosVD/Source/ChaosVD/Public/Components/ChaosVDInstancedStaticMeshComponent.h
--- a/Engine/Plugins/Experimental/ChaosVD/Source/ChaosVD/Public/Components/ChaosVDInstancedStaticMeshComponent.h
+++ b/Engine/Plugins/Experimental/ChaosVD/Source/ChaosVD/Public/Components/ChaosVDInstancedStaticMeshComponent.h
@@ -15,6 +15,8 @@ class UChaosVDInstancedStaticMeshComponent : public UInstancedStaticMeshComponen
 	UChaosVDInstancedStaticMeshComponent(const FObjectInitializer& ObjectInitializer)
 	: Super(ObjectInitializer)
 	{
+		SetCanEverAffectNavigation(false);
+		bNavigationRelevant = false;
 	}
 	
 	// BEGIN IChaosVDGeometryDataComponent Interface
diff --git a/Engine/Plugins/Experimental/ChaosVD/Source/ChaosVD/Public/Components/ChaosVDSolverCollisionDataComponent.h b/Engine/Plugins/Experimental/ChaosVD/Source/ChaosVD/Public/Components/ChaosVDSolverCollisionDataComponent.h
new file mode 100644
--- /dev/null
+++ b/Engine/Plugins/Experimental/ChaosVD/Source/ChaosVD/Public/Components/ChaosVDSolverCollisionDataComponent.h
@@ -0,0 +1,103 @@
+// Copyright Epic Games, Inc. All Rights Reserved.
+
+#pragma once
+
+#include "Components/ActorComponent.h"
+
+#include "Containers/Array.h"
+#include "Containers/Map.h"
+#include "DataWrappers/ChaosVDCollisionDataWrappers.h"
+#include "GameFramework/Info.h"
+#include "Templates/SharedPointer.h"
+
+#include "ChaosVDSolverCollisionDataComponent.generated.h"
+
+class FChaosVDCollisionDataVisualizer;
+class FChaosVDScene;
+class FPrimitiveDrawInterface;
+class FSceneView;
+
+struct FChaosVDConstraint;
+
+typedef TMap<int32, TArray<TSharedPtr<FChaosVDParticlePairMidPhase>>> FChaosVDMidPhaseByParticleMap;
+typedef TMap<int32, TArray<FChaosVDConstraint*>> FChaosVDConstraintByParticleMap;
+
+enum class EChaosVDGetCollisionDataOptions
+{
+	Owner,
+	Secondary,
+	Any
+};
+
+UCLASS()
+class UChaosVDSolverCollisionDataComponent : public UActorComponent
+{
+	GENERATED_BODY()
+
+public:
+	UChaosVDSolverCollisionDataComponent();
+
+	void UpdateCollisionData(const TArray<TSharedPtr<FChaosVDParticlePairMidPhase>>& InMidPhaseData);
+
+	const TArray<TSharedPtr<FChaosVDParticlePairMidPhase>>& GetMidPhases() const { return AllMidPhases; }
+	const TArray<TSharedPtr<FChaosVDParticlePairMidPhase>>* GetMidPhasesForParticle(int32 ParticleID, EChaosVDGetCollisionDataOptions Options) const;
+	const TArray<FChaosVDConstraint*>* GetConstraintsForParticle(int32 ParticleID, EChaosVDGetCollisionDataOptions Options) const;
+
+	void DrawVisualization(const FSceneView* View, FPrimitiveDrawInterface* PDI);
+
+protected:
+
+	void ClearCollisionData();
+
+	template<typename MapType, typename CollisionDataType>
+	void AddCollisionDataToParticleIDMap(MapType& MapToUpdate, const CollisionDataType& MidPhaseData, int32 ParticleID);
+
+	template<typename MapType, typename CollisionDataType>
+	const TArray<CollisionDataType>* GetCollisionDataFromMap(const MapType& Map0ToQuery, const MapType& Map1ToQuery, int32 ParticleID, EChaosVDGetCollisionDataOptions Options) const;
+
+	TArray<TSharedPtr<FChaosVDParticlePairMidPhase>> AllMidPhases;
+	FChaosVDMidPhaseByParticleMap MidPhasesByParticleID0;
+	FChaosVDMidPhaseByParticleMap MidPhasesByParticleID1;
+
+	FChaosVDConstraintByParticleMap ConstraintsByParticleID0;
+	FChaosVDConstraintByParticleMap ConstraintsByParticleID1;
+
+};
+
+template <typename MapType, typename CollisionDataType>
+void UChaosVDSolverCollisionDataComponent::AddCollisionDataToParticleIDMap(MapType& MapToUpdate, const CollisionDataType& MidPhaseData, int32 ParticleID)
+{
+	if (TArray<CollisionDataType>* ParticleCollisionData = MapToUpdate.Find(ParticleID))
+	{
+		ParticleCollisionData->Add(MidPhaseData);
+	}
+	else
+	{
+		MapToUpdate.Add(ParticleID, { MidPhaseData });
+	}
+}
+
+template <typename MapType, typename CollisionDataType>
+const TArray<CollisionDataType>* UChaosVDSolverCollisionDataComponent::GetCollisionDataFromMap(const MapType& Map0ToQuery, const MapType& Map1ToQuery, int32 ParticleID, EChaosVDGetCollisionDataOptions Options) const
+{
+	switch (Options)
+	{
+	case EChaosVDGetCollisionDataOptions::Owner:
+		{
+			return Map0ToQuery.Find(ParticleID);
+			break;
+		}
+	case EChaosVDGetCollisionDataOptions::Secondary:
+		{
+			return Map1ToQuery.Find(ParticleID);
+			break;
+		}
+	case EChaosVDGetCollisionDataOptions::Any:
+		{
+			const TArray<CollisionDataType>* FoundMidPhasesContainer = Map0ToQuery.Find(ParticleID);
+			return FoundMidPhasesContainer ? FoundMidPhasesContainer : Map1ToQuery.Find(ParticleID);
+			break;
+		}
+	}
+	return nullptr;
+}
diff --git a/Engine/Plugins/Experimental/ChaosVD/Source/ChaosVD/Public/Components/ChaosVDStaticMeshComponent.h b/Engine/Plugins/Experimental/ChaosVD/Source/ChaosVD/Public/Components/ChaosVDStaticMeshComponent.h
--- a/Engine/Plugins/Experimental/ChaosVD/Source/ChaosVD/Public/Components/ChaosVDStaticMeshComponent.h
+++ b/Engine/Plugins/Experimental/ChaosVD/Source/ChaosVD/Public/Components/ChaosVDStaticMeshComponent.h
@@ -14,7 +14,9 @@ class UChaosVDStaticMeshComponent : public UStaticMeshComponent, public IChaosVD
 public:
 	UChaosVDStaticMeshComponent(const FObjectInitializer& ObjectInitializer)
 		: Super(ObjectInitializer)
-	{		
+	{
+		SetCanEverAffectNavigation(false);
+		bNavigationRelevant = false;
 	}
 
 	// BEGIN IChaosVDGeometryDataComponent Interface
diff --git a/Engine/Plugins/Experimental/ChaosVD/Source/ChaosVD/Public/Visualizers/ChaosVDCollisionDataVisualizer.h b/Engine/Plugins/Experimental/ChaosVD/Source/ChaosVD/Public/Visualizers/ChaosVDCollisionDataVisualizer.h
deleted file mode 100644
diff --git a/Engine/Plugins/Experimental/ChaosVD/Source/ChaosVD/Public/Visualizers/ChaosVDDataVisualizerBase.h b/Engine/Plugins/Experimental/ChaosVD/Source/ChaosVD/Public/Visualizers/ChaosVDDataVisualizerBase.h
--- a/Engine/Plugins/Experimental/ChaosVD/Source/ChaosVD/Public/Visualizers/ChaosVDDataVisualizerBase.h
+++ b/Engine/Plugins/Experimental/ChaosVD/Source/ChaosVD/Public/Visualizers/ChaosVDDataVisualizerBase.h
@@ -19,7 +19,7 @@ class IChaosVDVisualizerContainerInterface
 {
 	GENERATED_BODY()
 public:
-	virtual void DrawVisualization(const FSceneView* View, FPrimitiveDrawInterface* PDI){};
+	virtual void DrawVisualization(const FSceneView* View, FPrimitiveDrawInterface* PDI){}
 };
 
 /**
diff --git a/Engine/Plugins/Experimental/ChaosVD/Source/ChaosVD/Public/Visualizers/IChaosVDParticleVisualizationDataProvider.h b/Engine/Plugins/Experimental/ChaosVD/Source/ChaosVD/Public/Visualizers/IChaosVDParticleVisualizationDataProvider.h
--- a/Engine/Plugins/Experimental/ChaosVD/Source/ChaosVD/Public/Visualizers/IChaosVDParticleVisualizationDataProvider.h
+++ b/Engine/Plugins/Experimental/ChaosVD/Source/ChaosVD/Public/Visualizers/IChaosVDParticleVisualizationDataProvider.h
@@ -14,6 +14,7 @@ struct FChaosVDVisualizationContext
 	FTransform SpaceTransform;
 	TWeakPtr<FChaosVDScene> CVDScene;
 	int32 SolverID = INDEX_NONE;
+	uint32 VisualizationFlags = 0;
 };
 
 /** Interface to be used by any object that contains Particle Data that needs to be visualized.
diff --git a/Engine/Plugins/Experimental/ChaosVD/Source/ChaosVD/Public/Widgets/SChaosVDCollisionDataInspector.h b/Engine/Plugins/Experimental/ChaosVD/Source/ChaosVD/Public/Widgets/SChaosVDCollisionDataInspector.h
new file mode 100644
--- /dev/null
+++ b/Engine/Plugins/Experimental/ChaosVD/Source/ChaosVD/Public/Widgets/SChaosVDCollisionDataInspector.h
@@ -0,0 +1,80 @@
+// Copyright Epic Games, Inc. All Rights Reserved.
+
+#pragma once
+
+#include "Widgets/SCompoundWidget.h"
+
+class AChaosVDParticleActor;
+class FChaosVDScene;
+class IStructureDetailsView;
+class IChaosVDCollisionDataProviderInterface;
+class SChaosVDNameListPicker;
+
+struct FChaosVDCollisionDataFinder;
+struct FChaosVDParticlePairMidPhase;
+
+typedef TSharedPtr<FChaosVDParticlePairMidPhase> FChaosVDMidPhasePtr;
+
+enum class EChaosVDCollisionParticleSelector
+{
+	Index_0,
+	Index_1
+};
+
+class SChaosVDCollisionDataInspector : public SCompoundWidget
+{
+public:
+
+	SLATE_BEGIN_ARGS(SChaosVDCollisionDataInspector)
+	{
+	}
+
+	SLATE_END_ARGS()
+
+	virtual ~SChaosVDCollisionDataInspector() override;
+
+	/** Constructs this widget with InArgs */
+	void Construct(const FArguments& InArgs, const TWeakPtr<FChaosVDScene>& InScenePtr);
+
+	void SetCollisionDataProviderObjectToInspect(IChaosVDCollisionDataProviderInterface* CollisionDataProvider);
+	void SetSingleContactDataToInspect(const FChaosVDCollisionDataFinder& InContactFinderData);
+
+protected:
+
+	void HandleSceneUpdated();
+
+	void ClearInspector();
+
+	FText GetObjectBeingInspectedName() const;
+	TSharedPtr<SWidget> GenerateObjectNameRowWidget();
+
+	TSharedPtr<FChaosVDCollisionDataFinder> GetCurrentDataBeingInspected();
+
+	EVisibility GetOutOfDateWarningVisibility() const;
+
+	void HandleSessionNameSelected(TSharedPtr<FName> SelectedName);
+
+	TSharedPtr<FName> GenerateNameForCollisionDataItem(const FChaosVDCollisionDataFinder& InContactFinderData);
+
+	TSharedPtr<IStructureDetailsView> CreateCollisionDataDetailsView();
+
+	EVisibility GetDetailsSectionVisibility() const;
+
+	FReply SelectParticleForCurrentCollisionData(EChaosVDCollisionParticleSelector ParticleSelector);
+
+	TWeakPtr<FChaosVDScene> SceneWeakPtr;
+
+	TSharedPtr<SChaosVDNameListPicker> CollisionDataAvailableList;
+	
+	TMap<FName, TSharedPtr<FChaosVDCollisionDataFinder>> CollisionDataByNameMap;
+
+	TSharedPtr<FName> CurrentSelectedName;
+
+	TSharedPtr<IStructureDetailsView> ContactDataDetailsView;
+	
+	TSharedPtr<IStructureDetailsView> ConstraintDataDetailsView;
+
+	FName CurrentObjectBeingInspectedName;
+
+	bool bIsUpToDate = true;
+};
diff --git a/Engine/Plugins/Experimental/ChaosVD/Source/ChaosVD/Public/Widgets/SChaosVDMainTab.h b/Engine/Plugins/Experimental/ChaosVD/Source/ChaosVD/Public/Widgets/SChaosVDMainTab.h
--- a/Engine/Plugins/Experimental/ChaosVD/Source/ChaosVD/Public/Widgets/SChaosVDMainTab.h
+++ b/Engine/Plugins/Experimental/ChaosVD/Source/ChaosVD/Public/Widgets/SChaosVDMainTab.h
@@ -4,8 +4,11 @@
 
 #include "Framework/Docking/TabManager.h"
 #include "Templates/SharedPointer.h"
+#include "Toolkits/IToolkitHost.h"
 #include "Widgets/SCompoundWidget.h"
 
+class FComponentVisualizer;
+class FChaosVDEditorModeTools;
 class FChaosVDTabSpawnerBase;
 class FChaosVDEditorVisualizationSettingsTab;
 class FChaosVDSolversTracksTab;
@@ -17,23 +20,48 @@ class FChaosVDWorldOutlinerTab;
 class SDockTab;
 
 /** The main widget containing the Chaos Visual Debugger interface */
-class SChaosVDMainTab : public SCompoundWidget
+class SChaosVDMainTab : public SCompoundWidget, public IToolkitHost
 {
 public:
-
+	
 	SLATE_BEGIN_ARGS(SChaosVDMainTab) {}
 		SLATE_ARGUMENT(TSharedPtr<SDockTab>, OwnerTab)
 	SLATE_END_ARGS()
 
 	void Construct(const FArguments& InArgs, TSharedPtr<FChaosVDEngine> InChaosVDEngine);
 
-	TSharedRef<FChaosVDEngine> GetChaosVDEngineInstance() const { return ChaosVDEngine.ToSharedRef(); };
-	const FName& GetStatusBarID() const { return StatusBarID; };
+	TSharedRef<FChaosVDEngine> GetChaosVDEngineInstance() const { return ChaosVDEngine.ToSharedRef(); }
+
+	template<typename TabType>
+	TWeakPtr<TabType> GetTabSpawnerInstance(FName TabID);
+
+	// BEGIN ITOOLKITHOST Interface
+	virtual TSharedRef<SWidget> GetParentWidget() override { return AsShared(); }
+	virtual void BringToFront() override;
+	virtual TSharedPtr<FTabManager> GetTabManager() const override { return TabManager; };
+	virtual void OnToolkitHostingStarted(const TSharedRef<IToolkit>& Toolkit) override;
+	virtual void OnToolkitHostingFinished(const TSharedRef<IToolkit>& Toolkit) override;
+	virtual UWorld* GetWorld() const override;
+	virtual FEditorModeTools& GetEditorModeManager() const override;
 
-	TWeakPtr<FChaosVDPlaybackViewportTab> GetPlaybackViewportTab() { return PlaybackViewportTab;}
+	virtual UTypedElementCommonActions* GetCommonActions() const override { return nullptr; }
+	virtual FName GetStatusBarName() const override { return StatusBarID; };
+	virtual FOnActiveViewportChanged& OnActiveViewportChanged() override { return ViewportChangedDelegate;};
+	// END ITOOLKITHOST Interface
+
+	TSharedPtr<FComponentVisualizer> FindComponentVisualizer(UClass* ClassPtr);
+	TSharedPtr<FComponentVisualizer> FindComponentVisualizer(FName ClassName);
 
 private:
 
+	template<typename TabType>
+	void RegisterTabSpawner(FName TabID);
+
+	void RegisterComponentVisualizer(FName ClassName, const TSharedPtr<FComponentVisualizer>& Visualizer);
+
+	void HandleTabSpawned(TSharedRef<SDockTab> Tab, FName TabID);
+	void HandleTabDestroyed(TSharedRef<SDockTab> Tab, FName TabID);
+
 	TSharedRef<FTabManager::FLayout> GenerateMainLayout();
 
 	void GenerateMainWindowMenu();
@@ -46,18 +74,46 @@ private:
 
 	TSharedPtr<FChaosVDEngine> ChaosVDEngine;
 
-	// TODO Convert this to a map ID-> Tab
-	TSharedPtr<FChaosVDWorldOutlinerTab> WorldOutlinerTab;
-	TSharedPtr<FChaosVDObjectDetailsTab> ObjectDetailsTab;
-	TSharedPtr<FChaosVDPlaybackViewportTab> PlaybackViewportTab;
-	TSharedPtr<FChaosVDSolversTracksTab> SolversTracksTab;
-	TSharedPtr<FChaosVDOutputLogTab> OutputLogTab;
-	TSharedPtr<FChaosVDEditorVisualizationSettingsTab> EditorSettingsTab;
-
 	FName StatusBarID;
 
 	TSharedPtr<FTabManager> TabManager;
+	TWeakPtr<SDockTab> OwnerTab;
+	TSharedPtr<FChaosVDEditorModeTools> EditorModeTools;
+	
+	TMap<FName, TSharedPtr<FChaosVDTabSpawnerBase>> TabSpawnersByIDMap;
+
+	TMap<FName, TSharedPtr<FComponentVisualizer>> ComponentVisualizersMap;
+
+	TMap<FName, TWeakPtr<SDockTab>> ActiveTabsByID;
+
+	FOnActiveViewportChanged ViewportChangedDelegate;
 
 	FReply HandleSessionConnectionClicked();
 	FText GetConnectButtonText() const;
 };
+
+
+template <typename TabType>
+TWeakPtr<TabType> SChaosVDMainTab::GetTabSpawnerInstance(FName TabID)
+{
+	if (TSharedPtr<FChaosVDTabSpawnerBase>* TabSpawnerPtrPtr = TabSpawnersByIDMap.Find(TabID))
+	{
+		return StaticCastSharedPtr<TabType>(*TabSpawnerPtrPtr);
+	}
+
+	return nullptr;
+}
+
+template <typename TabType>
+void SChaosVDMainTab::RegisterTabSpawner(FName TabID)
+{
+	static_assert(std::is_base_of_v<FChaosVDTabSpawnerBase, TabType> , "SChaosVDMainTab::RegisterTabSpawner Only supports FChaosVDTabSpawnerBase based spawners");
+
+	if (!TabSpawnersByIDMap.Contains(TabID))
+	{
+		TSharedPtr<TabType> TabSpawner = MakeShared<TabType>(TabID, TabManager, StaticCastWeakPtr<SChaosVDMainTab>(AsWeak()));
+		TabSpawner->OnTabSpawned().AddRaw(this, &SChaosVDMainTab::HandleTabSpawned, TabID);
+		TabSpawner->OnTabDestroyed().AddRaw(this, &SChaosVDMainTab::HandleTabDestroyed, TabID);
+		TabSpawnersByIDMap.Add(TabID, TabSpawner);
+	}
+}
diff --git a/Engine/Plugins/Experimental/ChaosVD/Source/ChaosVD/Public/Widgets/SChaosVDNameListPicker.h b/Engine/Plugins/Experimental/ChaosVD/Source/ChaosVD/Public/Widgets/SChaosVDNameListPicker.h
--- a/Engine/Plugins/Experimental/ChaosVD/Source/ChaosVD/Public/Widgets/SChaosVDNameListPicker.h
+++ b/Engine/Plugins/Experimental/ChaosVD/Source/ChaosVD/Public/Widgets/SChaosVDNameListPicker.h
@@ -28,8 +28,12 @@ public:
 	 */
 	void UpdateNameList(TArray<TSharedPtr<FName>>&& NewNameList);
 	
+	void SelectName(const TSharedPtr<FName>& NameToSelect, ESelectInfo::Type SelectionInfo = ESelectInfo::Direct);
+	
 protected:
 
+	bool HasElements() const;
+
 	FText GetCurrentDropDownButtonLabel() const;
 	
 	void OnClickItem(TSharedPtr<FName> Item);
diff --git a/Engine/Plugins/Experimental/ChaosVD/Source/ChaosVD/Public/Widgets/SChaosVDPlaybackViewport.h b/Engine/Plugins/Experimental/ChaosVD/Source/ChaosVD/Public/Widgets/SChaosVDPlaybackViewport.h
--- a/Engine/Plugins/Experimental/ChaosVD/Source/ChaosVD/Public/Widgets/SChaosVDPlaybackViewport.h
+++ b/Engine/Plugins/Experimental/ChaosVD/Source/ChaosVD/Public/Widgets/SChaosVDPlaybackViewport.h
@@ -32,7 +32,7 @@ public:
 
 	virtual ~SChaosVDPlaybackViewport() override;
 
-	void Construct(const FArguments& InArgs, TWeakPtr<FChaosVDScene> InScene, TWeakPtr<FChaosVDPlaybackController> InPlaybackController);
+	void Construct(const FArguments& InArgs, TWeakPtr<FChaosVDScene> InScene, TWeakPtr<FChaosVDPlaybackController> InPlaybackController, TSharedPtr<FEditorModeTools> InEditorModeTools);
 
 	// BEING ICommonEditorViewportToolbarInfoProvider interface
 	virtual TSharedRef<SEditorViewport> GetViewportWidget() override;
@@ -65,5 +65,5 @@ protected:
 
 	TSharedPtr<FExtender> Extender;
 
-	TSharedPtr<FChaosVDEditorModeTools> EditorModeTools;
+	TSharedPtr<FEditorModeTools> EditorModeTools;
 };
diff --git a/Engine/Plugins/Experimental/ChaosVD/Source/ChaosVD/Public/Widgets/SChaosVDWarningMessageBox.h b/Engine/Plugins/Experimental/ChaosVD/Source/ChaosVD/Public/Widgets/SChaosVDWarningMessageBox.h
new file mode 100644
--- /dev/null
+++ b/Engine/Plugins/Experimental/ChaosVD/Source/ChaosVD/Public/Widgets/SChaosVDWarningMessageBox.h
@@ -0,0 +1,19 @@
+// Copyright Epic Games, Inc. All Rights Reserved.
+
+#pragma once
+#include "Widgets/SCompoundWidget.h"
+
+class FText;
+
+class SChaosVDWarningMessageBox : public SCompoundWidget
+{
+public:
+	SLATE_BEGIN_ARGS(SChaosVDWarningMessageBox)
+		: _WarningText()
+	{}
+		/** The  text to show in the warning */
+		SLATE_ATTRIBUTE(FText, WarningText)
+	SLATE_END_ARGS()
+
+	void Construct(const FArguments& InArgs);
+};
diff --git a/Engine/Source/Runtime/Experimental/ChaosVDData/Public/ChaosVDRecording.h b/Engine/Source/Runtime/Experimental/ChaosVDData/Public/ChaosVDRecording.h
--- a/Engine/Source/Runtime/Experimental/ChaosVDData/Public/ChaosVDRecording.h
+++ b/Engine/Source/Runtime/Experimental/ChaosVDData/Public/ChaosVDRecording.h
@@ -19,6 +19,7 @@ struct FChaosVDStepData
 	FString StepName;
 	TArray<FChaosVDParticleDataWrapper> RecordedParticlesData;
 	TArray<TSharedPtr<FChaosVDParticlePairMidPhase>> RecordedMidPhases;
+	TArray<FChaosVDConstraint> RecordedConstraints;
 	TMap<int32, TArray<FChaosVDConstraint>> RecordedConstraintsByParticleID;
 	TMap<int32, TArray<TSharedPtr<FChaosVDParticlePairMidPhase>>> RecordedMidPhasesByParticleID;
 	TSet<int32> ParticlesDestroyedIDs;
diff --git a/Engine/Source/Runtime/Experimental/ChaosVisualDebugger/Public/DataWrappers/ChaosVDCollisionDataWrappers.h b/Engine/Source/Runtime/Experimental/ChaosVisualDebugger/Public/DataWrappers/ChaosVDCollisionDataWrappers.h
--- a/Engine/Source/Runtime/Experimental/ChaosVisualDebugger/Public/DataWrappers/ChaosVDCollisionDataWrappers.h
+++ b/Engine/Source/Runtime/Experimental/ChaosVisualDebugger/Public/DataWrappers/ChaosVDCollisionDataWrappers.h
@@ -57,23 +57,23 @@ struct CHAOSVDRUNTIME_API FChaosVDContactPoint
 	GENERATED_BODY()
 
 	// Shape-space contact points on the two bodies
-	UPROPERTY(EditAnywhere, Category=Contact)
+	UPROPERTY(VisibleAnywhere, Category=Contact)
 	TArray<FVector> ShapeContactPoints;
 
 	// Shape-space contact normal on the second shape with direction that points away from shape 1
-	UPROPERTY(EditAnywhere, Category=Contact)
+	UPROPERTY(VisibleAnywhere, Category=Contact)
 	FVector  ShapeContactNormal = FVector(ForceInit);
 
 	// Contact separation (negative for overlap)
-	UPROPERTY(EditAnywhere, Category=Contact)
+	UPROPERTY(VisibleAnywhere, Category=Contact)
 	float Phi = 0.f;
 
 	// Face index of the shape we hit. Only valid for Heightfield and Trimesh contact points, otherwise INDEX_NONE
-	UPROPERTY(EditAnywhere, Category=Contact)
+	UPROPERTY(VisibleAnywhere, Category=Contact)
 	int32 FaceIndex = 0;
 
 	// Whether this is a vertex-plane contact, edge-edge contact etc.
-	UPROPERTY(EditAnywhere, Category=Contact)
+	UPROPERTY(VisibleAnywhere, Category=Contact)
 	EChaosVDContactPointType ContactType = EChaosVDContactPointType::Unknown;
 
 	bool Serialize(FArchive& Ar);
@@ -99,35 +99,35 @@ struct CHAOSVDRUNTIME_API FChaosVDManifoldPoint
 {
 	GENERATED_BODY()
 
-	UPROPERTY(EditAnywhere, Category=ManifoldPointFalgs)
+	UPROPERTY(VisibleAnywhere, Category=ContactData)
 	uint8 bDisabled:1 = false;
-	UPROPERTY(EditAnywhere, Category=ManifoldPointFalgs)
+	UPROPERTY()
 	uint8 bWasRestored:1 = false;
-	UPROPERTY(EditAnywhere, Category=ManifoldPointFalgs)
+	UPROPERTY()
 	uint8 bWasReplaced:1 = false;
-	UPROPERTY(EditAnywhere, Category=ManifoldPointFalgs)
+	UPROPERTY(VisibleAnywhere, Category=ContactData)
 	uint8 bHasStaticFrictionAnchor:1 = false;
 
-	UPROPERTY(EditAnywhere, Category=ManifoldPointResult)
+	UPROPERTY()
 	uint8 bIsValid:1 = false;
-	UPROPERTY(EditAnywhere, Category=ManifoldPointResult)
+	UPROPERTY(VisibleAnywhere, Category=ContactData)
 	uint8 bInsideStaticFrictionCone:1 = false;
 
-	UPROPERTY(EditAnywhere, Category=ManifoldPointResult)
+	UPROPERTY(VisibleAnywhere, Category=ContactData)
 	FVector NetPushOut = FVector(ForceInit);
-	UPROPERTY(EditAnywhere, Category=ManifoldPointResult)
+	UPROPERTY(VisibleAnywhere, Category=ContactData)
 	FVector NetImpulse = FVector(ForceInit);
 
-	UPROPERTY(EditAnywhere, Category=ManifoldPoint)
+	UPROPERTY(VisibleAnywhere, Category=ContactData)
 	float TargetPhi = 0.f;
-	UPROPERTY(EditAnywhere, Category=ManifoldPoint)
+	UPROPERTY()
 	TArray<FVector> ShapeAnchorPoints;
-	UPROPERTY(EditAnywhere, Category=ManifoldPoint)
+	UPROPERTY()
 	TArray<FVector> InitialShapeContactPoints;
-	UPROPERTY(EditAnywhere, Category=ManifoldPoint)
+	UPROPERTY(VisibleAnywhere, Category=ContactData)
 	FChaosVDContactPoint ContactPoint;
 
-	UPROPERTY(EditAnywhere, Category=SavedManifoldPoint)
+	UPROPERTY()
 	TArray<FVector> ShapeContactPoints;
 
 	bool Serialize(FArchive& Ar);
@@ -155,86 +155,87 @@ struct CHAOSVDRUNTIME_API FChaosVDConstraint
 
 	inline static FStringView WrapperTypeName = TEXT("FChaosVDConstraint");
 
-	UPROPERTY(EditAnywhere, Category=Falgs)
+	UPROPERTY()
 	uint8 bIsCurrent:1 = false;
-	UPROPERTY(EditAnywhere, Category=Falgs)
+	UPROPERTY(VisibleAnywhere, Category=ConstraintData)
 	uint8 bDisabled:1 = false;
-	UPROPERTY(EditAnywhere, Category=Falgs)
+	UPROPERTY()
 	uint8 bUseManifold:1 = false;
-	UPROPERTY(EditAnywhere, Category=Falgs)
+	UPROPERTY()
 	uint8 bUseIncrementalManifold:1 = false;
-	UPROPERTY(EditAnywhere, Category=Falgs)
+	UPROPERTY()
 	uint8 bCanRestoreManifold:1 = false;
-	UPROPERTY(EditAnywhere, Category=Falgs)
+	UPROPERTY()
 	uint8 bWasManifoldRestored:1 = false;
-	UPROPERTY(EditAnywhere, Category=Falgs)
+	UPROPERTY()
 	uint8 bIsQuadratic0:1 = false;
-	UPROPERTY(EditAnywhere, Category=Falgs)
+	UPROPERTY()
 	uint8 bIsQuadratic1:1 = false;
-	UPROPERTY(EditAnywhere, Category=Falgs)
+	UPROPERTY(VisibleAnywhere, Category=ConstraintData)
 	uint8 bIsProbe:1 = false;
-	UPROPERTY(EditAnywhere, Category=Falgs)
+	UPROPERTY(VisibleAnywhere, Category=ConstraintData)
 	uint8 bCCDEnabled:1 = false;
-	UPROPERTY(EditAnywhere, Category=Falgs)
+	UPROPERTY(VisibleAnywhere, Category=ConstraintData)
 	uint8 bCCDSweepEnabled:1 = false;
-	UPROPERTY(EditAnywhere, Category=Falgs)
+	UPROPERTY(VisibleAnywhere, Category=ConstraintData)
 	uint8 bModifierApplied:1 = false;
-	UPROPERTY(EditAnywhere, Category=Falgs)
+	UPROPERTY(VisibleAnywhere, Category=ConstraintData)
 	uint8 bMaterialSet:1 = false;
 	
-	UPROPERTY(EditAnywhere, Category=Collision)
+	UPROPERTY(VisibleAnywhere, Category=ConstraintData)
 	FVector AccumulatedImpulse = FVector(ForceInit);
 	
-	UPROPERTY(EditAnywhere, Category=Collision)
+	UPROPERTY(VisibleAnywhere, Category=ConstraintData)
 	EChaosVDContactShapesType ShapesType = EChaosVDContactShapesType::Unknown;
 	
-	UPROPERTY(EditAnywhere, Category=Collision)
+	UPROPERTY(VisibleAnywhere, Category=ConstraintData)
 	TArray<FTransform> ShapeWorldTransforms;
 
-	UPROPERTY(EditAnywhere, Category=Collision)
+	UPROPERTY(VisibleAnywhere, Category=ConstraintData)
 	TArray<FTransform> ImplicitTransforms;
 	
-	UPROPERTY(EditAnywhere, Category=Collision)
+	UPROPERTY(VisibleAnywhere, Category=ConstraintData)
 	float CullDistance = 0.f;
 	
-	UPROPERTY(EditAnywhere, Category=Collision)
+	UPROPERTY(VisibleAnywhere, Category=ConstraintData)
 	TArray<float> CollisionMargins;
 	
-	UPROPERTY(EditAnywhere, Category=Collision)
+	UPROPERTY(VisibleAnywhere, Category=ConstraintData)
 	float CollisionTolerance = 0.f;
 	
-	UPROPERTY(EditAnywhere, Category=Collision)
+	UPROPERTY()
 	int32 ClosestManifoldPointIndex = 0;
 	
-	UPROPERTY(EditAnywhere, Category=Collision)
+	UPROPERTY()
 	int32 ExpectedNumManifoldPoints = 0;
 	
-	UPROPERTY(EditAnywhere, Category=Collision)
+	UPROPERTY()
 	FVector LastShapeWorldPositionDelta = FVector(ForceInit);
 	
-	UPROPERTY(EditAnywhere, Category=Collision)
+	UPROPERTY()
 	FQuat LastShapeWorldRotationDelta = FQuat(ForceInit);
 	
-	UPROPERTY(EditAnywhere, Category=Collision)
+	UPROPERTY(VisibleAnywhere, Category=ConstraintData)
 	float Stiffness = 0.f;
 	
-	UPROPERTY(EditAnywhere, Category=Collision)
+	UPROPERTY(VisibleAnywhere, Category=ConstraintData)
 	float CCDTimeOfImpact = 0.f;
 	
-	UPROPERTY(EditAnywhere, Category=Collision)
+	UPROPERTY(VisibleAnywhere, Category=ConstraintData)
 	float CCDEnablePenetration = 0.f;
 	
-	UPROPERTY(EditAnywhere, Category=Collision)
+	UPROPERTY(VisibleAnywhere, Category=ConstraintData)
 	float CCDTargetPenetration = 0.f;
 
-	UPROPERTY(EditAnywhere, Category=Collision)
+	UPROPERTY(VisibleAnywhere, Category=ConstraintData)
 	TArray<FChaosVDManifoldPoint> ManifoldPoints;
 
-	UPROPERTY(VisibleAnywhere, Category=Collision)
+	UPROPERTY()
 	int32 Particle0Index = INDEX_NONE;
-	UPROPERTY(VisibleAnywhere, Category=Collision)
+	UPROPERTY()
 	int32 Particle1Index = INDEX_NONE;
 
+	UPROPERTY()
 	int32 SolverID = INDEX_NONE;
 
 	bool Serialize(FArchive& Ar);
@@ -262,29 +263,29 @@ struct CHAOSVDRUNTIME_API FChaosVDParticlePairMidPhase
 
 	inline static FStringView WrapperTypeName = TEXT("FChaosVDParticlePairMidPhase");
 
-	UPROPERTY(VisibleAnywhere, Category=Solver)
+	UPROPERTY()
 	int32 SolverID = INDEX_NONE;
 
-	UPROPERTY(EditAnywhere, Category=Falgs)
+	UPROPERTY(VisibleAnywhere, Category=Flags)
 	uint8 bIsActive:1 = false;
-	UPROPERTY(EditAnywhere, Category=Falgs)
+	UPROPERTY(VisibleAnywhere, Category=Flags)
 	uint8 bIsCCD:1 = false;
-	UPROPERTY(EditAnywhere, Category=Falgs)
+	UPROPERTY(VisibleAnywhere, Category=Flags)
 	uint8 bIsCCDActive:1 = false;
-	UPROPERTY(EditAnywhere, Category=Falgs)
+	UPROPERTY(VisibleAnywhere, Category=Flags)
 	uint8 bIsSleeping:1 = false;
-	UPROPERTY(EditAnywhere, Category=Falgs)
+	UPROPERTY(VisibleAnywhere, Category=Flags)
 	uint8 bIsModified:1 = false;
 
-	UPROPERTY(EditAnywhere, Category=Misc)
+	UPROPERTY(VisibleAnywhere, Category=Misc)
 	int32 LastUsedEpoch = 0;
 
-	UPROPERTY(EditAnywhere, Category=Particle)
+	UPROPERTY(VisibleAnywhere, Category=Particle)
 	int32 Particle0Idx = 0;
-	UPROPERTY(EditAnywhere, Category=Particle)
+	UPROPERTY(VisibleAnywhere, Category=Particle)
 	int32 Particle1Idx = 0;
 
-	UPROPERTY(EditAnywhere, Category=Particle)
+	UPROPERTY()
 	TArray<FChaosVDConstraint> Constraints;
 
 	bool Serialize(FArchive& Ar);
diff --git a/Engine/Source/Runtime/Experimental/ChaosVisualDebugger/Public/DataWrappers/ChaosVDParticleDataWrapper.h b/Engine/Source/Runtime/Experimental/ChaosVisualDebugger/Public/DataWrappers/ChaosVDParticleDataWrapper.h
--- a/Engine/Source/Runtime/Experimental/ChaosVisualDebugger/Public/DataWrappers/ChaosVDParticleDataWrapper.h
+++ b/Engine/Source/Runtime/Experimental/ChaosVisualDebugger/Public/DataWrappers/ChaosVDParticleDataWrapper.h
@@ -409,7 +409,7 @@ struct FChaosVDParticleDataWrapper : public FChaosVDParticleDataBase
 	{
 	}
 
-	UPROPERTY(EditAnywhere, Category= "Particle Non Frequent Data")
+	UPROPERTY(VisibleAnywhere, Category= "Particle Non Frequent Data")
 	uint32 GeometryHash = 0;
 
 	UPROPERTY(VisibleAnywhere, Category= "Particle Non Frequent Data")
@@ -439,13 +439,7 @@ struct FChaosVDParticleDataWrapper : public FChaosVDParticleDataBase
 	UPROPERTY(EditAnywhere, Category= "Particle Particle Mass Props")
 	FChaosVDParticleMassProps ParticleMassProps;
 
-	UPROPERTY(EditAnywhere, Category= "Particle Collision")
-	TArray<FChaosVDParticlePairMidPhase> ParticleMidPhases;
-
-	UPROPERTY(EditAnywhere, Category= "Particle Collision")
-	TArray<FChaosVDConstraint> ParticleConstraints;
-
-	UPROPERTY(EditAnywhere, Category= "Particle Collision")
+	UPROPERTY()
 	TArray<FChaosVDShapeCollisionData> CollisionDataPerShape;
 	
 	/** Only used during recording */
