[ChaosVD] Improvements to Collision Data visualization

- Fixed an issue where NetPushOut and NetImpulse vectors were being debug draw in the incorrect direction
- Changed the Debug draw for contact points from point to circle (now matches the runtime debug draw implementation that already exists)
- Added debug draw for the constraint and contact point state
- Added support to debug Draw Circles and boxes
- Fixed an issue were the tip of debug draw arrows was too big for small vectors
- Moved the contact click handling code to the component visualizer for collision data, using the existing component visualizer interface
- Added selection feedback in-viewport. Now when you select a contact, a yellow wireframe box will be drawn containing the selected contact.

diff --git a/Engine/Plugins/Experimental/ChaosVD/Source/ChaosVD/Private/ChaosVDParticleActor.cpp b/Engine/Plugins/Experimental/ChaosVD/Source/ChaosVD/Private/ChaosVDParticleActor.cpp
--- a/Engine/Plugins/Experimental/ChaosVD/Source/ChaosVD/Private/ChaosVDParticleActor.cpp
+++ b/Engine/Plugins/Experimental/ChaosVD/Source/ChaosVD/Private/ChaosVDParticleActor.cpp
@@ -267,7 +267,7 @@ const TArray<TSharedPtr<FChaosVDParticlePairMidPhase>>* AChaosVDParticleActor::G
 	{
 		if (const UChaosVDSolverCollisionDataComponent* CollisionDataComponent = SolverInfoActor->GetCollisionDataComponent())
 		{
-			return CollisionDataComponent->GetMidPhasesForParticle(ParticleDataViewer.ParticleIndex, EChaosVDGetCollisionDataOptions::Any);
+			return CollisionDataComponent->GetMidPhasesForParticle(ParticleDataViewer.ParticleIndex, EChaosVDCollisionParticlePairSlot::Any);
 		}
 	}
 
diff --git a/Engine/Plugins/Experimental/ChaosVD/Source/ChaosVD/Private/ChaosVDPlaybackViewportClient.cpp b/Engine/Plugins/Experimental/ChaosVD/Source/ChaosVD/Private/ChaosVDPlaybackViewportClient.cpp
--- a/Engine/Plugins/Experimental/ChaosVD/Source/ChaosVD/Private/ChaosVDPlaybackViewportClient.cpp
+++ b/Engine/Plugins/Experimental/ChaosVD/Source/ChaosVD/Private/ChaosVDPlaybackViewportClient.cpp
@@ -8,8 +8,6 @@
 #include "ChaosVDPlaybackController.h"
 #include "ChaosVDScene.h"
 #include "ChaosVDSkySphereInterface.h"
-#include "ChaosVDCollisionDataDetailsTab.h"
-#include "ChaosVDTabsIDs.h"
 #include "ComponentVisualizer.h"
 #include "EditorModeManager.h"
 #include "Elements/Framework/TypedElementSelectionSet.h"
@@ -21,8 +19,6 @@
 #include "Actors/ChaosVDSolverInfoActor.h"
 #include "Components/ChaosVDSolverCollisionDataComponent.h"
 #include "Visualizers/ChaosVDDebugDrawUtils.h"
-#include "Visualizers/ChaosVDSolverCollisionDataComponentVisualizer.h"
-#include "Widgets/SChaosVDCollisionDataInspector.h"
 #include "Widgets/SChaosVDMainTab.h"
 
 FChaosVDPlaybackViewportClient::FChaosVDPlaybackViewportClient(const TSharedPtr<FEditorModeTools>& InModeTools) : FEditorViewportClient(InModeTools.Get()), CVDWorld(nullptr)
@@ -66,8 +62,27 @@ void FChaosVDPlaybackViewportClient::ProcessClick(FSceneView& View, HHitProxy* H
 		return;
 	}
 
+	const FViewportClick Click(&View, this, Key, Event, HitX, HitY);
+
 	if (TSharedPtr<FChaosVDScene> ScenePtr = CVDScene.Pin())
 	{
+		bool bClickHandled = false;
+
+		// TODO: Iterate trough all registered visualizers
+		HComponentVisProxy* ComponentVisProxy = HitProxyCast<HComponentVisProxy>(HitProxy);
+		if (const TSharedPtr<FComponentVisualizer> Visualizer = MainTabToolkitHost->FindComponentVisualizer(UChaosVDSolverCollisionDataComponent::StaticClass()))
+		{
+			// Not sure if this is compliant with the normal use of the component visualizers,
+			// but passing a null hitproxy when the hit proxy was not a component
+			// It allow us to handle things like clear selection on the Collision Data Visualizer
+			bClickHandled = Visualizer->VisProxyHandleClick(this, ComponentVisProxy, Click);
+		}
+
+		if (bClickHandled)
+		{
+			return;
+		}
+
 		if (HitProxy->IsA(HActor::StaticGetType()))
 		{
 			const HActor* ActorHitProxy = static_cast<HActor*>(HitProxy);
@@ -76,26 +91,6 @@ void FChaosVDPlaybackViewportClient::ProcessClick(FSceneView& View, HHitProxy* H
 				ScenePtr->SetSelectedObject(ClickedActor);
 			}
 		}
-		else
-		{
-			if (HChaosVDContactPointProxy* ContactProxy = HitProxyCast<HChaosVDContactPointProxy>(HitProxy))
-			{
-				if (ContactProxy->ContactFinder.OwningMidPhase.Pin())
-				{
-					if (TSharedPtr<FChaosVDCollisionDataDetailsTab> CollisionDataDetailsTab = MainTabToolkitHost->GetTabSpawnerInstance<FChaosVDCollisionDataDetailsTab>(FChaosVDTabID::CollisionDataDetails).Pin())
-					{
-						if (TSharedPtr<SChaosVDCollisionDataInspector> CollisionInspector = CollisionDataDetailsTab->GetCollisionInspectorInstance().Pin())
-						{
-							if (TSharedPtr<FTabManager> TabManager = MainTabToolkitHost->GetTabManager())
-							{
-								TabManager->TryInvokeTab(FChaosVDTabID::CollisionDataDetails);
-								CollisionInspector->SetSingleContactDataToInspect(ContactProxy->ContactFinder);
-							}
-						}
-					}
-				}
-			}
-		}
 	}
 }
 
@@ -278,8 +273,6 @@ void FChaosVDPlaybackViewportClient::PerformSelectedTrackingForFrame(FChaosVDGam
 
 void FChaosVDPlaybackViewportClient::Draw(const FSceneView* View, FPrimitiveDrawInterface* PDI)
 {
-	FEditorViewportClient::Draw(View, PDI);
-
 	const TSharedPtr<SChaosVDMainTab> MainTabToolkitHost = ModeTools.IsValid() ? StaticCastSharedPtr<SChaosVDMainTab>(ModeTools->GetToolkitHost()) : nullptr;
 	if (!MainTabToolkitHost.IsValid())
 	{
@@ -311,6 +304,8 @@ void FChaosVDPlaybackViewportClient::Draw(const FSceneView* View, FPrimitiveDraw
 			}
 		}
 	}
+
+	FEditorViewportClient::Draw(View, PDI);
 }
 
 void FChaosVDPlaybackViewportClient::DrawCanvas(FViewport& InViewport, FSceneView& View, FCanvas& Canvas)
diff --git a/Engine/Plugins/Experimental/ChaosVD/Source/ChaosVD/Private/Components/ChaosVDSolverCollisionDataComponent.cpp b/Engine/Plugins/Experimental/ChaosVD/Source/ChaosVD/Private/Components/ChaosVDSolverCollisionDataComponent.cpp
--- a/Engine/Plugins/Experimental/ChaosVD/Source/ChaosVD/Private/Components/ChaosVDSolverCollisionDataComponent.cpp
+++ b/Engine/Plugins/Experimental/ChaosVD/Source/ChaosVD/Private/Components/ChaosVDSolverCollisionDataComponent.cpp
@@ -35,12 +35,12 @@ void UChaosVDSolverCollisionDataComponent::UpdateCollisionData(const TArray<TSha
 	}
 }
 
-const TArray<TSharedPtr<FChaosVDParticlePairMidPhase>>* UChaosVDSolverCollisionDataComponent::GetMidPhasesForParticle(int32 ParticleID, EChaosVDGetCollisionDataOptions Options) const
+const TArray<TSharedPtr<FChaosVDParticlePairMidPhase>>* UChaosVDSolverCollisionDataComponent::GetMidPhasesForParticle(int32 ParticleID, EChaosVDCollisionParticlePairSlot Options) const
 {
 	return GetCollisionDataFromMap<FChaosVDMidPhaseByParticleMap, TSharedPtr<FChaosVDParticlePairMidPhase>>(MidPhasesByParticleID0, MidPhasesByParticleID1, ParticleID, Options);
 }
 
-const TArray<FChaosVDConstraint*>* UChaosVDSolverCollisionDataComponent::GetConstraintsForParticle(int32 ParticleID, EChaosVDGetCollisionDataOptions Options) const
+const TArray<FChaosVDConstraint*>* UChaosVDSolverCollisionDataComponent::GetConstraintsForParticle(int32 ParticleID, EChaosVDCollisionParticlePairSlot Options) const
 {
 	return GetCollisionDataFromMap<FChaosVDConstraintByParticleMap, FChaosVDConstraint*>(ConstraintsByParticleID0, ConstraintsByParticleID1, ParticleID, Options);
 }
diff --git a/Engine/Plugins/Experimental/ChaosVD/Source/ChaosVD/Private/Visualizers/ChaosVDDebugDrawUtils.cpp b/Engine/Plugins/Experimental/ChaosVD/Source/ChaosVD/Private/Visualizers/ChaosVDDebugDrawUtils.cpp
--- a/Engine/Plugins/Experimental/ChaosVD/Source/ChaosVD/Private/Visualizers/ChaosVDDebugDrawUtils.cpp
+++ b/Engine/Plugins/Experimental/ChaosVD/Source/ChaosVD/Private/Visualizers/ChaosVDDebugDrawUtils.cpp
@@ -4,33 +4,58 @@
 
 #include "CanvasItem.h"
 #include "ChaosVDEditorSettings.h"
-#include "DebugRenderSceneProxy.h"
 #include "Engine/Engine.h"
+#include "SceneView.h"
 
 TQueue<FChaosVDDebugDrawUtils::FChaosVDQueuedTextToDraw> FChaosVDDebugDrawUtils::TexToDrawQueue = TQueue<FChaosVDQueuedTextToDraw>();
 
-void FChaosVDDebugDrawUtils::DrawArrowVector(FPrimitiveDrawInterface* PDI, const FVector& StartLocation, const FVector& EndLocation, FStringView DebugText, const FColor& Color, float ArrowSize)
+void FChaosVDDebugDrawUtils::DrawArrowVector(FPrimitiveDrawInterface* PDI, const FVector& StartLocation, const FVector& EndLocation, FStringView DebugText, const FColor& Color, ESceneDepthPriorityGroup DepthPriority)
 {
-	const FDebugRenderSceneProxy::FArrowLine VelocityArrowLine = { StartLocation, EndLocation, Color};
-	VelocityArrowLine.Draw(PDI, ArrowSize);
+	if (!PDI)
+	{
+		return;
+	}
+
+	const FVector LineVectorToDraw = EndLocation - StartLocation;
+
+	FVector ArrowDir;
+	float ArrowLength;
+	LineVectorToDraw.ToDirectionAndLength(ArrowDir, ArrowLength);
+
+	FVector YAxis, ZAxis;
+	ArrowDir.FindBestAxisVectors(YAxis,ZAxis);
+	const FMatrix ArrowTransformMatrix(ArrowDir, YAxis, ZAxis,StartLocation);
+
+	constexpr float MinTipOfArrowSize = 0.2f;
+	constexpr float MaxTipOfArrowSize = 10.0f;
+	constexpr float MaxVectorSizeForArrow = 100.0f; // The vector size that is the upper limit after which we just use the max size for the tip of the arrow
+
+	const float ProportionalArrowSize = MaxTipOfArrowSize * (ArrowLength / MaxVectorSizeForArrow);
+	const float ArrowSize = FMath::Clamp(ProportionalArrowSize, MinTipOfArrowSize, MaxTipOfArrowSize);
+	
+	DrawDirectionalArrow(PDI, ArrowTransformMatrix, Color, ArrowLength, ArrowSize, DepthPriority);
 
 	if (!DebugText.IsEmpty())
 	{
 		// Draw the text in the middle of the vector line
-		const FVector VectorToDraw = EndLocation - StartLocation;
-		const FVector TextWorldPosition = StartLocation + VectorToDraw  * 0.5f;
+		const FVector TextWorldPosition = StartLocation + LineVectorToDraw  * 0.5f;
 		DrawText(DebugText.GetData(), TextWorldPosition , Color);
 	}
 }
 
-void FChaosVDDebugDrawUtils::DrawPoint(FPrimitiveDrawInterface* PDI, const FVector& Location, FStringView DebugText, const FColor& Color, float Size)
+void FChaosVDDebugDrawUtils::DrawPoint(FPrimitiveDrawInterface* PDI, const FVector& Location, FStringView DebugText, const FColor& Color, float Size, ESceneDepthPriorityGroup DepthPriority)
 {
+	if (!PDI)
+	{
+		return;
+	}
+
 	if (DebugText.IsEmpty())
 	{
 		return;
 	}
 
-	PDI->DrawPoint(Location, Color, Size, ESceneDepthPriorityGroup::SDPG_World);
+	PDI->DrawPoint(Location, Color, Size, DepthPriority);
 
 	if (!DebugText.IsEmpty())
 	{
@@ -49,8 +74,87 @@ void FChaosVDDebugDrawUtils::DrawText(FStringView StringToDraw, const FVector& L
 	}
 }
 
+void FChaosVDDebugDrawUtils::DrawCircle(FPrimitiveDrawInterface* PDI, const FVector& Origin, float Radius, int32 Segments, const FColor& Color, float Thickness, const FVector& XAxis, const FVector& YAxis, FStringView DebugText, ESceneDepthPriorityGroup DepthPriority)
+{
+	if (!PDI)
+	{
+		return;
+	}
+
+	constexpr float DepthBias = 0;
+	const bool bScreenSpace = Thickness > 0;
+
+	// Need at least 2 sides
+	Segments = FMath::Max(Segments, 2);
+	const float	AngleDelta = 2.0f * UE_PI / Segments;
+	FVector	LastVertex = Origin + XAxis * Radius;
+
+	for (int32 SideIndex = 0; SideIndex < Segments; SideIndex++)
+	{
+		const FVector Vertex = Origin + (XAxis * FMath::Cos(AngleDelta * (SideIndex + 1)) + YAxis * FMath::Sin(AngleDelta * (SideIndex + 1))) * Radius;
+
+		PDI->DrawLine(LastVertex, Vertex, Color, DepthPriority, Thickness, DepthBias, bScreenSpace);
+
+		LastVertex = Vertex;
+	}
+	
+	if (!DebugText.IsEmpty())
+	{
+		DrawText(DebugText, Origin, Color);
+	}
+}
+
+void FChaosVDDebugDrawUtils::DrawBox(FPrimitiveDrawInterface* PDI, const FVector& InExtents, const FColor& InColor, const FTransform& InTransform, FStringView DebugText, ESceneDepthPriorityGroup DepthPriority)
+{
+	if (!PDI)
+	{
+		return;
+	}
+
+	constexpr int32 MaxBoxLines = 12;
+
+	// Array for direction offsets for the start/end point of each line
+	static TPair<FVector, FVector> VertexOffsetDirectionFromOrigin[MaxBoxLines] =
+	{
+		{FVector(1,1,1), FVector(1,-1,1)},
+		{FVector(1,-1,1) ,FVector(-1,-1,1)},
+		{FVector(-1,-1,1), FVector(-1,1,1)},
+		{FVector(-1,1,1), FVector(1,1,1)},
+		{FVector(1,1,-1), FVector(1,-1,-1)},
+		{FVector(1, -1,-1), FVector(-1,-1,-1)},
+		{FVector(-1,-1,-1), FVector(-1,1,-1)},
+		{FVector(-1,1,-1), FVector(1,1,-1)},
+		{FVector(1,1,1), FVector(1,1,-1)},
+		{FVector(1,-1,1), FVector(1,-1,-1)},
+		{FVector(-1,-1,1), FVector(-1,-1,-1)},
+		{FVector(-1, 1,1), FVector(-1,1,-1)},
+	};
+
+	constexpr float Thickness = 2.0f;
+	constexpr float DepthBias = 0;
+	constexpr bool bScreenSpace = Thickness > 0;
+
+	for (int32 BoxLineIndex = 0; BoxLineIndex < MaxBoxLines; BoxLineIndex++)
+	{
+		FVector LineStart = InTransform.TransformPosition(InExtents * VertexOffsetDirectionFromOrigin[BoxLineIndex].Key);
+		FVector LineEnd = InTransform.TransformPosition(InExtents * VertexOffsetDirectionFromOrigin[BoxLineIndex].Value);
+
+		PDI->DrawLine(LineStart, LineEnd, InColor, DepthPriority, Thickness,DepthBias, bScreenSpace);
+	}
+	
+	if (!DebugText.IsEmpty())
+	{
+		DrawText(DebugText, InTransform.GetLocation(), InColor);
+	}
+}
+
 void FChaosVDDebugDrawUtils::DrawCanvas(FViewport& InViewport, FSceneView& View, FCanvas& Canvas)
 {
+	if (!GEngine)
+	{
+		return;
+	}
+
 	while (!TexToDrawQueue.IsEmpty())
 	{
 		FChaosVDQueuedTextToDraw TextToDraw;
diff --git a/Engine/Plugins/Experimental/ChaosVD/Source/ChaosVD/Private/Visualizers/ChaosVDSolverCollisionDataComponentVisualizer.cpp b/Engine/Plugins/Experimental/ChaosVD/Source/ChaosVD/Private/Visualizers/ChaosVDSolverCollisionDataComponentVisualizer.cpp
--- a/Engine/Plugins/Experimental/ChaosVD/Source/ChaosVD/Private/Visualizers/ChaosVDSolverCollisionDataComponentVisualizer.cpp
+++ b/Engine/Plugins/Experimental/ChaosVD/Source/ChaosVD/Private/Visualizers/ChaosVDSolverCollisionDataComponentVisualizer.cpp
@@ -2,14 +2,31 @@
 
 #include "ChaosVDSolverCollisionDataComponentVisualizer.h"
 
+#include "ChaosVDCollisionDataDetailsTab.h"
 #include "ChaosVDEditorSettings.h"
 #include "ChaosVDParticleActor.h"
 #include "ChaosVDScene.h"
+#include "ChaosVDTabsIDs.h"
+#include "EditorModeManager.h"
+#include "EditorViewportClient.h"
 #include "Actors/ChaosVDSolverInfoActor.h"
 #include "Components/ChaosVDSolverCollisionDataComponent.h"
 #include "Visualizers/ChaosVDDebugDrawUtils.h"
+#include "Widgets/SChaosVDCollisionDataInspector.h"
+#include "Widgets/SChaosVDMainTab.h"
 
-IMPLEMENT_HIT_PROXY(HChaosVDContactPointProxy, HHitProxy)
+IMPLEMENT_HIT_PROXY(HChaosVDContactPointProxy, HComponentVisProxy)
+
+void FChaosVDCollisionDataFinder::SetIsSelected(bool bNewSelected)
+{
+	if (TSharedPtr<FChaosVDParticlePairMidPhase> MidPhaseData = OwningMidPhase.Pin())
+	{
+		if (OwningConstraint && OwningConstraint->ManifoldPoints.IsValidIndex(ContactIndex))
+		{
+			OwningConstraint->ManifoldPoints[ContactIndex].bIsSelectedInEditor = bNewSelected;
+		}
+	}
+}
 
 FChaosVDSolverCollisionDataComponentVisualizer::FChaosVDSolverCollisionDataComponentVisualizer()
 {
@@ -52,7 +69,7 @@ void FChaosVDSolverCollisionDataComponentVisualizer::DrawVisualization(const UAc
 	{
 		for (const TSharedPtr<FChaosVDParticlePairMidPhase>& MidPhase : CollisionDataComponent->GetMidPhases())
 		{
-			DrawnMidPhaseData(MidPhase, VisualizationContext, View, PDI);
+			DrawMidPhaseData(Component, MidPhase, VisualizationContext, View, PDI);
 		}
 	}
 	else
@@ -61,28 +78,85 @@ void FChaosVDSolverCollisionDataComponentVisualizer::DrawVisualization(const UAc
 
 		for (const int32 SelectedParticleID : SelectedParticlesIDs)
 		{
-			if (const TArray<TSharedPtr<FChaosVDParticlePairMidPhase>>* ParticleMidPhases = CollisionDataComponent->GetMidPhasesForParticle(SelectedParticleID, EChaosVDGetCollisionDataOptions::Owner))
+			if (const TArray<TSharedPtr<FChaosVDParticlePairMidPhase>>* ParticleMidPhases = CollisionDataComponent->GetMidPhasesForParticle(SelectedParticleID, EChaosVDCollisionParticlePairSlot::Any))
 			{
 				for (const TSharedPtr<FChaosVDParticlePairMidPhase>& MidPhase : *ParticleMidPhases)
 				{
-					DrawnMidPhaseData(MidPhase, VisualizationContext, View, PDI);
+					DrawMidPhaseData(Component, MidPhase, VisualizationContext, View, PDI);
 				}
 			}
-		}		
+		}
+	}
+}
+
+bool FChaosVDSolverCollisionDataComponentVisualizer::VisProxyHandleClick(FEditorViewportClient* InViewportClient, HComponentVisProxy* VisProxy, const FViewportClick& Click)
+{
+	const HChaosVDContactPointProxy* ContactProxy = HitProxyCast<HChaosVDContactPointProxy>(VisProxy);
+	if (ContactProxy == nullptr)
+	{
+		ClearCurrentSelection();
+		return false;
+	}
+
+	const TSharedPtr<SChaosVDMainTab> MainTabToolkitHost = InViewportClient->GetModeTools() ? StaticCastSharedPtr<SChaosVDMainTab>(InViewportClient->GetModeTools()->GetToolkitHost()) : nullptr;
+	if (!MainTabToolkitHost.IsValid())
+	{
+		return false;
 	}
+
+	const TSharedPtr<FChaosVDCollisionDataDetailsTab> CollisionDataDetailsTab = MainTabToolkitHost->GetTabSpawnerInstance<FChaosVDCollisionDataDetailsTab>(FChaosVDTabID::CollisionDataDetails).Pin();
+	if (!CollisionDataDetailsTab)
+	{
+		return false;
+	}
+
+	const TSharedPtr<SChaosVDCollisionDataInspector> CollisionInspector = CollisionDataDetailsTab->GetCollisionInspectorInstance().Pin();
+	if (!CollisionInspector)
+	{
+		return false;
+	}
+
+	const TSharedPtr<FTabManager> TabManager = MainTabToolkitHost->GetTabManager();
+	if (!TabManager)
+	{
+		return false;
+	}
+	
+	if (ContactProxy->ContactFinder.OwningMidPhase.Pin())
+	{
+		TabManager->TryInvokeTab(FChaosVDTabID::CollisionDataDetails);
+		CollisionInspector->SetSingleContactDataToInspect(ContactProxy->ContactFinder);
+		
+		ClearCurrentSelection();
+		CurrentSelectedContactData = ContactProxy->ContactFinder;
+		CurrentSelectedContactData.SetIsSelected(true);
+		return true;
+	}
+
+	return false;
 }
 
-void FChaosVDSolverCollisionDataComponentVisualizer::DrawnMidPhaseData(const TSharedPtr<FChaosVDParticlePairMidPhase>& MidPhase, const FChaosVDVisualizationContext& VisualizationContext, const FSceneView* View, FPrimitiveDrawInterface* PDI)
+void FChaosVDSolverCollisionDataComponentVisualizer::ClearCurrentSelection()
+{
+	CurrentSelectedContactData.SetIsSelected(false);
+}
+
+void FChaosVDSolverCollisionDataComponentVisualizer::DrawMidPhaseData(const UActorComponent* Component, const TSharedPtr<FChaosVDParticlePairMidPhase>& MidPhase, const FChaosVDVisualizationContext& VisualizationContext, const FSceneView* View, FPrimitiveDrawInterface* PDI)
 {
 	TSharedPtr<FChaosVDScene> ScenePtr = VisualizationContext.CVDScene.Pin();
 	if (!ScenePtr.IsValid())
 	{
 		return;
 	}
+
+	FChaosVDContactDebugDrawSettings DebugDrawSettings;
+	if (const UChaosVDEditorSettings* EditorSettings = GetDefault<UChaosVDEditorSettings>())
+	{
+		DebugDrawSettings = EditorSettings->ContactDebugDrawSettings;
+	}
 	
 	if (MidPhase.IsValid())
 	{
-
 		AChaosVDParticleActor* ParticleActor = ScenePtr->GetParticleActor(VisualizationContext.SolverID, MidPhase->Particle0Idx);
 		if (!ParticleActor)
 		{
@@ -96,10 +170,8 @@ void FChaosVDSolverCollisionDataComponentVisualizer::DrawnMidPhaseData(const TSh
 		}
 
 		EChaosVDCollisionVisualizationFlags VisualizationFlags = static_cast<EChaosVDCollisionVisualizationFlags>(VisualizationContext.VisualizationFlags);
-		 
-		const FVector& OwnerActorLocation = ParticleDataViewer->ParticlePositionRotation.MX;
 
-		for (const FChaosVDConstraint& Constraint : MidPhase->Constraints)
+		for (FChaosVDConstraint& Constraint : MidPhase->Constraints)
 		{
 			int32 ContactIndex = INDEX_NONE;
 			for (const FChaosVDManifoldPoint& ManifoldPoint : Constraint.ManifoldPoints)
@@ -119,7 +191,7 @@ void FChaosVDSolverCollisionDataComponentVisualizer::DrawnMidPhaseData(const TSh
 
 				const bool bIsProbe = Constraint.bIsProbe;
 				const bool bIsActive = ManifoldPoint.bIsValid && (!ManifoldPoint.NetPushOut.IsNearlyZero() || !ManifoldPoint.NetImpulse.IsNearlyZero() || (!Constraint.bUseManifold && !Constraint.AccumulatedImpulse.IsNearlyZero()));
-				if (!bIsActive && EnumHasAnyFlags(VisualizationFlags, EChaosVDCollisionVisualizationFlags::DrawOnlyActiveContacts))
+				if (!bIsActive && !EnumHasAnyFlags(VisualizationFlags, EChaosVDCollisionVisualizationFlags::DrawInactiveContacts))
 				{
 					continue;
 				}
@@ -133,7 +205,7 @@ void FChaosVDSolverCollisionDataComponentVisualizer::DrawnMidPhaseData(const TSh
 					continue;
 				}
 
-				PDI->SetHitProxy(new HChaosVDContactPointProxy(HitPRoxyDataFinder));
+				PDI->SetHitProxy(new HChaosVDContactPointProxy(Component, HitPRoxyDataFinder));
 
 				const FTransform& WorldActorTransform1 = CVDParticleActor1->GetTransform();
 				const FTransform& WorldActorTransform0 = CVDParticleActor0->GetTransform();
@@ -155,31 +227,45 @@ void FChaosVDSolverCollisionDataComponentVisualizer::DrawnMidPhaseData(const TSh
 				// Dynamic friction, restitution = red
 				// Static friction, no restitution = green
 				// Inactive = gray
-				FColor DiscColor = FColor(200, 0, 0);
-				FColor PlaneNormalColor = FColor(200, 0, 0);
-				FColor EdgeNormalColor = FColor(200, 150, 0);
-				FColor ImpulseColor = FColor(0, 0, 200);
-				FColor PushOutImpusleColor = FColor(0, 200, 200);
+				FColor DiscColor = FColor(250, 0, 0);
+				FColor PhiDiscColor = FColor(128, 128, 0);
+				FColor PlaneNormalColor = FColor(250, 0, 0);
+				FColor EdgeNormalColor = FColor(250, 150, 0);
+				FColor ImpulseColor = FColor(0, 0, 250);
+				FColor PushOutImpulseColor = FColor(0, 250, 250);
+				float ContactLenScale = 1.0f;
+
+				// TODO: Make this a setting when all CVD debug draw options support changing the line thickness
+				constexpr float LinesThickness = 2.0f;
+				constexpr int32 CircleSegments = 16;
+				
 				if (ManifoldPoint.bInsideStaticFrictionCone)
 				{
 					DiscColor = FColor(150, 200, 0);
+					PhiDiscColor = FColor(150, 200, 0);
 				}
 				if (bIsProbe)
 				{
 					DiscColor = FColor(50, 180, 180);
+					PhiDiscColor = FColor(50, 180, 180);
 					PlaneNormalColor = FColor(50, 180, 180);
 					EdgeNormalColor = FColor(50, 180, 130);
 				}
 				else if (!bIsActive)
 				{
 					DiscColor = FColor(100, 100, 100);
+					PhiDiscColor = FColor(150, 150, 150);
 					PlaneNormalColor = FColor(100, 0, 0);
 					EdgeNormalColor = FColor(100, 80, 0);
+					ContactLenScale = 0.75;
 				}
 				if (bPruned)
 				{
+					DiscColor = FColor(50, 50, 50);
+					PhiDiscColor = FColor(50, 50, 50);
 					PlaneNormalColor = FColor(200, 0, 200);
 					EdgeNormalColor = FColor(200, 0, 200);
+					ContactLenScale = 0.5;
 				}
 			
 				const FVector WorldPointLocation = VisualizationContext.SpaceTransform.TransformPosition(PointLocation);
@@ -187,57 +273,103 @@ void FChaosVDSolverCollisionDataComponentVisualizer::DrawnMidPhaseData(const TSh
 				const FVector WorldPointPlaneLocation = VisualizationContext.SpaceTransform.TransformPosition(PointPlaneLocation);
 				const FVector WorldPlaneNormal = VisualizationContext.SpaceTransform.TransformVectorNoScale(PlaneNormal);
 
+				const FMatrix Axes = FRotationMatrix::MakeFromX(WorldPlaneNormal);
+
 				if (EnumHasAnyFlags(VisualizationFlags, EChaosVDCollisionVisualizationFlags::NetPushOut))
 				{
 					if (ManifoldPoint.bIsValid && !ManifoldPoint.NetPushOut.IsNearlyZero())
 					{
-						FChaosVDDebugDrawUtils::DrawArrowVector(PDI, WorldPointPlaneLocation, OwnerActorLocation + WorldPointPlaneLocation + VisualizationContext.SpaceTransform.TransformPosition(ManifoldPoint.NetPushOut), TEXT("Net Push out"), PushOutImpusleColor);
+						FChaosVDDebugDrawUtils::DrawArrowVector(PDI, WorldPointPlaneLocation, WorldPointPlaneLocation + VisualizationContext.SpaceTransform.TransformPosition(ManifoldPoint.NetPushOut), TEXT("Net Push out"), PushOutImpulseColor, DebugDrawSettings.DepthPriority);
 					}
 				}
 
-				if (EnumHasAnyFlags(static_cast<EChaosVDCollisionVisualizationFlags>(VisualizationContext.VisualizationFlags), EChaosVDCollisionVisualizationFlags::NetImpulse))
+				if (EnumHasAnyFlags(VisualizationFlags, EChaosVDCollisionVisualizationFlags::NetImpulse))
 				{
 					if (ManifoldPoint.bIsValid && !ManifoldPoint.NetImpulse.IsNearlyZero())
 					{
-						FChaosVDDebugDrawUtils::DrawArrowVector(PDI, WorldPointPlaneLocation, OwnerActorLocation + WorldPointPlaneLocation + VisualizationContext.SpaceTransform.TransformPosition(ManifoldPoint.NetImpulse), TEXT("Net Impulse"), ImpulseColor);
+						FChaosVDDebugDrawUtils::DrawArrowVector(PDI, WorldPointPlaneLocation, WorldPointPlaneLocation + VisualizationContext.SpaceTransform.TransformPosition(ManifoldPoint.NetImpulse), TEXT("Net Impulse"), ImpulseColor, DebugDrawSettings.DepthPriority);
 					}
 				}
 
 				if (EnumHasAnyFlags(VisualizationFlags, EChaosVDCollisionVisualizationFlags::ContactPoints))
 				{
-					FChaosVDDebugDrawUtils::DrawPoint(PDI, WorldPlaneLocation, TEXT("Contact : WorldPlaneLocation"), DiscColor);
+					FChaosVDDebugDrawUtils::DrawCircle(PDI, WorldPlaneLocation, DebugDrawSettings.ContactCircleRadius, CircleSegments, DiscColor, LinesThickness, Axes.GetUnitAxis(EAxis::Y), Axes.GetUnitAxis(EAxis::Z), TEXT("Contact : WorldPlaneLocation"), DebugDrawSettings.DepthPriority);
 				}
 
 				if (EnumHasAnyFlags(VisualizationFlags, EChaosVDCollisionVisualizationFlags::ContactNormal))
 				{
 					FColor NormalColor = ((ManifoldPoint.ContactPoint.ContactType != EChaosVDContactPointType::EdgeEdge) ? PlaneNormalColor : EdgeNormalColor);
-					constexpr int32 Scale = 10.0f;
-					FChaosVDDebugDrawUtils::DrawArrowVector(PDI, WorldPlaneLocation, WorldPlaneLocation + WorldPlaneNormal * Scale, TEXT("Contact Normal"), NormalColor, 2.0f);
+					const int32 Scale = DebugDrawSettings.ContactNormalScale * ContactLenScale;
+					FChaosVDDebugDrawUtils::DrawArrowVector(PDI, WorldPlaneLocation, WorldPlaneLocation + WorldPlaneNormal * Scale, TEXT("Contact Normal"), NormalColor, DebugDrawSettings.DepthPriority);
 				}
 
 				if (EnumHasAnyFlags(VisualizationFlags, EChaosVDCollisionVisualizationFlags::ContactPoints))
 				{
 					if (ManifoldPoint.ContactPoint.Phi < FLT_MAX)
 					{
-						FColor Color = FColor(128, 128, 0);
-						FChaosVDDebugDrawUtils::DrawPoint(PDI, WorldPlaneLocation - ManifoldPoint.ContactPoint.Phi * WorldPlaneNormal, TEXT("Contact : WorldPlaneLocation - Phi"), Color);
-
-						FChaosVDDebugDrawUtils::DrawPoint(PDI, WorldPlaneLocation - ManifoldPoint.ContactPoint.Phi * WorldPlaneNormal, TEXT("Contact : WorldPlaneLocation - Phi"), Color);
+						FChaosVDDebugDrawUtils::DrawCircle(PDI, WorldPlaneLocation - ManifoldPoint.ContactPoint.Phi * WorldPlaneNormal, DebugDrawSettings.ContactPhiCircleRadius, CircleSegments, PhiDiscColor, LinesThickness, Axes.GetUnitAxis(EAxis::Y), Axes.GetUnitAxis(EAxis::Z), TEXT("Contact : WorldPlaneLocation - Phi"), DebugDrawSettings.DepthPriority);
 					}
-					
+
 					// Manifold point
-					FChaosVDDebugDrawUtils::DrawPoint(PDI, WorldPointLocation , TEXT("Contact : Manifold Point"), DiscColor);
+					FChaosVDDebugDrawUtils::DrawCircle(PDI, WorldPointLocation, DebugDrawSettings.ContactCircleRadius, CircleSegments, DiscColor, LinesThickness, Axes.GetUnitAxis(EAxis::Y), Axes.GetUnitAxis(EAxis::Z), TEXT("Contact : Manifold Point"), DebugDrawSettings.DepthPriority);				
 				}
 
 				if (EnumHasAnyFlags(VisualizationFlags, EChaosVDCollisionVisualizationFlags::AccumulatedImpulse))
 				{
 					if (!Constraint.AccumulatedImpulse.IsNearlyZero())
 					{
-						FChaosVDDebugDrawUtils::DrawArrowVector(PDI, WorldActorTransform0.GetLocation(), WorldActorTransform0.GetLocation() + Constraint.AccumulatedImpulse, TEXT("AccumulatedImpulse"), FColor::White, 5.0f);
+						FChaosVDDebugDrawUtils::DrawArrowVector(PDI, WorldActorTransform0.GetLocation(), WorldActorTransform0.GetLocation() + Constraint.AccumulatedImpulse, TEXT("AccumulatedImpulse"), FColor::White, DebugDrawSettings.DepthPriority);
+					}
+				}
+
+				if (EnumHasAnyFlags(VisualizationFlags, EChaosVDCollisionVisualizationFlags::ContactInfo))
+				{
+					FTransform Transform;
+					Transform.SetRotation(FRotationMatrix::MakeFromZ(WorldPlaneNormal).ToQuat());
+					Transform.SetLocation(WorldPlaneLocation);
+
+					FVector BoxExtents(DebugDrawSettings.ContactCircleRadius, DebugDrawSettings.ContactCircleRadius, 0.01f);
+					FColor Color;
+
+					if (Constraint.bWasManifoldRestored)
+					{
+						Color = FColor::Blue;
+						FChaosVDDebugDrawUtils::DrawBox(PDI, BoxExtents, Color, Transform, nullptr, DebugDrawSettings.DepthPriority);
+					}
+					else if (ManifoldPoint.bWasRestored)
+					{
+						Color = FColor::Purple;
+						FChaosVDDebugDrawUtils::DrawBox(PDI, BoxExtents, Color, Transform, nullptr, DebugDrawSettings.DepthPriority);
+					}
+					else if (ManifoldPoint.bWasReplaced)
+					{
+						Color = FColor::Orange;
+						FChaosVDDebugDrawUtils::DrawBox(PDI, BoxExtents, Color, Transform, nullptr, DebugDrawSettings.DepthPriority);
+					}
+
+					if (MidPhase->bIsSleeping)
+					{
+						// This box should surround the debug draw box we made before
+						FChaosVDDebugDrawUtils::DrawBox(PDI, BoxExtents * 1.1f, FColor::Black, Transform, nullptr, DebugDrawSettings.DepthPriority);
 					}
 				}
 
 				PDI->SetHitProxy(nullptr);
+
+				if (ManifoldPoint.bIsSelectedInEditor)
+				{
+					// We don't have an easy way to show a contact is selected with debug draw
+					// but 3D box surrounding the contact is better than nothing
+					FTransform SelectionBoxTransform;
+					SelectionBoxTransform.SetRotation(FRotationMatrix::MakeFromZ(WorldPlaneNormal).ToQuat());
+					SelectionBoxTransform.SetLocation(WorldPlaneLocation);
+
+					// The Selection box should be a bit bigger than the configured circle radius for the debug draw contact
+					const float ContactSelectionBoxSize = DebugDrawSettings.ContactCircleRadius * 1.5f;
+
+					FVector SelectionBoxExtents(ContactSelectionBoxSize,ContactSelectionBoxSize,ContactSelectionBoxSize);
+					FChaosVDDebugDrawUtils::DrawBox(PDI, SelectionBoxExtents, FColor::Yellow, SelectionBoxTransform, nullptr, DebugDrawSettings.DepthPriority);
+				}
 			}
 		}
 	}
diff --git a/Engine/Plugins/Experimental/ChaosVD/Source/ChaosVD/Private/Visualizers/ChaosVDSolverCollisionDataComponentVisualizer.h b/Engine/Plugins/Experimental/ChaosVD/Source/ChaosVD/Private/Visualizers/ChaosVDSolverCollisionDataComponentVisualizer.h
--- a/Engine/Plugins/Experimental/ChaosVD/Source/ChaosVD/Private/Visualizers/ChaosVDSolverCollisionDataComponentVisualizer.h
+++ b/Engine/Plugins/Experimental/ChaosVD/Source/ChaosVD/Private/Visualizers/ChaosVDSolverCollisionDataComponentVisualizer.h
@@ -12,18 +12,25 @@ class FChaosVDScene;
 struct FChaosVDCollisionDataFinder
 {
 	TWeakPtr<FChaosVDParticlePairMidPhase> OwningMidPhase;
-	const FChaosVDConstraint* OwningConstraint = nullptr;
+	FChaosVDConstraint* OwningConstraint = nullptr;
 	int32 ContactIndex = INDEX_NONE;
+
+	void SetIsSelected(bool bNewSelected);
 };
 
-struct HChaosVDContactPointProxy : public HHitProxy
+struct HChaosVDContactPointProxy : public HComponentVisProxy
 {
 	DECLARE_HIT_PROXY()
 	
-	HChaosVDContactPointProxy(const FChaosVDCollisionDataFinder& InContactFinderData) : HHitProxy(HPP_UI), ContactFinder(InContactFinderData)
+	HChaosVDContactPointProxy(const UActorComponent* Component, const FChaosVDCollisionDataFinder& InContactFinderData) : HComponentVisProxy(Component, HPP_UI), ContactFinder(InContactFinderData)
 	{	
 	}
 
+	virtual EMouseCursor::Type GetMouseCursor() override
+	{
+		return EMouseCursor::Crosshairs;
+	}
+
 	FChaosVDCollisionDataFinder ContactFinder;
 };
 
@@ -32,12 +39,12 @@ enum class EChaosVDCollisionVisualizationFlags: uint32
 {
 	None					= 0 UMETA(Hidden),
 	ContactPoints			= 1 << 0,
-	ContactInfo				= 1 << 1 UMETA(Hidden), // Not used yet
+	ContactInfo				= 1 << 1,
 	NetPushOut				= 1 << 2,
 	NetImpulse				= 1 << 3,
 	ContactNormal			= 1 << 4,
 	AccumulatedImpulse		= 1 << 5,
-	DrawOnlyActiveContacts	= 1 << 6,
+	DrawInactiveContacts	= 1 << 6,
 	EnableDrawForAllParticles	= 1 << 7,
 };
 ENUM_CLASS_FLAGS(EChaosVDCollisionVisualizationFlags);
@@ -51,7 +58,13 @@ public:
 	virtual bool ShowWhenSelected() override;
 	virtual void DrawVisualization(const UActorComponent* Component, const FSceneView* View, FPrimitiveDrawInterface* PDI) override;
 
+	virtual bool VisProxyHandleClick(FEditorViewportClient* InViewportClient, HComponentVisProxy* VisProxy, const FViewportClick& Click) override;
+
 protected:
+	
+	void ClearCurrentSelection();
+
+	void DrawMidPhaseData(const UActorComponent* Component, const TSharedPtr<FChaosVDParticlePairMidPhase>& MidPhase, const FChaosVDVisualizationContext& VisualizationContext, const FSceneView* View, FPrimitiveDrawInterface* PDI);
 
-	void DrawnMidPhaseData(const TSharedPtr<FChaosVDParticlePairMidPhase>& MidPhase, const FChaosVDVisualizationContext& VisualizationContext, const FSceneView* View, FPrimitiveDrawInterface* PDI);
+	FChaosVDCollisionDataFinder CurrentSelectedContactData;
 };
diff --git a/Engine/Plugins/Experimental/ChaosVD/Source/ChaosVD/Public/ChaosVDEditorSettings.h b/Engine/Plugins/Experimental/ChaosVD/Source/ChaosVD/Public/ChaosVDEditorSettings.h
--- a/Engine/Plugins/Experimental/ChaosVD/Source/ChaosVD/Public/ChaosVDEditorSettings.h
+++ b/Engine/Plugins/Experimental/ChaosVD/Source/ChaosVD/Public/ChaosVDEditorSettings.h
@@ -44,6 +44,29 @@ enum class EChaosVDGeometryVisibilityFlags : uint8
 };
 ENUM_CLASS_FLAGS(EChaosVDGeometryVisibilityFlags)
 
+/** Structure holding the settings using to debug draw contact data on the Chaos Visual Debugger */
+USTRUCT()
+struct FChaosVDContactDebugDrawSettings
+{
+	GENERATED_BODY()
+
+	/** The depth priority used for while drawing contact data. Can be World or Foreground (with this one the shapes will be drawn on top of the geometry and be always visible) */
+	UPROPERTY(EditAnywhere, Category=DebugDraw)
+	TEnumAsByte<ESceneDepthPriorityGroup> DepthPriority = ESceneDepthPriorityGroup::SDPG_World;
+
+	/** The radius of the debug draw circle used to represent a contact point */
+	UPROPERTY(EditAnywhere, Category=DebugDraw)
+	float ContactCircleRadius = 6.0f;
+
+	/** The scale value to be applied to the normal vector of a contact used to change its size to make it easier to see */
+	UPROPERTY(EditAnywhere, Category=DebugDraw)
+	float ContactNormalScale = 30.0f;
+
+	/** The radius of the debug draw circle used to represent the Phi value (penetration) of a contact point */
+	UPROPERTY(EditAnywhere, Category=DebugDraw)
+	float ContactPhiCircleRadius = 2.0f;
+};
+
 UCLASS(config = Engine)
 class UChaosVDEditorSettings : public UObject
 {
@@ -56,9 +79,12 @@ public:
 	UPROPERTY(EditAnywhere, Category = "Viewport Visualization Flags", meta = (Bitmask, BitmaskEnum = "/Script/ChaosVD.EChaosVDCollisionVisualizationFlags"))
 	uint32 GlobalCollisionDataVisualizationFlags = 0;
 
-	UPROPERTY(EditAnywhere, Category = "Viewport Visualization Flags", meta = (Bitmask, BitmaskEnum = "/Script/ChaosVD.EChaosVDCollisionVisualizationFlags"))
+	UPROPERTY(EditAnywhere, Category = "Viewport Visualization")
 	bool bShowDebugText = false;
 
+	UPROPERTY(EditAnywhere, Category = "Viewport Visualization",  meta=(EditCondition = "GlobalCollisionDataVisualizationFlags != 0", EditConditionHides))
+	FChaosVDContactDebugDrawSettings ContactDebugDrawSettings;
+
 	UPROPERTY(EditAnywhere, Category = "Viewport Tracking")
 	EChaosVDActorTrackingTarget TrackingTarget;
 
diff --git a/Engine/Plugins/Experimental/ChaosVD/Source/ChaosVD/Public/Components/ChaosVDSolverCollisionDataComponent.h b/Engine/Plugins/Experimental/ChaosVD/Source/ChaosVD/Public/Components/ChaosVDSolverCollisionDataComponent.h
--- a/Engine/Plugins/Experimental/ChaosVD/Source/ChaosVD/Public/Components/ChaosVDSolverCollisionDataComponent.h
+++ b/Engine/Plugins/Experimental/ChaosVD/Source/ChaosVD/Public/Components/ChaosVDSolverCollisionDataComponent.h
@@ -22,9 +22,9 @@ struct FChaosVDConstraint;
 typedef TMap<int32, TArray<TSharedPtr<FChaosVDParticlePairMidPhase>>> FChaosVDMidPhaseByParticleMap;
 typedef TMap<int32, TArray<FChaosVDConstraint*>> FChaosVDConstraintByParticleMap;
 
-enum class EChaosVDGetCollisionDataOptions
+enum class EChaosVDCollisionParticlePairSlot
 {
-	Owner,
+	Primary,
 	Secondary,
 	Any
 };
@@ -40,8 +40,8 @@ public:
 	void UpdateCollisionData(const TArray<TSharedPtr<FChaosVDParticlePairMidPhase>>& InMidPhaseData);
 
 	const TArray<TSharedPtr<FChaosVDParticlePairMidPhase>>& GetMidPhases() const { return AllMidPhases; }
-	const TArray<TSharedPtr<FChaosVDParticlePairMidPhase>>* GetMidPhasesForParticle(int32 ParticleID, EChaosVDGetCollisionDataOptions Options) const;
-	const TArray<FChaosVDConstraint*>* GetConstraintsForParticle(int32 ParticleID, EChaosVDGetCollisionDataOptions Options) const;
+	const TArray<TSharedPtr<FChaosVDParticlePairMidPhase>>* GetMidPhasesForParticle(int32 ParticleID, EChaosVDCollisionParticlePairSlot Options) const;
+	const TArray<FChaosVDConstraint*>* GetConstraintsForParticle(int32 ParticleID, EChaosVDCollisionParticlePairSlot Options) const;
 
 	void DrawVisualization(const FSceneView* View, FPrimitiveDrawInterface* PDI);
 
@@ -53,7 +53,7 @@ protected:
 	void AddCollisionDataToParticleIDMap(MapType& MapToUpdate, const CollisionDataType& MidPhaseData, int32 ParticleID);
 
 	template<typename MapType, typename CollisionDataType>
-	const TArray<CollisionDataType>* GetCollisionDataFromMap(const MapType& Map0ToQuery, const MapType& Map1ToQuery, int32 ParticleID, EChaosVDGetCollisionDataOptions Options) const;
+	const TArray<CollisionDataType>* GetCollisionDataFromMap(const MapType& Map0ToQuery, const MapType& Map1ToQuery, int32 ParticleID, EChaosVDCollisionParticlePairSlot Options) const;
 
 	TArray<TSharedPtr<FChaosVDParticlePairMidPhase>> AllMidPhases;
 	FChaosVDMidPhaseByParticleMap MidPhasesByParticleID0;
@@ -78,21 +78,21 @@ void UChaosVDSolverCollisionDataComponent::AddCollisionDataToParticleIDMap(MapTy
 }
 
 template <typename MapType, typename CollisionDataType>
-const TArray<CollisionDataType>* UChaosVDSolverCollisionDataComponent::GetCollisionDataFromMap(const MapType& Map0ToQuery, const MapType& Map1ToQuery, int32 ParticleID, EChaosVDGetCollisionDataOptions Options) const
+const TArray<CollisionDataType>* UChaosVDSolverCollisionDataComponent::GetCollisionDataFromMap(const MapType& Map0ToQuery, const MapType& Map1ToQuery, int32 ParticleID, EChaosVDCollisionParticlePairSlot Options) const
 {
 	switch (Options)
 	{
-	case EChaosVDGetCollisionDataOptions::Owner:
+	case EChaosVDCollisionParticlePairSlot::Primary:
 		{
 			return Map0ToQuery.Find(ParticleID);
 			break;
 		}
-	case EChaosVDGetCollisionDataOptions::Secondary:
+	case EChaosVDCollisionParticlePairSlot::Secondary:
 		{
 			return Map1ToQuery.Find(ParticleID);
 			break;
 		}
-	case EChaosVDGetCollisionDataOptions::Any:
+	case EChaosVDCollisionParticlePairSlot::Any:
 		{
 			const TArray<CollisionDataType>* FoundMidPhasesContainer = Map0ToQuery.Find(ParticleID);
 			return FoundMidPhasesContainer ? FoundMidPhasesContainer : Map1ToQuery.Find(ParticleID);
diff --git a/Engine/Plugins/Experimental/ChaosVD/Source/ChaosVD/Public/Visualizers/ChaosVDDebugDrawUtils.h b/Engine/Plugins/Experimental/ChaosVD/Source/ChaosVD/Public/Visualizers/ChaosVDDebugDrawUtils.h
--- a/Engine/Plugins/Experimental/ChaosVD/Source/ChaosVD/Public/Visualizers/ChaosVDDebugDrawUtils.h
+++ b/Engine/Plugins/Experimental/ChaosVD/Source/ChaosVD/Public/Visualizers/ChaosVDDebugDrawUtils.h
@@ -8,9 +8,13 @@
 class FChaosVDDebugDrawUtils
 {
 public:
-	static void DrawArrowVector(FPrimitiveDrawInterface* PDI, const FVector& StartLocation, const FVector& EndLocation, FStringView DebugText, const FColor& Color, float ArrowSize = 10.0f);
-	static void DrawPoint(FPrimitiveDrawInterface* PDI, const FVector& Location, FStringView DebugText, const FColor& Color, float Size = 10.0f);
+	//TODO: Create a common args struct that can be passed as parameter to all debug draw methods
+
+	static void DrawArrowVector(FPrimitiveDrawInterface* PDI, const FVector& StartLocation, const FVector& EndLocation, FStringView DebugText, const FColor& Color, ESceneDepthPriorityGroup DepthPriority = SDPG_World);
+	static void DrawPoint(FPrimitiveDrawInterface* PDI, const FVector& Location, FStringView DebugText, const FColor& Color, float Size = 10.0f, ESceneDepthPriorityGroup DepthPriority = SDPG_World);
 	static void DrawText(FStringView StringToDraw, const FVector& Location, const FColor& Color);
+	static void DrawCircle(FPrimitiveDrawInterface* PDI, const FVector& Origin, float Radius, int32 Segments, const FColor& Color, float Thickness, const FVector& XAxis, const FVector& YAxis, FStringView DebugText, ESceneDepthPriorityGroup DepthPriority = SDPG_World);
+	static void DrawBox(FPrimitiveDrawInterface* PDI, const FVector& InExtents, const FColor& InColor, const FTransform& InTransform, FStringView DebugText, ESceneDepthPriorityGroup DepthPriority = SDPG_World);
 
 	static void DrawCanvas(FViewport& InViewport, FSceneView& View, FCanvas& Canvas);
 
diff --git a/Engine/Source/Runtime/Experimental/ChaosVisualDebugger/Public/DataWrappers/ChaosVDCollisionDataWrappers.h b/Engine/Source/Runtime/Experimental/ChaosVisualDebugger/Public/DataWrappers/ChaosVDCollisionDataWrappers.h
--- a/Engine/Source/Runtime/Experimental/ChaosVisualDebugger/Public/DataWrappers/ChaosVDCollisionDataWrappers.h
+++ b/Engine/Source/Runtime/Experimental/ChaosVisualDebugger/Public/DataWrappers/ChaosVDCollisionDataWrappers.h
@@ -130,6 +130,8 @@ struct CHAOSVDRUNTIME_API FChaosVDManifoldPoint
 	UPROPERTY()
 	TArray<FVector> ShapeContactPoints;
 
+	bool bIsSelectedInEditor = false;
+
 	bool Serialize(FArchive& Ar);
 };
 
