[ChaosVD] Added colors support for particles

- Fixed an issue where shape data was not being applied to heightfields
- Fixed an issue where if geometry was not ready, some task where dropped instead of being performed when the geometry was eventually ready.
- Now GeometryDataComponets have a ptr to the specific implicit they represent, not just the root one.
- Added a two new settings to configure the colors for each state and shape type (it uses the same defaults as the runtime Debug Draw).

diff --git a/Engine/Plugins/Experimental/ChaosVD/Config/Engine.ini b/Engine/Plugins/Experimental/ChaosVD/Config/Engine.ini
--- a/Engine/Plugins/Experimental/ChaosVD/Config/Engine.ini
+++ b/Engine/Plugins/Experimental/ChaosVD/Config/Engine.ini
@@ -1,4 +1,4 @@
 ﻿[/Script/ChaosVD.ChaosVDEditorSettings]
 QueryOnlyMeshesMaterial=/ChaosVD/Materials/ChaosVDQueryOnlyMesh.ChaosVDQueryOnlyMesh
-SimMeshesMaterial=/ChaosVD/Materials/ChasoVDSimMesh.ChasoVDSimMesh
+SimOnlyMeshesMaterial=/ChaosVD/Materials/ChasoVDSimMesh.ChasoVDSimMesh
 SkySphereActorClass=/ChaosVD/BP_CVD_SkySphere.BP_CVD_SkySphere_C
\ No newline at end of file
diff --git a/Engine/Plugins/Experimental/ChaosVD/Source/ChaosVD/Private/Actors/ChaosVDSolverInfoActor.cpp b/Engine/Plugins/Experimental/ChaosVD/Source/ChaosVD/Private/Actors/ChaosVDSolverInfoActor.cpp
--- a/Engine/Plugins/Experimental/ChaosVD/Source/ChaosVD/Private/Actors/ChaosVDSolverInfoActor.cpp
+++ b/Engine/Plugins/Experimental/ChaosVD/Source/ChaosVD/Private/Actors/ChaosVDSolverInfoActor.cpp
@@ -42,7 +42,21 @@ void AChaosVDSolverInfoActor::HandleVisibilitySettingsUpdated()
 {
 	for (const TPair<int32, AChaosVDParticleActor*>& ParticleWithIDPair : SolverParticlesByID)
 	{
-		ParticleWithIDPair.Value->UpdateGeometryComponentsVisibility();
+		if (AChaosVDParticleActor* ParticleActor = ParticleWithIDPair.Value)
+		{
+			ParticleActor->UpdateGeometryComponentsVisibility();
+		}
+	}
+}
+
+void AChaosVDSolverInfoActor::HandleColorsSettingsUpdated()
+{
+	for (const TPair<int32, AChaosVDParticleActor*>& ParticleWithIDPair : SolverParticlesByID)
+	{
+		if (AChaosVDParticleActor* ParticleActor = ParticleWithIDPair.Value)
+		{
+			ParticleActor->UpdateGeometryColors();
+		}	
 	}
 }
 
diff --git a/Engine/Plugins/Experimental/ChaosVD/Source/ChaosVD/Private/ChaosVDEditorSettings.cpp b/Engine/Plugins/Experimental/ChaosVD/Source/ChaosVD/Private/ChaosVDEditorSettings.cpp
--- a/Engine/Plugins/Experimental/ChaosVD/Source/ChaosVD/Private/ChaosVDEditorSettings.cpp
+++ b/Engine/Plugins/Experimental/ChaosVD/Source/ChaosVD/Private/ChaosVDEditorSettings.cpp
@@ -1,6 +1,54 @@
 // Copyright Epic Games, Inc. All Rights Reserved.
 
 #include "ChaosVDEditorSettings.h"
+#include "DataWrappers/ChaosVDParticleDataWrapper.h"
+#include "Chaos/ImplicitObjectType.h"
+
+FColor FChaosDebugDrawColorsByState::GetColorFromState(EChaosVDObjectStateType State) const
+{
+	switch (State)
+	{
+		case EChaosVDObjectStateType::Sleeping:
+			return SleepingColor;
+		case EChaosVDObjectStateType::Kinematic:
+			return KinematicColor;
+		case EChaosVDObjectStateType::Static:
+			return StaticColor;
+		case EChaosVDObjectStateType::Dynamic:
+			return DynamicColor;
+		default:
+			return FColor::Purple;
+	}
+}
+
+FColor FChaosDebugDrawColorsByShapeType::GetColorFromShapeType(Chaos::EImplicitObjectType ShapeType) const
+{
+	switch(ShapeType)
+	{
+		case Chaos::ImplicitObjectType::Sphere:
+			return SimpleTypeColor;
+		case Chaos::ImplicitObjectType::Box:
+			return SimpleTypeColor;
+		case Chaos::ImplicitObjectType::Plane:
+			return SimpleTypeColor;
+		case Chaos::ImplicitObjectType::Capsule:
+			return SimpleTypeColor;
+		case Chaos::ImplicitObjectType::TaperedCylinder:
+			return SimpleTypeColor;
+		case Chaos::ImplicitObjectType::Cylinder:
+			return SimpleTypeColor;
+		case Chaos::ImplicitObjectType::Convex:
+			return ConvexColor;
+		case Chaos::ImplicitObjectType::HeightField:
+			return HeightFieldColor;
+		case Chaos::ImplicitObjectType::TriangleMesh:
+			return TriangleMeshColor;
+		case Chaos::ImplicitObjectType::LevelSet:
+			return LevelSetColor;			
+		default:
+			return FColor::Purple; 
+	}
+}
 
 void UChaosVDEditorSettings::PostEditChangeProperty(FPropertyChangedEvent& PropertyChangedEvent)
 {
@@ -11,4 +59,10 @@ void UChaosVDEditorSettings::PostEditChangeProperty(FPropertyChangedEvent& Prope
 	{
 		VisibilitySettingsChangedDelegate.Broadcast(this);
 	}
+	else if (PropertyName == GET_MEMBER_NAME_CHECKED(UChaosVDEditorSettings, ParticleColorMode)
+			|| PropertyName == GET_MEMBER_NAME_CHECKED(UChaosVDEditorSettings, ColorsByParticleState)
+			|| PropertyName == GET_MEMBER_NAME_CHECKED(UChaosVDEditorSettings, ColorsByShapeType))
+	{
+		ColorsSettingsChangedDelegate.Broadcast(this);
+	}
 }
diff --git a/Engine/Plugins/Experimental/ChaosVD/Source/ChaosVD/Private/ChaosVDGeometryBuilder.cpp b/Engine/Plugins/Experimental/ChaosVD/Source/ChaosVD/Private/ChaosVDGeometryBuilder.cpp
--- a/Engine/Plugins/Experimental/ChaosVD/Source/ChaosVD/Private/ChaosVDGeometryBuilder.cpp
+++ b/Engine/Plugins/Experimental/ChaosVD/Source/ChaosVD/Private/ChaosVDGeometryBuilder.cpp
@@ -207,7 +207,7 @@ void FChaosVDGeometryBuilder::ApplyMeshToComponentFromKey(TWeakObjectPtr<UMeshCo
 	if (IChaosVDGeometryDataComponent* DataComponent = Cast<IChaosVDGeometryDataComponent>(MeshComponent.Get()))
 	{
 		DataComponent->SetIsMeshReady(true);
-		DataComponent->OnMeshReady()->ExecuteIfBound(*DataComponent);
+		DataComponent->OnMeshReady()->Broadcast(*DataComponent);
 	}
 }
 
diff --git a/Engine/Plugins/Experimental/ChaosVD/Source/ChaosVD/Private/ChaosVDGeometryDataComponent.cpp b/Engine/Plugins/Experimental/ChaosVD/Source/ChaosVD/Private/ChaosVDGeometryDataComponent.cpp
--- a/Engine/Plugins/Experimental/ChaosVD/Source/ChaosVD/Private/ChaosVDGeometryDataComponent.cpp
+++ b/Engine/Plugins/Experimental/ChaosVD/Source/ChaosVD/Private/ChaosVDGeometryDataComponent.cpp
@@ -5,9 +5,12 @@
 #include "ChaosVDEditorSettings.h"
 #include "ChaosVDGeometryBuilder.h"
 #include "ChaosVDModule.h"
+#include "ChaosVDParticleActor.h"
 #include "Components/MeshComponent.h"
 #include "MaterialDomain.h"
 #include "Materials/Material.h"
+#include "Materials/MaterialInterface.h"
+#include "Materials/MaterialInstanceDynamic.h"
 
 void FChaosVDGeometryDataComponentBase::UpdateVisibility_Internal(const FChaosVDShapeCollisionData& InCollisionData, UMeshComponent* MeshComponent)
 {
@@ -20,36 +23,30 @@ void FChaosVDGeometryDataComponentBase::UpdateVisibility_Internal(const FChaosVD
 	{
 		if (const UChaosVDEditorSettings* EditorSettings = GetDefault<UChaosVDEditorSettings>())
 		{
-			if (UMaterialInterface* QueryOnlyMaterial = EditorSettings->QueryOnlyMeshesMaterial.Get())
-			{
-				const bool bIsQueryOnly = InCollisionData.bQueryCollision && !InCollisionData.bSimCollision;
-
-				if (bIsQueryOnly)
-				{
-					MeshComponent->SetMaterial(0, QueryOnlyMaterial);
-				}
-				else
-				{
-					MeshComponent->SetMaterial(0, UMaterial::GetDefaultMaterial(EMaterialDomain::MD_Surface));
-				}
-			}
-			else
-			{
-				ensure(false);
-				UE_LOG(LogChaosVDEditor, Error, TEXT("[%s] Failed to get Query Only material for, applying the default mesh to all geometry"), ANSI_TO_TCHAR(__FUNCTION__));
-				
-				MeshComponent->SetMaterial(0, UMaterial::GetDefaultMaterial(EMaterialDomain::MD_Surface));
-			}
+			const EChaosVDGeometryVisibilityFlags CurrentVisibilityFlags = static_cast<EChaosVDGeometryVisibilityFlags>(EditorSettings->GeometryVisibilityFlags);
 
 			bool bShouldGeometryBeVisible = false;
 
-			// Complex vs Simple takes priority although this is subject to change
-			const bool bShouldBeVisibleIfComplex = InCollisionData.bIsComplex && EnumHasAnyFlags(static_cast<EChaosVDGeometryVisibilityFlags>(EditorSettings->GeometryVisibilityFlags), EChaosVDGeometryVisibilityFlags::Complex);
-			const bool bShouldBeVisibleIfSimple = !InCollisionData.bIsComplex && EnumHasAnyFlags(static_cast<EChaosVDGeometryVisibilityFlags>(EditorSettings->GeometryVisibilityFlags), EChaosVDGeometryVisibilityFlags::Simple);
-			if (bShouldBeVisibleIfComplex || bShouldBeVisibleIfSimple)
+			// TODO: Re-visit the way we determine visibility of the meshes.
+			// Now that the options have grown and they will continue to do so, these checks are becoming hard to read and extend
+
+			const bool bIsHeightfield = GetImplicitObject() && Chaos::GetInnerType(GetImplicitObject()->GetType()) == Chaos::ImplicitObjectType::HeightField;
+
+			if (bIsHeightfield && EnumHasAnyFlags(CurrentVisibilityFlags, EChaosVDGeometryVisibilityFlags::ShowHeightfields))
 			{
-				bShouldGeometryBeVisible = (InCollisionData.bSimCollision && EnumHasAnyFlags(static_cast<EChaosVDGeometryVisibilityFlags>(EditorSettings->GeometryVisibilityFlags), EChaosVDGeometryVisibilityFlags::Simulated))
-				|| (InCollisionData.bQueryCollision && EnumHasAnyFlags(static_cast<EChaosVDGeometryVisibilityFlags>(EditorSettings->GeometryVisibilityFlags), EChaosVDGeometryVisibilityFlags::Query));
+				bShouldGeometryBeVisible = true;
+			}
+			else
+			{
+				// Complex vs Simple takes priority although this is subject to change
+				const bool bShouldBeVisibleIfComplex = InCollisionData.bIsComplex && EnumHasAnyFlags(CurrentVisibilityFlags, EChaosVDGeometryVisibilityFlags::Complex);
+				const bool bShouldBeVisibleIfSimple = !InCollisionData.bIsComplex && EnumHasAnyFlags(CurrentVisibilityFlags, EChaosVDGeometryVisibilityFlags::Simple);
+			
+				if (bShouldBeVisibleIfComplex || bShouldBeVisibleIfSimple)
+				{
+					bShouldGeometryBeVisible = (InCollisionData.bSimCollision && EnumHasAnyFlags(CurrentVisibilityFlags, EChaosVDGeometryVisibilityFlags::Simulated))
+					|| (InCollisionData.bQueryCollision && EnumHasAnyFlags(CurrentVisibilityFlags, EChaosVDGeometryVisibilityFlags::Query));
+				}
 			}
 
 			MeshComponent->SetVisibility(bShouldGeometryBeVisible);
@@ -57,7 +54,7 @@ void FChaosVDGeometryDataComponentBase::UpdateVisibility_Internal(const FChaosVD
 	}
 }
 
-void FChaosVDGeometryDataComponentBase::UpdateDataFromShapeArray_Internal(const TArray<FChaosVDShapeCollisionData>& InShapeArray, FChaosVDShapeCollisionData& CollisionDataToUpdate)
+void FChaosVDGeometryDataComponentBase::UpdateDataFromShapeArray_Internal(const TArray<FChaosVDShapeCollisionData>& InShapeArray, FChaosVDShapeCollisionData& CollisionDataToUpdate, UMeshComponent* MeshComponent)
 {
 	if (!ensureMsgf(RootImplicitObject.IsValid(), TEXT("Tried to Update Collision Data without a valid Implicit Object")))
 	{
@@ -66,41 +63,160 @@ void FChaosVDGeometryDataComponentBase::UpdateDataFromShapeArray_Internal(const
 
 	RootImplicitObject.GetReference()->VisitObjects([this, &CollisionDataToUpdate, &InShapeArray] (const Chaos::FImplicitObject* ImplicitA, const Chaos::FRigidTransform3& RelativeTransformA, const int32 RootObjectIndexA, const int32 ObjectIndex, const int32 LeafObjectIndexA)
 	{
-		if (!InShapeArray.IsValidIndex(RootObjectIndexA))
+		if (!InShapeArray.IsValidIndex(LeafObjectIndexA))
 		{
 			return true;
 		}
 
-		if (GetCachedGeometryHash(ImplicitA) == GeometryID)
+		if (ImplicitA == ImplicitObject)
 		{
-			CollisionDataToUpdate = InShapeArray[RootObjectIndexA];
+			CollisionDataToUpdate = InShapeArray[LeafObjectIndexA];
 			CollisionDataToUpdate.bIsComplex = FChaosVDGeometryBuilder::DoesImplicitContainType(ImplicitA, Chaos::ImplicitObjectType::HeightField) || FChaosVDGeometryBuilder::DoesImplicitContainType(ImplicitA, Chaos::ImplicitObjectType::TriangleMesh);
 			CollisionDataToUpdate.bIsValid = true;
-
-			return false;
 		}
 
 		return true;
 	});
+
+	// If our collision data was successfully updated and we have a valid mesh component, set the correct material
+	if (CollisionDataToUpdate.bIsValid && MeshComponent)
+	{
+		if (const UChaosVDEditorSettings* EditorSettings = GetDefault<UChaosVDEditorSettings>())
+		{
+			const bool bIsQueryOnly = CollisionDataToUpdate.bQueryCollision && !CollisionDataToUpdate.bSimCollision;
+
+			UMaterialInterface* MaterialToApply = bIsQueryOnly ? GetCachedMaterialInstance(EChaosVDMaterialType::QueryOnlyMaterial) : GetCachedMaterialInstance(EChaosVDMaterialType::SimOnlyMaterial);
+
+			if (MaterialToApply)
+			{
+				MeshComponent->SetMaterial(0, MaterialToApply);
+			}
+			else
+			{
+				ensure(false);
+				UE_LOG(LogChaosVDEditor, Error, TEXT("[%s] Failed to get Query Only material for, applying the default mesh to all geometry"), ANSI_TO_TCHAR(__FUNCTION__));
+				
+				MeshComponent->SetMaterial(0, UMaterial::GetDefaultMaterial(EMaterialDomain::MD_Surface));
+			}
+		}
+	}
+}
+
+void FChaosVDGeometryDataComponentBase::UpdateColors_Internal(UMeshComponent* MeshComponent)
+{
+	if (!MeshComponent)
+	{
+		return;
+	}
+
+	//TODO: Remove this direct dependency to AChaosVDParticleACtor
+	AChaosVDParticleActor* ParticleActor = Cast<AChaosVDParticleActor>(MeshComponent->GetOwner());
+	if (!ParticleActor)
+	{
+		return;
+	}
+
+	const FChaosVDParticleDataWrapper* ParticleData = ParticleActor->GetParticleData();
+	if (!ParticleData)
+	{
+		return;
+	}
+
+	const UChaosVDEditorSettings* EditorSettings = GetDefault<UChaosVDEditorSettings>();
+	if (!EditorSettings)
+	{
+		return;
+	}
+
+	constexpr FLinearColor DefaultColor(0.088542f, 0.088542f, 0.088542f);
+	FLinearColor ColorToApply = DefaultColor;
+
+	switch(EditorSettings->ParticleColorMode)
+	{
+		case EChaosVDParticleDebugColorMode::ShapeType:
+			{
+				ColorToApply = GetImplicitObject() ? EditorSettings->ColorsByShapeType.GetColorFromShapeType(Chaos::GetInnerType(GetImplicitObject()->GetType())) : DefaultColor;
+				break;
+			}
+		case EChaosVDParticleDebugColorMode::State:
+			{
+				if (ParticleData->Type == EChaosVDParticleType::Static)
+				{
+					ColorToApply = EditorSettings->ColorsByParticleState.GetColorFromState(EChaosVDObjectStateType::Static);
+				}
+				else
+				{
+					ColorToApply = EditorSettings->ColorsByParticleState.GetColorFromState(ParticleData->ParticleDynamicsMisc.MObjectState);
+				}
+				break;
+			}
+
+		case EChaosVDParticleDebugColorMode::None:
+		default:
+			// Nothing to do here. Color to apply is already set to the default
+			break;
+	}
+
+	if (CurrentGeometryColor == ColorToApply)
+	{
+		return;
+	}
+	
+	if (UMaterialInstanceDynamic* DynamicMaterial = Cast<UMaterialInstanceDynamic>(MeshComponent->GetMaterial(0)))
+	{
+		DynamicMaterial->SetVectorParameterValue(TEXT("BaseColor"), ColorToApply);
+		CurrentGeometryColor = ColorToApply;
+	}
 }
 
+void FChaosVDGeometryDataComponentBase::SetImplicitObject_Internal(const Chaos::FImplicitObject* InImplicitObject)
+{
+	ImplicitObject = InImplicitObject;
+}
 
-uint32 FChaosVDGeometryDataComponentBase::GetCachedGeometryHash(const Chaos::FImplicitObject* InImplicitObject)
+UMaterialInstanceDynamic* FChaosVDGeometryDataComponentBase::GetCachedMaterialInstance(EChaosVDMaterialType Type)
 {
-	if (!ensure(InImplicitObject))
+	const UChaosVDEditorSettings* EditorSettings = GetDefault<UChaosVDEditorSettings>();
+	if (!EditorSettings)
 	{
-		return 0;
+		return nullptr;
 	}
 
-	if (uint32* FoundCachedHashPtr = CachedGeometryHashes.Find(InImplicitObject))
+	if (const TStrongObjectPtr<UMaterialInstanceDynamic>* MaterialInstance = MaterialInstancesByID.Find(Type))
 	{
-		return *FoundCachedHashPtr;
+		return MaterialInstance->Get();
 	}
 	else
 	{
-		uint32 GeometryHash = InImplicitObject->GetTypeHash();
-		CachedGeometryHashes.Add(InImplicitObject, GeometryHash);
+		UMaterialInterface* MaterialToCreate = nullptr;
+		switch(Type)
+		{
+			case EChaosVDMaterialType::QueryOnlyMaterial:
+				{
+					MaterialToCreate = EditorSettings->QueryOnlyMeshesMaterial.Get();
+					break;
+				}
+			case EChaosVDMaterialType::SimOnlyMaterial:
+				{
+					MaterialToCreate = EditorSettings->SimOnlyMeshesMaterial.Get();
+					break;
+				}
+		}
+		
+		if (MaterialToCreate)
+		{
+			UMaterialInstanceDynamic* DynamicMaterial = UMaterialInstanceDynamic::Create(MaterialToCreate, nullptr);
+			MaterialInstancesByID.Add(Type, TStrongObjectPtr<UMaterialInstanceDynamic>(DynamicMaterial));
 
-		return GeometryHash;
+			return DynamicMaterial;
+		}
 	}
+
+	return nullptr;
+}
+
+const Chaos::FImplicitObject* FChaosVDGeometryDataComponentBase::GetImplicitObject() const
+{
+	// If the root object is no longer valid, the implicit object ptr we have is probably garbage
+	return RootImplicitObject.IsValid() ? ImplicitObject : nullptr;
 }
diff --git a/Engine/Plugins/Experimental/ChaosVD/Source/ChaosVD/Private/ChaosVDParticleActor.cpp b/Engine/Plugins/Experimental/ChaosVD/Source/ChaosVD/Private/ChaosVDParticleActor.cpp
--- a/Engine/Plugins/Experimental/ChaosVD/Source/ChaosVD/Private/ChaosVDParticleActor.cpp
+++ b/Engine/Plugins/Experimental/ChaosVD/Source/ChaosVD/Private/ChaosVDParticleActor.cpp
@@ -67,6 +67,8 @@ void AChaosVDParticleActor::UpdateFromRecordedParticleData(const FChaosVDParticl
 		UpdateShapeDataComponents();
 		UpdateGeometryComponentsVisibility();
 	}
+
+	UpdateGeometryColors();
 }
 
 void AChaosVDParticleActor::UpdateGeometry(const Chaos::FConstImplicitObjectPtr& InImplicitObject, EChaosVDActorGeometryUpdateFlags OptionsFlags)
@@ -124,11 +126,12 @@ void AChaosVDParticleActor::UpdateGeometry(const Chaos::FConstImplicitObjectPtr&
 					if (IChaosVDGeometryDataComponent* DataComponent = Cast<IChaosVDGeometryDataComponent>(MeshComponent.Get()))
 					{
 						DataComponent->SetRootImplicitObject(InImplicitObject);
+						DataComponent->UpdateDataFromShapeArray(ParticleDataViewer.CollisionDataPerShape);
 					}
 				}
 
-				UpdateShapeDataComponents();
 				UpdateGeometryComponentsVisibility();
+				UpdateGeometryColors();
 
 				bIsGeometryDataGenerationStarted = true;
 			}
@@ -274,19 +277,14 @@ const TArray<TSharedPtr<FChaosVDParticlePairMidPhase>>* AChaosVDParticleActor::G
 	return nullptr;
 }
 
-void AChaosVDParticleActor::UpdateShapeDataComponents()
+void AChaosVDParticleActor::PerformTaskOnGeometryComponents(TFunction<void(IChaosVDGeometryDataComponent& InDataComponent)> TaskToPerform)
 {
-	for (TWeakObjectPtr<UMeshComponent> MeshComponent : MeshComponents)
+	if (!ensure(TaskToPerform))
 	{
-		if (IChaosVDGeometryDataComponent* DataComponent = Cast<IChaosVDGeometryDataComponent>(MeshComponent.Get()))
-		{
-			DataComponent->UpdateDataFromShapeArray(ParticleDataViewer.CollisionDataPerShape);
-		}
-	}			
-}
+		UE_LOG(LogChaosVDEditor, Error, TEXT("[%s] Called with an invalid task callback..."), ANSI_TO_TCHAR(__FUNCTION__))
+		return;
+	}
 
-void AChaosVDParticleActor::UpdateGeometryComponentsVisibility()
-{
 	for (TWeakObjectPtr<UMeshComponent> MeshComponent : MeshComponents)
 	{
 		if (IChaosVDGeometryDataComponent* DataComponent = Cast<IChaosVDGeometryDataComponent>(MeshComponent.Get()))
@@ -294,19 +292,43 @@ void AChaosVDParticleActor::UpdateGeometryComponentsVisibility()
 			// We need wait until we have a valid mesh
 			if (DataComponent->IsMeshReady())
 			{
-				DataComponent->UpdateVisibility();
+				TaskToPerform(*DataComponent);
 			}
-			else if (!DataComponent->OnMeshReady()->IsBound())
+			else if (FChaosVDMeshReadyDelegate* MeshReadyDelegate = DataComponent->OnMeshReady())
 			{
-				DataComponent->OnMeshReady()->BindWeakLambda(this, [](IChaosVDGeometryDataComponent& GeometryDataComponent)
+				// TODO: Guard against this being multiple times for a specific task.
+				// These tasks should probably be implemented on the components themselves
+				// and they decide if they need to wait and how to do so.
+				MeshReadyDelegate->AddWeakLambda(MeshComponent.Get(), [TaskToPerform](IChaosVDGeometryDataComponent& GeometryDataComponent)
 				{
-					GeometryDataComponent.UpdateVisibility();
+					TaskToPerform(GeometryDataComponent);
 				});
 			}
 		}
 	}
 }
 
+void AChaosVDParticleActor::UpdateShapeDataComponents()
+{
+	for (TWeakObjectPtr<UMeshComponent> MeshComponent : MeshComponents)
+	{
+		if (IChaosVDGeometryDataComponent* DataComponent = Cast<IChaosVDGeometryDataComponent>(MeshComponent.Get()))
+		{
+			DataComponent->UpdateDataFromShapeArray(ParticleDataViewer.CollisionDataPerShape);
+		}
+	}
+}
+
+void AChaosVDParticleActor::UpdateGeometryComponentsVisibility()
+{
+	PerformTaskOnGeometryComponents([](IChaosVDGeometryDataComponent& GeometryDataComponent){ GeometryDataComponent.UpdateVisibility(); });
+}
+
+void AChaosVDParticleActor::UpdateGeometryColors()
+{
+	PerformTaskOnGeometryComponents([](IChaosVDGeometryDataComponent& GeometryDataComponent){ GeometryDataComponent.UpdateColors(); });
+}
+
 void AChaosVDParticleActor::SetIsActive(bool bNewActive)
 {
 	if (bIsActive != bNewActive)
diff --git a/Engine/Plugins/Experimental/ChaosVD/Source/ChaosVD/Private/ChaosVDScene.cpp b/Engine/Plugins/Experimental/ChaosVD/Source/ChaosVD/Private/ChaosVDScene.cpp
--- a/Engine/Plugins/Experimental/ChaosVD/Source/ChaosVD/Private/ChaosVDScene.cpp
+++ b/Engine/Plugins/Experimental/ChaosVD/Source/ChaosVD/Private/ChaosVDScene.cpp
@@ -59,8 +59,10 @@ void FChaosVDScene::Initialize()
 		// TODO: Do an async load instead, and prepare a loading screen or notification popup
 		// Jira for tracking UE-191639
 		StreamableManager->RequestSyncLoad(Settings->QueryOnlyMeshesMaterial.ToSoftObjectPath());
+		StreamableManager->RequestSyncLoad(Settings->SimOnlyMeshesMaterial.ToSoftObjectPath());
 		
 		Settings->OnVisibilitySettingsChanged().AddRaw(this, &FChaosVDScene::HandleVisibilitySettingsChanged);
+		Settings->OnColorSettingsChanged().AddRaw(this, &FChaosVDScene::HandleColorSettingsChanged);
 	}
 
 	bIsInitialized = true;
@@ -76,6 +78,7 @@ void FChaosVDScene::DeInitialize()
 	if (UChaosVDEditorSettings* Settings = GetMutableDefault<UChaosVDEditorSettings>())
 	{
 		Settings->OnVisibilitySettingsChanged().RemoveAll(this);
+		Settings->OnColorSettingsChanged().RemoveAll(this);
 	}
 
 	DeInitializeSelectionSets();
@@ -486,6 +489,17 @@ void FChaosVDScene::HandleVisibilitySettingsChanged(UChaosVDEditorSettings* Sett
 	}
 }
 
+void FChaosVDScene::HandleColorSettingsChanged(UChaosVDEditorSettings* SettingsObject)
+{
+	for (const TPair<int32, AChaosVDSolverInfoActor*>& SolverDataInfoWithID : SolverDataContainerBySolverID)
+	{
+		if (AChaosVDSolverInfoActor* SolverDataInfo = SolverDataInfoWithID.Value)
+		{
+			SolverDataInfo->HandleColorsSettingsUpdated();
+		}
+	}
+}
+
 void FChaosVDScene::InitializeSelectionSets()
 {
 	SelectionSet = NewObject<UTypedElementSelectionSet>(GetTransientPackage(), NAME_None, RF_Transactional);
diff --git a/Engine/Plugins/Experimental/ChaosVD/Source/ChaosVD/Private/Components/ChaosVDInstancedStaticMeshComponent.cpp b/Engine/Plugins/Experimental/ChaosVD/Source/ChaosVD/Private/Components/ChaosVDInstancedStaticMeshComponent.cpp
--- a/Engine/Plugins/Experimental/ChaosVD/Source/ChaosVD/Private/Components/ChaosVDInstancedStaticMeshComponent.cpp
+++ b/Engine/Plugins/Experimental/ChaosVD/Source/ChaosVD/Private/Components/ChaosVDInstancedStaticMeshComponent.cpp
@@ -9,7 +9,17 @@ void UChaosVDInstancedStaticMeshComponent::UpdateVisibility()
 
 void UChaosVDInstancedStaticMeshComponent::UpdateDataFromShapeArray(const TArray<FChaosVDShapeCollisionData>& InShapeArray)
 {
-	UpdateDataFromShapeArray_Internal(InShapeArray, CollisionData);
+	UpdateDataFromShapeArray_Internal(InShapeArray, CollisionData, this);
+}
+
+void UChaosVDInstancedStaticMeshComponent::SetImplicitObject(const Chaos::FImplicitObject* InImplicitObject)
+{
+	SetImplicitObject_Internal(InImplicitObject);
+}
+
+void UChaosVDInstancedStaticMeshComponent::UpdateColors()
+{
+	UpdateColors_Internal(this);
 }
 
 void UChaosVDInstancedStaticMeshComponent::SetRootImplicitObject(const Chaos::FConstImplicitObjectPtr& InImplicitObject)
diff --git a/Engine/Plugins/Experimental/ChaosVD/Source/ChaosVD/Private/Components/ChaosVDStaticMeshComponent.cpp b/Engine/Plugins/Experimental/ChaosVD/Source/ChaosVD/Private/Components/ChaosVDStaticMeshComponent.cpp
--- a/Engine/Plugins/Experimental/ChaosVD/Source/ChaosVD/Private/Components/ChaosVDStaticMeshComponent.cpp
+++ b/Engine/Plugins/Experimental/ChaosVD/Source/ChaosVD/Private/Components/ChaosVDStaticMeshComponent.cpp
@@ -9,7 +9,17 @@ void UChaosVDStaticMeshComponent::UpdateVisibility()
 
 void UChaosVDStaticMeshComponent::UpdateDataFromShapeArray(const TArray<FChaosVDShapeCollisionData>& InShapeArray)
 {
-	UpdateDataFromShapeArray_Internal(InShapeArray, CollisionData);
+	UpdateDataFromShapeArray_Internal(InShapeArray, CollisionData, this);
+}
+
+void UChaosVDStaticMeshComponent::SetImplicitObject(const Chaos::FImplicitObject* InImplicitObject)
+{
+	SetImplicitObject_Internal(InImplicitObject);
+}
+
+void UChaosVDStaticMeshComponent::UpdateColors()
+{
+	UpdateColors_Internal(this);
 }
 
 void UChaosVDStaticMeshComponent::SetRootImplicitObject(const Chaos::FConstImplicitObjectPtr& InImplicitObject)
diff --git a/Engine/Plugins/Experimental/ChaosVD/Source/ChaosVD/Public/Actors/ChaosVDSolverInfoActor.h b/Engine/Plugins/Experimental/ChaosVD/Source/ChaosVD/Public/Actors/ChaosVDSolverInfoActor.h
--- a/Engine/Plugins/Experimental/ChaosVD/Source/ChaosVD/Public/Actors/ChaosVDSolverInfoActor.h
+++ b/Engine/Plugins/Experimental/ChaosVD/Source/ChaosVD/Public/Actors/ChaosVDSolverInfoActor.h
@@ -37,6 +37,7 @@ public:
 	const TArray<int32>& GetSelectedParticlesIDs() const { return SelectedParticlesID; }
 
 	void HandleVisibilitySettingsUpdated();
+	void HandleColorsSettingsUpdated();
 
 	virtual void BeginDestroy() override;
 
diff --git a/Engine/Plugins/Experimental/ChaosVD/Source/ChaosVD/Public/ChaosVDEditorSettings.h b/Engine/Plugins/Experimental/ChaosVD/Source/ChaosVD/Public/ChaosVDEditorSettings.h
--- a/Engine/Plugins/Experimental/ChaosVD/Source/ChaosVD/Public/ChaosVDEditorSettings.h
+++ b/Engine/Plugins/Experimental/ChaosVD/Source/ChaosVD/Public/ChaosVDEditorSettings.h
@@ -15,7 +15,7 @@
 class UChaosVDEditorSettings;
 class UMaterial;
 
-DECLARE_MULTICAST_DELEGATE_OneParam(FChaosVDVisibilitySettingsChaged, UChaosVDEditorSettings* CVDEditorSettingsObject)
+DECLARE_MULTICAST_DELEGATE_OneParam(FChaosVDSettingChaged, UChaosVDEditorSettings* CVDEditorSettingsObject)
 
 UENUM()
 enum class EChaosVDActorTrackingMode
@@ -41,6 +41,7 @@ enum class EChaosVDGeometryVisibilityFlags : uint8
 	Simulated = 1 << 2,
 	Simple = 1 << 3,
 	Complex = 1 << 4,
+	ShowHeightfields = 1 << 5, // Selecting this will show heightfields even if complex is not selected
 };
 ENUM_CLASS_FLAGS(EChaosVDGeometryVisibilityFlags)
 
@@ -67,6 +68,68 @@ struct FChaosVDContactDebugDrawSettings
 	float ContactPhiCircleRadius = 2.0f;
 };
 
+/** Structure holding the settings using to debug draw Particles shape based on their state on the Chaos Visual Debugger */
+USTRUCT()
+struct FChaosDebugDrawColorsByState
+{
+	GENERATED_BODY()
+
+	/** Color used for dynamic particles */
+	UPROPERTY(EditAnywhere, Category=DebugDraw)
+	FColor DynamicColor = FColor(255, 255, 0);
+	
+	/** Color used for sleeping particles */
+	UPROPERTY(EditAnywhere, Category=DebugDraw)
+	FColor SleepingColor = FColor(128, 128, 128);
+
+	/** Color used for kinematic particles */
+	UPROPERTY(EditAnywhere, Category=DebugDraw)
+	FColor KinematicColor = FColor(0, 128, 255);
+
+	/** Color used for static particles */
+	UPROPERTY(EditAnywhere, Category=DebugDraw)
+	FColor StaticColor = FColor(255, 0, 0);
+
+	FColor GetColorFromState(EChaosVDObjectStateType State) const;
+};
+
+/** Structure holding the settings using to debug draw Particles shape based on their shape type on the Chaos Visual Debugger */
+USTRUCT()
+struct FChaosDebugDrawColorsByShapeType
+{
+	GENERATED_BODY()
+
+	/** Color used for Sphere, Plane, Cube, Capsule, Cylinder, tapered shapes */
+	UPROPERTY(EditAnywhere, Category=DebugDraw)
+	FColor SimpleTypeColor = FColor(0, 255, 0); 
+
+	/** Color used for convex shapes */
+	UPROPERTY(EditAnywhere, Category=DebugDraw)
+	FColor ConvexColor = FColor(0, 255, 255);
+
+	/** Color used for heightfield */
+	UPROPERTY(EditAnywhere, Category=DebugDraw)
+	FColor HeightFieldColor = FColor(0, 0, 255);
+	
+	/** Color used for triangle meshes */
+	UPROPERTY(EditAnywhere, Category=DebugDraw)
+	FColor TriangleMeshColor = FColor(255, 0, 0);
+
+	/** Color used for triangle LevelSets */
+	UPROPERTY(EditAnywhere, Category=DebugDraw)
+	FColor LevelSetColor = FColor(255, 0, 128);
+
+	FColor GetColorFromShapeType(Chaos::EImplicitObjectType ShapeType) const;
+};
+
+UENUM()
+enum class EChaosVDParticleDebugColorMode
+{
+	None,
+	State,
+	ShapeType
+};
+
 UCLASS(config = Engine)
 class UChaosVDEditorSettings : public UObject
 {
@@ -85,6 +148,15 @@ public:
 	UPROPERTY(EditAnywhere, Category = "Viewport Visualization",  meta=(EditCondition = "GlobalCollisionDataVisualizationFlags != 0", EditConditionHides))
 	FChaosVDContactDebugDrawSettings ContactDebugDrawSettings;
 
+	UPROPERTY(EditAnywhere, Category = "Viewport Visualization")
+	EChaosVDParticleDebugColorMode ParticleColorMode;
+	
+	UPROPERTY(EditAnywhere, Category = "Viewport Visualization", meta=(EditCondition = "ParticleColorMode == EChaosVDParticleDebugColorMode::ShapeType", EditConditionHides))
+	FChaosDebugDrawColorsByShapeType ColorsByShapeType;
+	
+	UPROPERTY(EditAnywhere, Category = "Viewport Visualization", meta=(EditCondition = "ParticleColorMode == EChaosVDParticleDebugColorMode::State", EditConditionHides))
+	FChaosDebugDrawColorsByState ColorsByParticleState;
+
 	UPROPERTY(EditAnywhere, Category = "Viewport Tracking")
 	EChaosVDActorTrackingTarget TrackingTarget;
 
@@ -98,21 +170,27 @@ public:
 	float ExpandViewTrackingBy = 60.0f;
 
 	UPROPERTY(EditAnywhere, Category = "Geometry Visibility", meta = (Bitmask, BitmaskEnum = "/Script/ChaosVD.EChaosVDGeometryVisibilityFlags"))
-	uint8 GeometryVisibilityFlags = static_cast<uint8>(EChaosVDGeometryVisibilityFlags::Simulated | EChaosVDGeometryVisibilityFlags::Simple);
+	uint8 GeometryVisibilityFlags = static_cast<uint8>(EChaosVDGeometryVisibilityFlags::Simulated | EChaosVDGeometryVisibilityFlags::Simple |  EChaosVDGeometryVisibilityFlags::ShowHeightfields);
 
 	UPROPERTY(Config)
 	TSoftObjectPtr<UMaterial> QueryOnlyMeshesMaterial;
 
+	UPROPERTY(Config)
+	TSoftObjectPtr<UMaterial> SimOnlyMeshesMaterial;
+
 	UPROPERTY(Config)
 	FSoftClassPath SkySphereActorClass;
 
 	virtual void PostEditChangeProperty(FPropertyChangedEvent& PropertyChangedEvent) override;
 
-	FChaosVDVisibilitySettingsChaged& OnVisibilitySettingsChanged() { return VisibilitySettingsChangedDelegate; }
+	FChaosVDSettingChaged& OnVisibilitySettingsChanged() { return VisibilitySettingsChangedDelegate; }
+
+	FChaosVDSettingChaged& OnColorSettingsChanged() { return ColorsSettingsChangedDelegate; }
 
 	TSharedPtr<FName> SelectedTrackedTransformName;
 	TSharedPtr<FName> SelectedTrackedLocationName;
 
 protected:
-	FChaosVDVisibilitySettingsChaged VisibilitySettingsChangedDelegate;
+	FChaosVDSettingChaged VisibilitySettingsChangedDelegate;
+	FChaosVDSettingChaged ColorsSettingsChangedDelegate;
 };
diff --git a/Engine/Plugins/Experimental/ChaosVD/Source/ChaosVD/Public/ChaosVDGeometryBuilder.h b/Engine/Plugins/Experimental/ChaosVD/Source/ChaosVD/Public/ChaosVDGeometryBuilder.h
--- a/Engine/Plugins/Experimental/ChaosVD/Source/ChaosVD/Public/ChaosVDGeometryBuilder.h
+++ b/Engine/Plugins/Experimental/ChaosVD/Source/ChaosVD/Public/ChaosVDGeometryBuilder.h
@@ -439,8 +439,6 @@ void FChaosVDGeometryBuilder::CreateMeshComponentsFromImplicit(const Chaos::FImp
 
 					DispatchCreateAndCacheMeshForImplicitAsync<MeshType>(GeometryKey, HeightFieldMeshGen, DesiredLODCount);
 				}
-
-				return;
 			}
 		
 			break;
@@ -451,6 +449,11 @@ void FChaosVDGeometryBuilder::CreateMeshComponentsFromImplicit(const Chaos::FImp
 
 		if (MeshComponent != nullptr)
 		{
+			if (IChaosVDGeometryDataComponent* DataComponent = Cast<IChaosVDGeometryDataComponent>(MeshComponent))
+			{
+				DataComponent->SetImplicitObject(InImplicitObject);
+			}
+
 			OutMeshComponents.Add(MeshComponent);
 		}
 }
diff --git a/Engine/Plugins/Experimental/ChaosVD/Source/ChaosVD/Public/ChaosVDGeometryDataComponent.h b/Engine/Plugins/Experimental/ChaosVD/Source/ChaosVD/Public/ChaosVDGeometryDataComponent.h
--- a/Engine/Plugins/Experimental/ChaosVD/Source/ChaosVD/Public/ChaosVDGeometryDataComponent.h
+++ b/Engine/Plugins/Experimental/ChaosVD/Source/ChaosVD/Public/ChaosVDGeometryDataComponent.h
@@ -1,25 +1,28 @@
 // Copyright Epic Games, Inc. All Rights Reserved.
 
 #pragma once
-#include "Containers/Array.h"
+
 #include "Chaos/ImplicitFwd.h"
+#include "Containers/Array.h"
+#include "Containers/Map.h"
+#include "Chaos/ImplicitObject.h"
 #include "DataWrappers/ChaosVDCollisionDataWrappers.h"
 #include "HAL/Platform.h"
 #include "UObject/Interface.h"
-
-#include "Chaos/ImplicitObject.h"
+#include "UObject/StrongObjectPtr.h"
 
 #include "ChaosVDGeometryDataComponent.generated.h"
 
 class IChaosVDGeometryDataComponent;
 class UMeshComponent;
+class UMaterialInstanceDynamic;
 
 namespace Chaos
 {
 	class FImplicitObject;
 }
 
-DECLARE_DELEGATE_OneParam(FChaosVDMeshReadyDelegate, IChaosVDGeometryDataComponent&)
+DECLARE_MULTICAST_DELEGATE_OneParam(FChaosVDMeshReadyDelegate, IChaosVDGeometryDataComponent&)
 
 UINTERFACE()
 class UChaosVDGeometryDataComponent : public UInterface
@@ -55,16 +58,29 @@ public:
 	/** Stores a shader ptr to the root implicit object of the implicit object this component represents  */
 	virtual void SetRootImplicitObject(const Chaos::FConstImplicitObjectPtr& InImplicitObject) PURE_VIRTUAL(IChaosVDGeometryDataComponent::SetRootImplicitObject);
 
+	/** Stores a shader ptr to the implicit object of the implicit object this component represents  */
+	virtual void SetImplicitObject(const Chaos::FImplicitObject* InImplicitObject) PURE_VIRTUAL(IChaosVDGeometryDataComponent::SetImplicitObject);
+
 	/** Returns a ptr to the CVD Collision Data */
 	virtual FChaosVDShapeCollisionData* GetCollisionData() PURE_VIRTUAL(IChaosVDGeometryDataComponent::GetCollisionData, return nullptr;);
 
 	/** Updates the visibility of this component based on the stored CVD data*/
 	virtual void UpdateVisibility() PURE_VIRTUAL(IChaosVDGeometryDataComponent::UpdateVisibility);
 
+	/** Updates the colors of this component based on the stored CVD data */
+	virtual void UpdateColors() PURE_VIRTUAL(IChaosVDGeometryDataComponent::UpdateColors);
+
 	/** Finds and updates the Shape data using the provided array as source*/
 	virtual void UpdateDataFromShapeArray(const TArray<FChaosVDShapeCollisionData>& InShapeArray) PURE_VIRTUAL(IChaosVDGeometryDataComponent::UpdateDataFromShapeArray);
 };
 
+UENUM()
+enum class EChaosVDMaterialType
+{
+	SimOnlyMaterial,
+	QueryOnlyMaterial
+};
+
 /** Base Implementation for a GeometryDataComponent */
 class FChaosVDGeometryDataComponentBase
 {
@@ -73,15 +89,30 @@ protected:
 
 	void UpdateVisibility_Internal(const FChaosVDShapeCollisionData& InCollisionData, UMeshComponent* MeshComponent);
 	
-	void UpdateDataFromShapeArray_Internal(const TArray<FChaosVDShapeCollisionData>& InShapeArray, FChaosVDShapeCollisionData& CollisionDataToUpdate);
+	void UpdateDataFromShapeArray_Internal(const TArray<FChaosVDShapeCollisionData>& InShapeArray, FChaosVDShapeCollisionData& CollisionDataToUpdate, UMeshComponent* MeshComponent);
+
+	void UpdateColors_Internal(UMeshComponent* MeshComponent);
+
+	void SetImplicitObject_Internal(const Chaos::FImplicitObject* InImplicitObject);
 
-	uint32 GetCachedGeometryHash(const Chaos::FImplicitObject* InImplicitObject);
+	UMaterialInstanceDynamic* GetCachedMaterialInstance(EChaosVDMaterialType Type);
+
+	const Chaos::FImplicitObject* GetImplicitObject() const;
 	
 	uint32 GeometryID = 0;
+	
 	bool bIsMeshReady = false;
 
+	FLinearColor CurrentGeometryColor = FLinearColor(ForceInitToZero);
+
 	FChaosVDMeshReadyDelegate MeshReadyDelegate;
 
 	TMap<const Chaos::FImplicitObject*, uint32> CachedGeometryHashes;
+
 	Chaos::FConstImplicitObjectPtr RootImplicitObject = nullptr;
+
+	const Chaos::FImplicitObject* ImplicitObject = nullptr;
+
+	// TODO: Remove the use of strong objt ptr
+	TMap<EChaosVDMaterialType, TStrongObjectPtr<UMaterialInstanceDynamic>> MaterialInstancesByID;
 };
diff --git a/Engine/Plugins/Experimental/ChaosVD/Source/ChaosVD/Public/ChaosVDParticleActor.h b/Engine/Plugins/Experimental/ChaosVD/Source/ChaosVD/Public/ChaosVDParticleActor.h
--- a/Engine/Plugins/Experimental/ChaosVD/Source/ChaosVD/Public/ChaosVDParticleActor.h
+++ b/Engine/Plugins/Experimental/ChaosVD/Source/ChaosVD/Public/ChaosVDParticleActor.h
@@ -27,6 +27,7 @@ namespace Chaos
 }
 
 /** Options flags to control how geometry is updated in a ChaosVDActor */
+UENUM()
 enum class EChaosVDActorGeometryUpdateFlags : int32
 {
 	None = 0,
@@ -68,6 +69,7 @@ public:
 #endif
 
 	void UpdateGeometryComponentsVisibility();
+	void UpdateGeometryColors();
 
 	/** Changes the active state of this CVD Particle Actor */
 	void SetIsActive(bool bNewActive);
@@ -88,6 +90,8 @@ protected:
 	const TArray<TSharedPtr<FChaosVDParticlePairMidPhase>>* GetCollisionMidPhasesArray() const;
 
 	void UpdateShapeDataComponents();
+
+	void PerformTaskOnGeometryComponents(TFunction<void(IChaosVDGeometryDataComponent& InDataComponent)> TaskToPerform);
 	
 	UPROPERTY(EditAnywhere, Category = "Viewport Visualization Flags", meta = (Bitmask, BitmaskEnum = "/Script/ChaosVD.EChaosVDParticleDataVisualizationFlags"))
 	uint8 LocalParticleDataVisualizationFlags;
diff --git a/Engine/Plugins/Experimental/ChaosVD/Source/ChaosVD/Public/ChaosVDScene.h b/Engine/Plugins/Experimental/ChaosVD/Source/ChaosVD/Public/ChaosVDScene.h
--- a/Engine/Plugins/Experimental/ChaosVD/Source/ChaosVD/Public/ChaosVDScene.h
+++ b/Engine/Plugins/Experimental/ChaosVD/Source/ChaosVD/Public/ChaosVDScene.h
@@ -131,6 +131,7 @@ private:
 	void ClearSelectionAndNotify();
 
 	void HandleVisibilitySettingsChanged(UChaosVDEditorSettings* SettingsObject);
+	void HandleColorSettingsChanged(UChaosVDEditorSettings* SettingsObject);
 
 	void InitializeSelectionSets();
 	void DeInitializeSelectionSets();
diff --git a/Engine/Plugins/Experimental/ChaosVD/Source/ChaosVD/Public/Components/ChaosVDInstancedStaticMeshComponent.h b/Engine/Plugins/Experimental/ChaosVD/Source/ChaosVD/Public/Components/ChaosVDInstancedStaticMeshComponent.h
--- a/Engine/Plugins/Experimental/ChaosVD/Source/ChaosVD/Public/Components/ChaosVDInstancedStaticMeshComponent.h
+++ b/Engine/Plugins/Experimental/ChaosVD/Source/ChaosVD/Public/Components/ChaosVDInstancedStaticMeshComponent.h
@@ -37,6 +37,10 @@ class UChaosVDInstancedStaticMeshComponent : public UInstancedStaticMeshComponen
 	virtual void UpdateVisibility() override;
 	
 	virtual void UpdateDataFromShapeArray(const TArray<FChaosVDShapeCollisionData>& InShapeArray) override;
+	
+	virtual void SetImplicitObject(const Chaos::FImplicitObject* InImplicitObject) override;
+
+	virtual void UpdateColors() override;
 	// END IChaosVDGeometryDataComponent Interface
 
 protected:
diff --git a/Engine/Plugins/Experimental/ChaosVD/Source/ChaosVD/Public/Components/ChaosVDStaticMeshComponent.h b/Engine/Plugins/Experimental/ChaosVD/Source/ChaosVD/Public/Components/ChaosVDStaticMeshComponent.h
--- a/Engine/Plugins/Experimental/ChaosVD/Source/ChaosVD/Public/Components/ChaosVDStaticMeshComponent.h
+++ b/Engine/Plugins/Experimental/ChaosVD/Source/ChaosVD/Public/Components/ChaosVDStaticMeshComponent.h
@@ -11,6 +11,7 @@ UCLASS(HideCategories=("Transform"), MinimalAPI)
 class UChaosVDStaticMeshComponent : public UStaticMeshComponent, public IChaosVDGeometryDataComponent, public FChaosVDGeometryDataComponentBase
 {
 	GENERATED_BODY()
+
 public:
 	UChaosVDStaticMeshComponent(const FObjectInitializer& ObjectInitializer)
 		: Super(ObjectInitializer)
@@ -36,7 +37,11 @@ public:
 	
 	virtual void UpdateVisibility() override;
 	
-	virtual void UpdateDataFromShapeArray(const TArray<FChaosVDShapeCollisionData>& InShapeArray) override;;
+	virtual void UpdateDataFromShapeArray(const TArray<FChaosVDShapeCollisionData>& InShapeArray) override;
+
+	virtual void SetImplicitObject(const Chaos::FImplicitObject* InImplicitObject) override;
+	
+	virtual void UpdateColors() override;
 	// END IChaosVDGeometryDataComponent Interface
 
 protected:
