Chaos Visual Debugger
- add material data to collision constraint
- add scales to angvel etc vizualizer (hard-coded for now)

diff --git a/Engine/Plugins/Experimental/ChaosVD/Source/ChaosVD/Private/Visualizers/ChaosVDParticleDataVisualizer.cpp b/Engine/Plugins/Experimental/ChaosVD/Source/ChaosVD/Private/Visualizers/ChaosVDParticleDataVisualizer.cpp
--- a/Engine/Plugins/Experimental/ChaosVD/Source/ChaosVD/Private/Visualizers/ChaosVDParticleDataVisualizer.cpp
+++ b/Engine/Plugins/Experimental/ChaosVD/Source/ChaosVD/Private/Visualizers/ChaosVDParticleDataVisualizer.cpp
@@ -32,19 +32,28 @@ void FChaosVDParticleDataVisualizer::DrawVisualization(const FSceneView* View, F
 	}
 
 	// TODO: Implement scale settings for the vectors, and take into account the simulation space transform
+	const float VelScale = 0.5f;		// [cm/s]
+	const float AngVelScale = 50.0f;	// [rad/s]
+	const float AccScale = 0.005f;		// [cm/s2]
+	const float AngAccScale = 0.5f;		// [rad/s2]
+	const float ImpScale = 0.001;		// [g.m/s]
+	const float AngImpScale = 0.1f;		// [g.m2/s]
 
 	const FVector& OwnerLocation = ParticleDataViewer->ParticlePositionRotation.MX;
+	const FQuat& OwnerRotation = ParticleDataViewer->ParticlePositionRotation.MR;
+	const FVector OwnerCoMLocation = OwnerLocation + OwnerRotation * ParticleDataViewer->ParticleMassProps.MCenterOfMass;
+
 
 	if (ParticleDataViewer->ParticleVelocities.HasValidData())
 	{
 		if (EnumHasAnyFlags(VisualizationFlagsToUse, EChaosVDParticleDataVisualizationFlags::Velocity))
 		{
-			FChaosVDDebugDrawUtils::DrawArrowVector(PDI, OwnerLocation, OwnerLocation + ParticleDataViewer->ParticleVelocities.MV, TEXT("Velocity"), FColor::Green);
+			FChaosVDDebugDrawUtils::DrawArrowVector(PDI, OwnerCoMLocation, OwnerCoMLocation + VelScale * ParticleDataViewer->ParticleVelocities.MV, TEXT("Velocity"), FColor::Green);
 		}
 
 		if (EnumHasAnyFlags(VisualizationFlagsToUse, EChaosVDParticleDataVisualizationFlags::AngularVelocity))
 		{
-			FChaosVDDebugDrawUtils::DrawArrowVector(PDI, OwnerLocation, OwnerLocation + ParticleDataViewer->ParticleVelocities.MW, TEXT("Angular Velocity"), FColor::Blue);
+			FChaosVDDebugDrawUtils::DrawArrowVector(PDI, OwnerCoMLocation, OwnerCoMLocation + AngVelScale * ParticleDataViewer->ParticleVelocities.MW, TEXT("Angular Velocity"), FColor::Blue);
 		}
 	}
 
@@ -52,22 +61,22 @@ void FChaosVDParticleDataVisualizer::DrawVisualization(const FSceneView* View, F
 	{
 		if (EnumHasAnyFlags(VisualizationFlagsToUse, EChaosVDParticleDataVisualizationFlags::Acceleration))
 		{
-			FChaosVDDebugDrawUtils::DrawArrowVector(PDI, OwnerLocation, OwnerLocation + ParticleDataViewer->ParticleDynamics.MAcceleration, TEXT("Acceleration"), FColor::Orange);
+			FChaosVDDebugDrawUtils::DrawArrowVector(PDI, OwnerCoMLocation, OwnerCoMLocation + AccScale * ParticleDataViewer->ParticleDynamics.MAcceleration, TEXT("Acceleration"), FColor::Orange);
 		}
 		
 		if (EnumHasAnyFlags(VisualizationFlagsToUse, EChaosVDParticleDataVisualizationFlags::AngularAcceleration))
 		{
-			FChaosVDDebugDrawUtils::DrawArrowVector(PDI, OwnerLocation, OwnerLocation + ParticleDataViewer->ParticleDynamics.MAngularAcceleration, TEXT("Angular Acceleration"), FColor::Purple);
+			FChaosVDDebugDrawUtils::DrawArrowVector(PDI, OwnerCoMLocation, OwnerCoMLocation + AngAccScale * ParticleDataViewer->ParticleDynamics.MAngularAcceleration, TEXT("Angular Acceleration"), FColor::Purple);
 		}
 
 		if (EnumHasAnyFlags(VisualizationFlagsToUse, EChaosVDParticleDataVisualizationFlags::LinearImpulse))
 		{
-			FChaosVDDebugDrawUtils::DrawArrowVector(PDI, OwnerLocation, OwnerLocation + ParticleDataViewer->ParticleDynamics.MLinearImpulseVelocity, TEXT("Linear Implulse Velocity"), FColor::Purple);
+			FChaosVDDebugDrawUtils::DrawArrowVector(PDI, OwnerCoMLocation, OwnerCoMLocation + ImpScale * ParticleDataViewer->ParticleDynamics.MLinearImpulseVelocity, TEXT("Linear Implulse Velocity"), FColor::Purple);
 		}
 
 		if (EnumHasAnyFlags(VisualizationFlagsToUse, EChaosVDParticleDataVisualizationFlags::AngularImpulse))
 		{
-			FChaosVDDebugDrawUtils::DrawArrowVector(PDI, OwnerLocation, OwnerLocation + ParticleDataViewer->ParticleDynamics.MAngularImpulseVelocity, TEXT("Angular Implulse Velocity"), FColor::Emerald);
+			FChaosVDDebugDrawUtils::DrawArrowVector(PDI, OwnerCoMLocation, OwnerCoMLocation + AngImpScale * ParticleDataViewer->ParticleDynamics.MAngularImpulseVelocity, TEXT("Angular Implulse Velocity"), FColor::Emerald);
 		}
 	}
 }
diff --git a/Engine/Source/Runtime/Experimental/Chaos/Private/ChaosVisualDebugger/ChaosVDDataWrapperUtils.cpp b/Engine/Source/Runtime/Experimental/Chaos/Private/ChaosVisualDebugger/ChaosVDDataWrapperUtils.cpp
--- a/Engine/Source/Runtime/Experimental/Chaos/Private/ChaosVisualDebugger/ChaosVDDataWrapperUtils.cpp
+++ b/Engine/Source/Runtime/Experimental/Chaos/Private/ChaosVisualDebugger/ChaosVDDataWrapperUtils.cpp
@@ -37,6 +37,22 @@ void FChaosVDDataWrapperUtils::CopyManifoldPointResultsToDataWrapper(const Chaos
 	OutCopyTo.bInsideStaticFrictionCone =  InCopyFrom.bInsideStaticFrictionCone;
 }
 
+void FChaosVDDataWrapperUtils::CopyCollisionMaterialToDataWrapper(const Chaos::FPBDCollisionConstraintMaterial& InCopyFrom, FChaosVDCollisionMaterial& OutCopyTo)
+{
+	OutCopyTo.FaceIndex = InCopyFrom.FaceIndex;
+	OutCopyTo.MaterialDynamicFriction = InCopyFrom.MaterialDynamicFriction;
+	OutCopyTo.MaterialStaticFriction = InCopyFrom.MaterialStaticFriction;
+	OutCopyTo.MaterialRestitution = InCopyFrom.MaterialRestitution;
+	OutCopyTo.DynamicFriction = InCopyFrom.DynamicFriction;
+	OutCopyTo.StaticFriction = InCopyFrom.StaticFriction;
+	OutCopyTo.Restitution = InCopyFrom.Restitution;
+	OutCopyTo.RestitutionThreshold = InCopyFrom.RestitutionThreshold;
+	OutCopyTo.InvMassScale0 = InCopyFrom.InvMassScale0;
+	OutCopyTo.InvMassScale1 = InCopyFrom.InvMassScale1;
+	OutCopyTo.InvInertiaScale0 = InCopyFrom.InvInertiaScale0;
+	OutCopyTo.InvInertiaScale1 = InCopyFrom.InvInertiaScale1;
+}
+
 FChaosVDParticleDataWrapper FChaosVDDataWrapperUtils::BuildParticleDataWrapperFromParticle(const Chaos::FGeometryParticleHandle* ParticleHandlePtr)
 {
 	check(ParticleHandlePtr);
@@ -93,6 +109,8 @@ FChaosVDConstraint FChaosVDDataWrapperUtils::BuildConstraintDataWrapperFromConst
 	WrappedConstraintData.CCDEnablePenetration = InConstraint.CCDEnablePenetration;
 	WrappedConstraintData.CCDTargetPenetration = InConstraint.CCDTargetPenetration;
 	
+	CopyCollisionMaterialToDataWrapper(InConstraint.Material, WrappedConstraintData.Material);
+
 	WrappedConstraintData.AccumulatedImpulse = FVector(InConstraint.AccumulatedImpulse);
 
 	WrappedConstraintData.Particle0Index = InConstraint.GetParticle0()->UniqueIdx().Idx;
diff --git a/Engine/Source/Runtime/Experimental/Chaos/Public/ChaosVisualDebugger/ChaosVDDataWrapperUtils.h b/Engine/Source/Runtime/Experimental/Chaos/Public/ChaosVisualDebugger/ChaosVDDataWrapperUtils.h
--- a/Engine/Source/Runtime/Experimental/Chaos/Public/ChaosVisualDebugger/ChaosVDDataWrapperUtils.h
+++ b/Engine/Source/Runtime/Experimental/Chaos/Public/ChaosVisualDebugger/ChaosVDDataWrapperUtils.h
@@ -11,9 +11,11 @@ namespace Chaos
 {
 	class FParticlePairMidPhase;
 	class FPBDCollisionConstraint;
+	class FPBDCollisionConstraintMaterial;
 }
 
 class FChaosVisualDebuggerTrace;
+struct FChaosVDCollisionMaterial;
 struct FChaosVDConstraint;
 struct FChaosVDParticlePairMidPhase;
 struct FChaosVDParticleDataWrapper;
@@ -35,6 +37,9 @@ private:
 	/** Takes a FManifoldPointResult and copies the relevant data to the CVD counterpart */
 	static void CopyManifoldPointResultsToDataWrapper(const Chaos::FManifoldPointResult& InCopyFrom, FChaosVDManifoldPoint& OutCopyTo);
 
+	/** Takes a FPBDCollisionConstraintMaterial and copies the relevant data to the CVD counterpart */
+	static void CopyCollisionMaterialToDataWrapper(const Chaos::FPBDCollisionConstraintMaterial& InCopyFrom, FChaosVDCollisionMaterial& OutCopyTo);
+
 	/** Creates and populates a FChaosVDParticleDataWrapper with the data of the provided FGeometryParticleHandle */
 	static FChaosVDParticleDataWrapper BuildParticleDataWrapperFromParticle(const Chaos::FGeometryParticleHandle* ParticleHandlePtr);
 
diff --git a/Engine/Source/Runtime/Experimental/ChaosVisualDebugger/Private/DataWrappers/ChaosVDCollisionDataWrappers.cpp b/Engine/Source/Runtime/Experimental/ChaosVisualDebugger/Private/DataWrappers/ChaosVDCollisionDataWrappers.cpp
--- a/Engine/Source/Runtime/Experimental/ChaosVisualDebugger/Private/DataWrappers/ChaosVDCollisionDataWrappers.cpp
+++ b/Engine/Source/Runtime/Experimental/ChaosVisualDebugger/Private/DataWrappers/ChaosVDCollisionDataWrappers.cpp
@@ -35,6 +35,24 @@ bool FChaosVDManifoldPoint::Serialize(FArchive& Ar)
 	return true;
 }
 
+bool FChaosVDCollisionMaterial::Serialize(FArchive& Ar)
+{
+	Ar << FaceIndex;
+	Ar << MaterialDynamicFriction;
+	Ar << MaterialStaticFriction;
+	Ar << MaterialRestitution;
+	Ar << DynamicFriction;
+	Ar << StaticFriction;
+	Ar << Restitution;
+	Ar << RestitutionThreshold;
+	Ar << InvMassScale0;
+	Ar << InvMassScale1;
+	Ar << InvInertiaScale0;
+	Ar << InvInertiaScale1;
+
+	return true;
+}
+
 bool FChaosVDConstraint::Serialize(FArchive& Ar)
 {
 	FArchive_Serialize_BitfieldBool(Ar, bIsCurrent);
@@ -51,6 +69,7 @@ bool FChaosVDConstraint::Serialize(FArchive& Ar)
 	FArchive_Serialize_BitfieldBool(Ar, bModifierApplied);
 	FArchive_Serialize_BitfieldBool(Ar, bMaterialSet);
 
+	Ar << Material;
 	Ar << AccumulatedImpulse;
 	Ar << ShapesType;
 	Ar << ShapeWorldTransforms;
diff --git a/Engine/Source/Runtime/Experimental/ChaosVisualDebugger/Public/DataWrappers/ChaosVDCollisionDataWrappers.h b/Engine/Source/Runtime/Experimental/ChaosVisualDebugger/Public/DataWrappers/ChaosVDCollisionDataWrappers.h
--- a/Engine/Source/Runtime/Experimental/ChaosVisualDebugger/Public/DataWrappers/ChaosVDCollisionDataWrappers.h
+++ b/Engine/Source/Runtime/Experimental/ChaosVisualDebugger/Public/DataWrappers/ChaosVDCollisionDataWrappers.h
@@ -150,6 +150,66 @@ inline FArchive& operator<<(FArchive& Ar, FChaosVDManifoldPoint& Data)
 	return Ar;
 }
 
+USTRUCT()
+struct FChaosVDCollisionMaterial
+{
+	GENERATED_BODY()
+
+	UPROPERTY(VisibleAnywhere, Category = MaterialData)
+	int32 FaceIndex = INDEX_NONE;
+
+	UPROPERTY(VisibleAnywhere, Category = MaterialData)
+	float MaterialDynamicFriction = 0.0f;
+	
+	UPROPERTY(VisibleAnywhere, Category = MaterialData)
+	float MaterialStaticFriction = 0.0f;
+	
+	UPROPERTY(VisibleAnywhere, Category = MaterialData)
+	float MaterialRestitution = 0.0f;
+
+	UPROPERTY(VisibleAnywhere, Category = MaterialData)
+	float DynamicFriction = 0.0f;
+	
+	UPROPERTY(VisibleAnywhere, Category = MaterialData)
+	float StaticFriction = 0.0f;
+	
+	UPROPERTY(VisibleAnywhere, Category = MaterialData)
+	float Restitution = 0.0f;
+	
+	UPROPERTY(VisibleAnywhere, Category = MaterialData)
+	float RestitutionThreshold = 0.0f;
+	
+	UPROPERTY(VisibleAnywhere, Category = MaterialData)
+	float InvMassScale0 = 0.0f;
+	
+	UPROPERTY(VisibleAnywhere, Category = MaterialData)
+	float InvMassScale1 = 0.0f;
+	
+	UPROPERTY(VisibleAnywhere, Category = MaterialData)
+	float InvInertiaScale0 = 0.0f;
+	
+	UPROPERTY(VisibleAnywhere, Category = MaterialData)
+	float InvInertiaScale1 = 0.0f;
+
+	bool Serialize(FArchive& Ar);
+};
+
+template<>
+struct TStructOpsTypeTraits<FChaosVDCollisionMaterial> : public TStructOpsTypeTraitsBase2<FChaosVDCollisionMaterial>
+{
+	enum
+	{
+		WithSerializer = true,
+	};
+};
+
+inline FArchive& operator<<(FArchive& Ar, FChaosVDCollisionMaterial& Data)
+{
+	Data.Serialize(Ar);
+	return Ar;
+}
+
+
 USTRUCT()
 struct CHAOSVDRUNTIME_API FChaosVDConstraint
 {
@@ -184,6 +244,9 @@ struct CHAOSVDRUNTIME_API FChaosVDConstraint
 	UPROPERTY(VisibleAnywhere, Category=ConstraintData)
 	uint8 bMaterialSet:1 = false;
 	
+	UPROPERTY(VisibleAnywhere, Category = ConstraintData)
+	FChaosVDCollisionMaterial Material;
+
 	UPROPERTY(VisibleAnywhere, Category=ConstraintData)
 	FVector AccumulatedImpulse = FVector(ForceInit);
 	
