[ChaosVD] Ensure that particle data is up to date when stepping backwards and fix for a crash when recording invalid midphases.

The issue we are fixing happens because each frame we record deltas, therefore during playback if when stepping forward the delta has data about a particle, CVD will update it. But if when stepping back that changed particle is not part of the previous delta, CVD will ignore it (therefore keeping the new now incorrect value).

To fix this I am using the same system I have in place for when we jump more than 1 frame which is playback from the closes keyframe (by collapsing all previous deltas into a generated keyframe)

This affects the tool's perf while stepping, I will find a better solution later on.

diff --git a/Engine/Plugins/Experimental/ChaosVD/Source/ChaosVD/Private/ChaosVDPlaybackController.cpp b/Engine/Plugins/Experimental/ChaosVD/Source/ChaosVD/Private/ChaosVDPlaybackController.cpp
--- a/Engine/Plugins/Experimental/ChaosVD/Source/ChaosVD/Private/ChaosVDPlaybackController.cpp
+++ b/Engine/Plugins/Experimental/ChaosVD/Source/ChaosVD/Private/ChaosVDPlaybackController.cpp
@@ -166,7 +166,15 @@ void FChaosVDPlaybackController::GoToRecordedSolverStep_AssumesLocked(const int3
 			{
 				const int32 FrameDiff = FrameNumber - CurrentTrackInfo->CurrentFrame;
 				constexpr int32 FrameDriftTolerance = 1;
-				if (FMath::Abs(FrameDiff) > FrameDriftTolerance || CurrentTrackInfo->CurrentFrame == 0)
+
+				// If we go back, even for one single step and the particles that changed are not in the prev step, we have no data to restore their changed values.
+				// So for now if we are going backwards, always play from closest keyframe.
+				// TODO: Implement a less expensive way of handle these cases.
+				// We should keep the previous state of each loaded particle so if when going back they are not in the new delta we are evaluating, (and were not destroyed)
+				// we can just re-apply that last known state.
+				const bool bNeedsToPlayFromKeyframe = FrameDiff < 0 || FMath::Abs(FrameDiff) > FrameDriftTolerance;
+
+				if (bNeedsToPlayFromKeyframe || CurrentTrackInfo->CurrentFrame == 0)
 				{
 					// As Frames are recorded as delta, we need to make sure of playing back all the deltas since the closest keyframe
 					PlayFromClosestKeyFrame_AssumesLocked(InTrackID, FrameNumber, *SceneToControlSharedPtr.Get());
diff --git a/Engine/Source/Runtime/Experimental/Chaos/Private/ChaosVisualDebugger/ChaosVisualDebuggerTrace.cpp b/Engine/Source/Runtime/Experimental/Chaos/Private/ChaosVisualDebugger/ChaosVisualDebuggerTrace.cpp
--- a/Engine/Source/Runtime/Experimental/Chaos/Private/ChaosVisualDebugger/ChaosVisualDebuggerTrace.cpp
+++ b/Engine/Source/Runtime/Experimental/Chaos/Private/ChaosVisualDebugger/ChaosVisualDebuggerTrace.cpp
@@ -274,6 +274,11 @@ void FChaosVisualDebuggerTrace::TraceMidPhase(const Chaos::FParticlePairMidPhase
 		return;
 	}
 
+	if (!MidPhase->IsValid())
+	{
+		return;
+	}
+
 	FChaosVDParticlePairMidPhase CVDMidPhase = FChaosVDDataWrapperUtils::BuildMidPhaseDataWrapperFromMidPhase(*MidPhase);
 	CVDMidPhase.SolverID = CVDContextData->Id;
 
diff --git a/Engine/Source/Runtime/Experimental/ChaosVDData/Private/ChaosVDRecording.cpp b/Engine/Source/Runtime/Experimental/ChaosVDData/Private/ChaosVDRecording.cpp
--- a/Engine/Source/Runtime/Experimental/ChaosVDData/Private/ChaosVDRecording.cpp
+++ b/Engine/Source/Runtime/Experimental/ChaosVDData/Private/ChaosVDRecording.cpp
@@ -290,7 +290,9 @@ void FChaosVDRecording::GetAvailableSolverIDsAtGameFrame_AssumesLocked(const FCh
 
 void FChaosVDRecording::CollapseSolverFramesRange_AssumesLocked(int32 SolverID, int32 StartFrame, int32 EndFrame, FChaosVDSolverFrameData& OutCollapsedFrameData)
 {
-	TMap<int32, FChaosVDParticleDataWrapper> TempParticleDataByID;
+	static TMap<int32, FChaosVDParticleDataWrapper> TempParticleDataByID;
+
+	TempParticleDataByID.Reset();
 	
 	for (int32 CurrentFrameNumber = StartFrame; CurrentFrameNumber <= EndFrame; CurrentFrameNumber++)
 	{
