Chaos Visual Debugger - add support for client/server coloring

diff --git a/Engine/Plugins/Experimental/ChaosVD/Source/ChaosVD/Private/Actors/ChaosVDSolverInfoActor.cpp b/Engine/Plugins/Experimental/ChaosVD/Source/ChaosVD/Private/Actors/ChaosVDSolverInfoActor.cpp
--- a/Engine/Plugins/Experimental/ChaosVD/Source/ChaosVD/Private/Actors/ChaosVDSolverInfoActor.cpp
+++ b/Engine/Plugins/Experimental/ChaosVD/Source/ChaosVD/Private/Actors/ChaosVDSolverInfoActor.cpp
@@ -11,6 +11,7 @@
 AChaosVDSolverInfoActor::AChaosVDSolverInfoActor(const FObjectInitializer& ObjectInitializer) : Super(ObjectInitializer)
 {
 	CollisionDataComponent = CreateDefaultSubobject<UChaosVDSolverCollisionDataComponent>(TEXT("SolverCollisionDataComponent"));
+	bIsServer = false;
 }
 
 void AChaosVDSolverInfoActor::SetScene(TWeakPtr<FChaosVDScene> InScene)
diff --git a/Engine/Plugins/Experimental/ChaosVD/Source/ChaosVD/Private/ChaosVDEditorSettings.cpp b/Engine/Plugins/Experimental/ChaosVD/Source/ChaosVD/Private/ChaosVDEditorSettings.cpp
--- a/Engine/Plugins/Experimental/ChaosVD/Source/ChaosVD/Private/ChaosVDEditorSettings.cpp
+++ b/Engine/Plugins/Experimental/ChaosVD/Source/ChaosVD/Private/ChaosVDEditorSettings.cpp
@@ -50,6 +50,23 @@ FColor FChaosDebugDrawColorsByShapeType::GetColorFromShapeType(Chaos::EImplicitO
 	}
 }
 
+FColor FChaosDebugDrawColorsByClientServer::GetColorFromState(bool bIsServer, EChaosVDObjectStateType State) const
+{
+	switch (State)
+	{
+	case EChaosVDObjectStateType::Sleeping:
+		return bIsServer ? ServerSleepingColor : ClientSleepingColor;
+	case EChaosVDObjectStateType::Kinematic:
+		return bIsServer ? ServerColor : ClientColor;
+	case EChaosVDObjectStateType::Static:
+		return bIsServer ? ServerColor : ClientColor;
+	case EChaosVDObjectStateType::Dynamic:
+		return bIsServer ? ServerDynamicColor : ClientDynamicColor;
+	default:
+		return FColor::Purple;
+	}
+}
+
 void UChaosVDEditorSettings::PostEditChangeProperty(FPropertyChangedEvent& PropertyChangedEvent)
 {
 	Super::PostEditChangeProperty(PropertyChangedEvent);
diff --git a/Engine/Plugins/Experimental/ChaosVD/Source/ChaosVD/Private/ChaosVDGeometryDataComponent.cpp b/Engine/Plugins/Experimental/ChaosVD/Source/ChaosVD/Private/ChaosVDGeometryDataComponent.cpp
--- a/Engine/Plugins/Experimental/ChaosVD/Source/ChaosVD/Private/ChaosVDGeometryDataComponent.cpp
+++ b/Engine/Plugins/Experimental/ChaosVD/Source/ChaosVD/Private/ChaosVDGeometryDataComponent.cpp
@@ -6,6 +6,7 @@
 #include "ChaosVDGeometryBuilder.h"
 #include "ChaosVDModule.h"
 #include "ChaosVDParticleActor.h"
+#include "ChaosVDScene.h"
 #include "Components/MeshComponent.h"
 #include "MaterialDomain.h"
 #include "Materials/Material.h"
@@ -150,6 +151,20 @@ void FChaosVDGeometryDataComponentBase::UpdateColors_Internal(UMeshComponent* Me
 				}
 				break;
 			}
+		case EChaosVDParticleDebugColorMode::ClientServer:
+			{
+				const TSharedPtr<FChaosVDScene> Scene = ParticleActor->GetScene().Pin();
+				const bool bIsServer = (Scene.IsValid()) ? Scene->IsSolverForServer(ParticleData->SolverID) : false;
+				if (ParticleData->Type == EChaosVDParticleType::Static)
+				{
+					ColorToApply = EditorSettings->ColorsByClientServer.GetColorFromState(bIsServer, EChaosVDObjectStateType::Static);
+				}
+				else
+				{
+					ColorToApply = EditorSettings->ColorsByClientServer.GetColorFromState(bIsServer, ParticleData->ParticleDynamicsMisc.MObjectState);
+				}
+				break;
+			}
 
 		case EChaosVDParticleDebugColorMode::None:
 		default:
diff --git a/Engine/Plugins/Experimental/ChaosVD/Source/ChaosVD/Private/ChaosVDScene.cpp b/Engine/Plugins/Experimental/ChaosVD/Source/ChaosVD/Private/ChaosVDScene.cpp
--- a/Engine/Plugins/Experimental/ChaosVD/Source/ChaosVD/Private/ChaosVDScene.cpp
+++ b/Engine/Plugins/Experimental/ChaosVD/Source/ChaosVD/Private/ChaosVDScene.cpp
@@ -242,8 +242,12 @@ void FChaosVDScene::HandleEnterNewGameFrame(int32 FrameNumber, const TArray<int3
 		{
 			AChaosVDSolverInfoActor* CollisionDataContainer = PhysicsVDWorld->SpawnActor<AChaosVDSolverInfoActor>();
 			check(CollisionDataContainer);
+
+			const bool bIsServer = LoadedRecording->GetSolverName_AssumedLocked(SolverID).Contains(TEXT("Server"));
+
 			CollisionDataContainer->SetSolverID(SolverID);
 			CollisionDataContainer->SetScene(AsWeak());
+			CollisionDataContainer->SetIsServer(bIsServer);
 			SolverDataContainerBySolverID.Add(SolverID, CollisionDataContainer);
 		}
 	}
@@ -320,6 +324,17 @@ AChaosVDSolverInfoActor* FChaosVDScene::GetSolverInfoActor(int32 SolverID)
 	return nullptr;
 }
 
+bool FChaosVDScene::IsSolverForServer(int32 SolverID) const
+{
+	if (const AChaosVDSolverInfoActor* const * PSolverDataInfo = SolverDataContainerBySolverID.Find(SolverID))
+	{
+		const AChaosVDSolverInfoActor* SolverDataInfo = *PSolverDataInfo;
+		return SolverDataInfo->GetIsServer();
+	}
+
+	return false;
+}
+
 AChaosVDParticleActor* FChaosVDScene::SpawnParticleFromRecordedData(const FChaosVDParticleDataWrapper& InParticleData, const FChaosVDSolverFrameData& InFrameData)
 {
 	using namespace Chaos;
diff --git a/Engine/Plugins/Experimental/ChaosVD/Source/ChaosVD/Public/Actors/ChaosVDSolverInfoActor.h b/Engine/Plugins/Experimental/ChaosVD/Source/ChaosVD/Public/Actors/ChaosVDSolverInfoActor.h
--- a/Engine/Plugins/Experimental/ChaosVD/Source/ChaosVD/Public/Actors/ChaosVDSolverInfoActor.h
+++ b/Engine/Plugins/Experimental/ChaosVD/Source/ChaosVD/Public/Actors/ChaosVDSolverInfoActor.h
@@ -22,6 +22,9 @@ public:
 	void SetSolverID(int32 InSolverID) { SolverID = InSolverID; }
 	int32 GetSolverID() const { return SolverID; }
 
+	void SetIsServer(bool bInIsServer) { bIsServer = bInIsServer; }
+	bool GetIsServer() const { return bIsServer; }
+
 	virtual void SetScene(TWeakPtr<FChaosVDScene> InScene) override;
 
 	void SetSimulationTransform(const FTransform& InSimulationTransform) { SimulationTransform = InSimulationTransform; }
@@ -58,4 +61,6 @@ protected:
 	TMap<int32, AChaosVDParticleActor*> SolverParticlesByID;
 
 	TArray<int32> SelectedParticlesID;
+
+	bool bIsServer;
 };
diff --git a/Engine/Plugins/Experimental/ChaosVD/Source/ChaosVD/Public/ChaosVDEditorSettings.h b/Engine/Plugins/Experimental/ChaosVD/Source/ChaosVD/Public/ChaosVDEditorSettings.h
--- a/Engine/Plugins/Experimental/ChaosVD/Source/ChaosVD/Public/ChaosVDEditorSettings.h
+++ b/Engine/Plugins/Experimental/ChaosVD/Source/ChaosVD/Public/ChaosVDEditorSettings.h
@@ -122,12 +122,46 @@ struct FChaosDebugDrawColorsByShapeType
 	FColor GetColorFromShapeType(Chaos::EImplicitObjectType ShapeType) const;
 };
 
+/** Structure holding the settings using to debug draw Particles shape based on whether they are client or server objects (in PIE) Chaos Visual Debugger */
+USTRUCT()
+struct FChaosDebugDrawColorsByClientServer
+{
+	GENERATED_BODY()
+
+	/** Color used for server shapes that are not awake or sleeping dynamic */
+	UPROPERTY(EditAnywhere, Category=DebugDraw)
+	FColor ServerColor = FColor(50, 0, 0); 
+
+	/** Color used for server shapes that are awake dynamic */
+	UPROPERTY(EditAnywhere, Category = DebugDraw)
+	FColor ServerDynamicColor = FColor(150, 0, 0);
+
+	/** Color used for server shapes that are sleeping dynamics */
+	UPROPERTY(EditAnywhere, Category = DebugDraw)
+	FColor ServerSleepingColor = FColor(10, 0, 0);
+
+	/** Color used for client shapes that are not awake or sleeping dynamic */
+	UPROPERTY(EditAnywhere, Category = DebugDraw)
+	FColor ClientColor = FColor(0, 0, 50);
+
+	/** Color used for server shapes that are awake dynamic */
+	UPROPERTY(EditAnywhere, Category = DebugDraw)
+	FColor ClientDynamicColor = FColor(0, 0, 150);
+
+	/** Color used for client shapes that are sleeping dynamics */
+	UPROPERTY(EditAnywhere, Category = DebugDraw)
+	FColor ClientSleepingColor = FColor(0, 0, 100);
+
+	FColor GetColorFromState(bool bIsServer, EChaosVDObjectStateType State) const;
+};
+
 UENUM()
 enum class EChaosVDParticleDebugColorMode
 {
 	None,
 	State,
-	ShapeType
+	ShapeType,
+	ClientServer,
 };
 
 UCLASS(config = Engine)
@@ -157,6 +191,9 @@ public:
 	UPROPERTY(EditAnywhere, Category = "Viewport Visualization", meta=(EditCondition = "ParticleColorMode == EChaosVDParticleDebugColorMode::State", EditConditionHides))
 	FChaosDebugDrawColorsByState ColorsByParticleState;
 
+	UPROPERTY(EditAnywhere, Category = "Viewport Visualization", meta = (EditCondition = "ParticleColorMode == EChaosVDParticleDebugColorMode::ClientServer", EditConditionHides))
+	FChaosDebugDrawColorsByClientServer ColorsByClientServer;
+
 	UPROPERTY(EditAnywhere, Category = "Viewport Tracking")
 	EChaosVDActorTrackingTarget TrackingTarget;
 
diff --git a/Engine/Plugins/Experimental/ChaosVD/Source/ChaosVD/Public/ChaosVDScene.h b/Engine/Plugins/Experimental/ChaosVD/Source/ChaosVD/Public/ChaosVDScene.h
--- a/Engine/Plugins/Experimental/ChaosVD/Source/ChaosVD/Public/ChaosVDScene.h
+++ b/Engine/Plugins/Experimental/ChaosVD/Source/ChaosVD/Public/ChaosVDScene.h
@@ -98,6 +98,9 @@ public:
 	const FChaosVDSolverInfoByIDMap& GetSolverInfoActorsMap() { return SolverDataContainerBySolverID; }
 	AChaosVDSolverInfoActor* GetSolverInfoActor(int32 SolverID);
 
+	/** Is the specified solver from a Server or a Client? (note: currently inferred from the solver name) */
+	bool IsSolverForServer(int32 Solver) const;
+
 	AActor* GetSkySphereActor() const { return SkySphere; }
 
 	FChaosVDActorActiveStateUpdateDelegate& OnActorActiveStateChanged() { return ParticleActorUpdateDelegate; }
diff --git a/Engine/Source/Runtime/Experimental/ChaosVDData/Private/ChaosVDRecording.cpp b/Engine/Source/Runtime/Experimental/ChaosVDData/Private/ChaosVDRecording.cpp
--- a/Engine/Source/Runtime/Experimental/ChaosVDData/Private/ChaosVDRecording.cpp
+++ b/Engine/Source/Runtime/Experimental/ChaosVDData/Private/ChaosVDRecording.cpp
@@ -34,10 +34,15 @@ int32 FChaosVDRecording::GetAvailableSolverFramesNumber_AssumesLocked(int32 Solv
 
 FString FChaosVDRecording::GetSolverName(int32 SolverID)
 {
-	static FString DefaultName(TEXT("Invalid"));
-
 	FReadScopeLock ReadLock(RecordingDataLock);
 
+	return GetSolverName_AssumedLocked(SolverID);
+}
+
+FString FChaosVDRecording::GetSolverName_AssumedLocked(int32 SolverID)
+{
+	static FString DefaultName(TEXT("Invalid"));
+
 	// Currently we don't create an entry per solver, so we need to get the name from the frame data
 	// TODO: Record Solver specific data per instance and not per frame
 	if (TArray<FChaosVDSolverFrameData>* SolverFramesData = RecordedFramesDataPerSolver.Find(SolverID))
diff --git a/Engine/Source/Runtime/Experimental/ChaosVDData/Public/ChaosVDRecording.h b/Engine/Source/Runtime/Experimental/ChaosVDData/Public/ChaosVDRecording.h
--- a/Engine/Source/Runtime/Experimental/ChaosVDData/Public/ChaosVDRecording.h
+++ b/Engine/Source/Runtime/Experimental/ChaosVDData/Public/ChaosVDRecording.h
@@ -121,6 +121,12 @@ struct CHAOSVDDATA_API FChaosVDRecording
 	 */
 	FString GetSolverName(int32 SolverID);
 
+	/**
+	 * Returns the name of the specified solver id. Must be called from within a ReadLock
+	 * @param SolverID ID of the solver
+	 */
+	FString GetSolverName_AssumedLocked(int32 SolverID);
+
 	/**
 	 * Return a ptr to the existing solver frame data from the specified ID and Frame number
 	 * @param SolverID ID of the solver
