[ChaosVD] Added Cluster Union support and improved performance when playing back recordings with large amount of geometry and particles.

- Updated the CVD geometry builder code to recognize properly Clustered unions
- Decoupled the code that creates meshes based on the implicit meshes, from the one that creates mesh components.
- Added a new CVD macro to invalidate geometry when is modified, which will make sure the changes get recorded
- Create a pool system for the mesh components to avoid constantly create/destroy them as geometry changes
- Fixed an issue where when mesh components were generated for all objects in a ImplicitClustered objects resulted in multiples components with the same name. This caused issues like the component not being added to the actor.
- Fixed an issue where Particle color options where not applied to all mesh components in some situations
- Fixed an issue where disabled particles might not get hidden when the option to do so was enabled
- Removed the need to copy particle data each frame to each particle from the recording to the CVDParticle Actor (this is a considerable perf boost on recordings with a huge number of particles).
- Added a new Particle Data wrapper for Clustered particles. Now if we are recording a clustered particle, we will record its data.
- Updated CVD to set the hierarchy on cluster union particles as they were recorded.
- Now that we have the data, I added a small proof of concept to debug draw the connectivity data. A better implementation of this will be done when we move particle data visualization to the new system already in use for debug draw of collision data.
- Refactored how geometry is instanced and handled to use instanced static mesh components when possible
- Created a new interface to handle geometry in CVD that allows us to make changes we need regardless if it is backed by a Static Mesh component or if it is an instance on a Instanced Static Mesh Component
- Refactored how selection feedback works to take into account the fact that now when you clic in the viewport, that geometry might not be part of the actor you want to select (as it is just an instance on a Instanced Static Mesh component). If it is the later, CVD will figure out and select all the mesh instances that are representing the particle.
- Refactored how Updates to geometry are done to use the new interface that allows us to use Instanced Static mesh component
- Refactored how position updates are applied to the particles. Now we just update the transform of the mesh instances instead of moving the particle actor
- Added support to enable the flag that allows to see the triangle edges, making the unlit view mode way more useful
- Created new materials that allows us to update the colors of specific instances within a Instanced static mesh component using custom data floats. This is needed to keep the particle colors feature working with the new system

#virtualized

diff --git a/Engine/Plugins/Experimental/ChaosVD/Config/Engine.ini b/Engine/Plugins/Experimental/ChaosVD/Config/Engine.ini
--- a/Engine/Plugins/Experimental/ChaosVD/Config/Engine.ini
+++ b/Engine/Plugins/Experimental/ChaosVD/Config/Engine.ini
@@ -1,4 +1,6 @@
 ﻿[/Script/ChaosVD.ChaosVDEditorSettings]
 QueryOnlyMeshesMaterial=/ChaosVD/Materials/ChaosVDQueryOnlyMesh.ChaosVDQueryOnlyMesh
 SimOnlyMeshesMaterial=/ChaosVD/Materials/ChasoVDSimMesh.ChasoVDSimMesh
+InstancedMeshesMaterial=/ChaosVD/Materials/ChaosVDInstanced.ChaosVDInstanced
+InstancedMeshesQueryOnlyMaterial=/ChaosVD/Materials/ChaosVDInstancedQueryOnly.ChaosVDInstancedQueryOnly
 SkySphereActorClass=/ChaosVD/BP_CVD_SkySphere.BP_CVD_SkySphere_C
\ No newline at end of file
diff --git a/Engine/Plugins/Experimental/ChaosVD/Source/ChaosVD/Private/ChaosVDExtractedGeometryDataHandle.cpp b/Engine/Plugins/Experimental/ChaosVD/Source/ChaosVD/Private/ChaosVDExtractedGeometryDataHandle.cpp
new file mode 100644
--- /dev/null
+++ b/Engine/Plugins/Experimental/ChaosVD/Source/ChaosVD/Private/ChaosVDExtractedGeometryDataHandle.cpp
@@ -0,0 +1,14 @@
+﻿// Copyright Epic Games, Inc. All Rights Reserved.
+
+#include "ChaosVDExtractedGeometryDataHandle.h"
+
+#include "Chaos/ImplicitObject.h"
+
+FName FChaosVDExtractedGeometryDataHandle::GetName() const
+{
+	using namespace Chaos;
+
+	static FName InvalidName = TEXT("Invalid");
+
+	return ImplicitObject ? GetImplicitObjectTypeName(GetInnerType(ImplicitObject->GetType())) : InvalidName;
+}
diff --git a/Engine/Plugins/Experimental/ChaosVD/Source/ChaosVD/Private/ChaosVDGeometryBuilder.cpp b/Engine/Plugins/Experimental/ChaosVD/Source/ChaosVD/Private/ChaosVDGeometryBuilder.cpp
--- a/Engine/Plugins/Experimental/ChaosVD/Source/ChaosVD/Private/ChaosVDGeometryBuilder.cpp
+++ b/Engine/Plugins/Experimental/ChaosVD/Source/ChaosVD/Private/ChaosVDGeometryBuilder.cpp
@@ -2,19 +2,54 @@
 
 #include "ChaosVDGeometryBuilder.h"
 
+#include "ChaosVDConvexMeshGenerator.h"
 #include "ChaosVDGeometryDataComponent.h"
+#include "ChaosVDModule.h"
+#include "ChaosVDTriMeshGenerator.h"
 #include "DynamicMeshToMeshDescription.h"
+#include "Generators/CapsuleGenerator.h"
+#include "Generators/MinimalBoxMeshGenerator.h"
+#include "Generators/SphereGenerator.h"
 #include "MeshConstraints.h"
 #include "MeshConstraintsUtil.h"
 #include "MeshSimplification.h"
 #include "StaticMeshAttributes.h"
 #include "UDynamicMesh.h"
+#include "Chaos/HeightField.h"
 #include "UObject/UObjectGlobals.h"
 
-static FAutoConsoleVariable CVarChaosVDGeometryToProcessPerTick(
-	TEXT("p.Chaos.VD.Tool.GeometryToProcessPerTick"),
-	200,
-	TEXT("Number of generated geometry to process each tick when loading a teace file in the CVD tool"));
+void FChaosVDGeometryBuilder::Initialize(const TWeakPtr<FChaosVDScene>& ChaosVDScene)
+{
+	SceneWeakPtr = ChaosVDScene;
+
+	auto ProcessMeshComponent = [WeakThis = AsWeak()](uint32 GeometryKey, const TWeakObjectPtr<UMeshComponent> Object)
+	{
+		const TSharedPtr<FChaosVDGeometryBuilder> GeometryBuilder = WeakThis.Pin();
+		if (!GeometryBuilder)
+		{
+			UE_LOG(LogChaosVDEditor, Verbose, TEXT(" [%s] Failed to update mesh for Handle | Geometry Key [%u] | Handle is invalid"), ANSI_TO_TCHAR(__FUNCTION__), GeometryKey);
+
+			// If the the builder is no longer valid, just consume the request
+			return true;
+		}
+
+		return GeometryBuilder->ApplyMeshToComponentFromKey(Object, GeometryKey);
+	};
+
+	auto ShouldProcessObjectsForKey = [WeakThis = AsWeak()](uint32 GeometryKey)
+	{
+		if (const TSharedPtr<FChaosVDGeometryBuilder> GeometryBuilder = WeakThis.Pin())
+		{
+			return GeometryBuilder->HasGeometryInCache(GeometryKey);
+		}
+
+		return false;
+	};
+
+	MeshComponentsWaitingForGeometry = MakeUnique<FObjectsWaitingGeometryList<FMeshComponentWeakPtr>>(ProcessMeshComponent, NSLOCTEXT("ChaosVisualDebugger", "GeometryGenNotification","Mesh Components"), ShouldProcessObjectsForKey);
+
+	GameThreadTickDelegate = FTSTicker::GetCoreTicker().AddTicker(FTickerDelegate::CreateRaw(this, &FChaosVDGeometryBuilder::GameThreadTick));
+}
 
 void FChaosVDGeometryBuilder::AddReferencedObjects(FReferenceCollector& Collector)
 {
@@ -24,42 +59,48 @@ void FChaosVDGeometryBuilder::AddReferencedObjects(FReferenceCollector& Collecto
 
 bool FChaosVDGeometryBuilder::DoesImplicitContainType(const Chaos::FImplicitObject* InImplicitObject, const Chaos::EImplicitObjectType ImplicitTypeToCheck)
 {
+	TRACE_CPUPROFILER_EVENT_SCOPE(FChaosVDGeometryBuilder::DoesImplicitContainType);
+
 	using namespace Chaos;
+	
+	if (!InImplicitObject)
+	{
+		return false;
+	}
+
 	const EImplicitObjectType InnerType = GetInnerType(InImplicitObject->GetType());
 
 	switch (InnerType)
 	{
 		case ImplicitObjectType::Union:
+		case ImplicitObjectType::UnionClustered:
 			{
-				const FImplicitObjectUnion* Union = InImplicitObject->template GetObject<FImplicitObjectUnion>();
-
-				for (int i = 0; i < Union->GetObjects().Num(); ++i)
+				if (const FImplicitObjectUnion* Union = InImplicitObject->template AsA<FImplicitObjectUnion>())
 				{
-					const FImplicitObjectPtr& UnionImplicit = Union->GetObjects()[i];
-
-					if (DoesImplicitContainType(UnionImplicit.GetReference(), ImplicitTypeToCheck))
+					const TArray<Chaos::FImplicitObjectPtr>& UnionObjects = Union->GetObjects();
+					for (const FImplicitObjectPtr& UnionImplicit : UnionObjects)
 					{
-						return true;
+						if (DoesImplicitContainType(UnionImplicit.GetReference(), ImplicitTypeToCheck))
+						{
+							return true;
+						}
 					}
 				}
-
 				return false;
-				break;
 			}
 		case ImplicitObjectType::Transformed:
 			{
 				const TImplicitObjectTransformed<FReal, 3>* Transformed = InImplicitObject->template GetObject<TImplicitObjectTransformed<FReal, 3>>();
 				return DoesImplicitContainType(Transformed->GetTransformedObject(), ImplicitTypeToCheck);
-				break;
 			}
 	default:
 		return InnerType == ImplicitTypeToCheck;
 	}
 }
 
-bool FChaosVDGeometryBuilder::HasNegativeScale(const Chaos::FRigidTransform3& InTransform) const
+bool FChaosVDGeometryBuilder::HasNegativeScale(const Chaos::FRigidTransform3& InTransform)
 {
-	FVector ScaleSignVector = InTransform.GetScale3D().GetSignVector();
+	const FVector ScaleSignVector = InTransform.GetScale3D().GetSignVector();
 	return ScaleSignVector.X * ScaleSignVector.Y * ScaleSignVector.Z < 0;
 }
 
@@ -180,108 +221,314 @@ UStaticMesh* FChaosVDGeometryBuilder::CreateAndCacheStaticMesh(const uint32 Geom
 	return MainStaticMesh;
 }
 
-void FChaosVDGeometryBuilder::ApplyMeshToComponentFromKey(TWeakObjectPtr<UMeshComponent> MeshComponent, const uint32 GeometryKey)
+void FChaosVDGeometryBuilder::DestroyMeshComponent(UMeshComponent* MeshComponent)
 {
-	if (!MeshComponent.IsValid())
+	if (Cast<UChaosVDInstancedStaticMeshComponent>(MeshComponent))
 	{
-		return;
+		if (IChaosVDGeometryComponent* AsCVDGeometryComponent = Cast<IChaosVDGeometryComponent>(MeshComponent))
+		{
+			TMap<uint32, UChaosVDInstancedStaticMeshComponent*>& InstancedMeshComponentCache = GetInstancedStaticMeshComponentCacheMap(AsCVDGeometryComponent->GetMeshComponentAttributeFlags());
+			InstancedMeshComponentCache.Remove(AsCVDGeometryComponent->GetGeometryKey());
+
+			RemoveMeshComponentWaitingForGeometry(AsCVDGeometryComponent->GetGeometryKey(), MeshComponent);
+
+			AsCVDGeometryComponent->OnComponentEmpty()->RemoveAll(this);
+		}
 	}
 
-	if (UDynamicMeshComponent* DynamicMeshComponent = Cast<UDynamicMeshComponent>(MeshComponent))
+	ComponentMeshPool.DisposeMeshComponent(MeshComponent);
+}
+
+TMap<uint32, UChaosVDInstancedStaticMeshComponent*>& FChaosVDGeometryBuilder::GetInstancedStaticMeshComponentCacheMap(EChaosVDMeshAttributesFlags MeshAttributeFlags)
+{
+	if (EnumHasAnyFlags(MeshAttributeFlags, EChaosVDMeshAttributesFlags::MirroredGeometry))
 	{
-		if (TObjectPtr<UDynamicMesh>* DynamicMesh = DynamicMeshCacheMap.Find(GeometryKey))
+		if (EnumHasAnyFlags(MeshAttributeFlags, EChaosVDMeshAttributesFlags::TranslucentGeometry))
 		{
-			DynamicMeshComponent->SetDynamicMesh(*DynamicMesh);
+			return TranslucentMirroredInstancedMeshComponentByGeometryKey;
+		}
+		else
+		{
+			return MirroredInstancedMeshComponentByGeometryKey;
 		}
 	}
-	else if (UStaticMeshComponent* StaticMeshComponent = Cast<UStaticMeshComponent>(MeshComponent))
+	else
 	{
-		if (TObjectPtr<UStaticMesh>* StaticMesh = StaticMeshCacheMap.Find(GeometryKey))
+		if (EnumHasAnyFlags(MeshAttributeFlags, EChaosVDMeshAttributesFlags::TranslucentGeometry))
+		{
+			return TranslucentInstancedMeshComponentByGeometryKey;
+		}
+		else
 		{
-			StaticMeshComponent->SetStaticMesh(*StaticMesh);
+			return InstancedMeshComponentByGeometryKey;
 		}
 	}
+}
+
+bool FChaosVDGeometryBuilder::ApplyMeshToComponentFromKey(TWeakObjectPtr<UMeshComponent> MeshComponent, const uint32 GeometryKey)
+{
+	bool bApplyMeshRequestProcessed = false;
+	if (!MeshComponent.IsValid())
+	{
+		UE_LOG(LogChaosVDEditor, Warning, TEXT("[%s] Failed to apply geometry with key [%d] | Mesh Component is invalid"), ANSI_TO_TCHAR(__FUNCTION__), GeometryKey);
+
+		// If the component is no longer valid, just consume the request
+		bApplyMeshRequestProcessed = true;
+		return bApplyMeshRequestProcessed;
+	}
 
-	if (IChaosVDGeometryDataComponent* DataComponent = Cast<IChaosVDGeometryDataComponent>(MeshComponent.Get()))
+	IChaosVDGeometryComponent* DataComponent = Cast<IChaosVDGeometryComponent>(MeshComponent.Get());
+	if (!DataComponent)
 	{
+		// If the component is valid but not of the correct type, just consume the request and log the error
+
+		UE_LOG(LogChaosVDEditor, Error, TEXT("[%s] Failed to apply geometry with key [%d] | Mesh component is not a ChaosVDGeometryDataComponent"), ANSI_TO_TCHAR(__FUNCTION__), GeometryKey);
+
+		bApplyMeshRequestProcessed = true;
+		return bApplyMeshRequestProcessed;
+	}
+
+	if (HasGeometryInCache(GeometryKey))
+	{
+		if (UDynamicMeshComponent* DynamicMeshComponent = Cast<UDynamicMeshComponent>(MeshComponent))
+		{
+			DynamicMeshComponent->SetDynamicMesh(GetCachedMeshForImplicit<UDynamicMesh>(GeometryKey));
+		}
+		else if (UStaticMeshComponent* StaticMeshComponent = Cast<UStaticMeshComponent>(MeshComponent))
+		{
+			StaticMeshComponent->SetStaticMesh(GetCachedMeshForImplicit<UStaticMesh>(GeometryKey));
+		}
+
 		DataComponent->SetIsMeshReady(true);
 		DataComponent->OnMeshReady()->Broadcast(*DataComponent);
+		bApplyMeshRequestProcessed = true;
 	}
+	else
+	{
+		UE_LOG(LogChaosVDEditor, Verbose, TEXT("[%s] Failed to apply geometry with key [%u] | Geometry was not ready"), ANSI_TO_TCHAR(__FUNCTION__), GeometryKey);
+	}
+
+	return bApplyMeshRequestProcessed;
 }
 
-bool FChaosVDGeometryBuilder::GameThreadTick(float DeltaTime)
+TSharedPtr<UE::Geometry::FMeshShapeGenerator> FChaosVDGeometryBuilder::CreateMeshGeneratorForImplicitObject(const Chaos::FImplicitObject* InImplicit)
 {
-	GeometryGenerationNotification.Update(MeshComponentsWaitingGeometryNum);
+	using namespace Chaos;
 
+	switch (GetInnerType(InImplicit->GetType()))
 	{
-		int32 CurrentGeometryProcessedNum = 0;
-		bool bCanContinueProcessing = true;
+		case ImplicitObjectType::Sphere:
+		{
+			if (const Chaos::TSphere<FReal, 3>* Sphere = InImplicit->template GetObject<Chaos::TSphere<FReal, 3>>())
+			{
+				TSharedPtr<UE::Geometry::FSphereGenerator> SphereGen = MakeShared<UE::Geometry::FSphereGenerator>();
+				SphereGen->Radius = Sphere->GetRadius();
+				SphereGen->NumTheta = 50;
+				SphereGen->NumPhi = 50;
+				SphereGen->bPolygroupPerQuad = false;
 
-		FWriteScopeLock WaitListWriteLock(GeometryWaitListRWLock);
-		for (FChaosVDWaitListMeshMap::TIterator MeshWaitListRemoveIterator = MeshComponentsWaitingForGeometryByKey.CreateIterator(); MeshWaitListRemoveIterator; ++MeshWaitListRemoveIterator)
+				return SphereGen;	
+			}
+			break;
+		}
+		case ImplicitObjectType::Box:
+		{
+			if (const Chaos::TBox<FReal, 3>* Box = InImplicit->template GetObject<Chaos::TBox<FReal, 3>>())
+			{
+				TSharedPtr<UE::Geometry::FMinimalBoxMeshGenerator> BoxGen = MakeShared<UE::Geometry::FMinimalBoxMeshGenerator>();
+				UE::Geometry::FOrientedBox3d OrientedBox;
+				OrientedBox.Frame = UE::Geometry::FFrame3d(Box->Center());
+				OrientedBox.Extents = Box->Extents() * 0.5;
+				BoxGen->Box = OrientedBox;
+				return BoxGen;
+			}
+			break;
+		}
+		case ImplicitObjectType::Capsule:
 		{
-			const uint32 GeometryKey = MeshWaitListRemoveIterator.Key();
+			if (const Chaos::FCapsule* Capsule = InImplicit->template GetObject<Chaos::FCapsule>())
+			{
+				TSharedPtr<UE::Geometry::FCapsuleGenerator> CapsuleGenerator = MakeShared<UE::Geometry::FCapsuleGenerator>();
+				CapsuleGenerator->Radius = FMath::Max(FMathf::ZeroTolerance, Capsule->GetRadius());
+				CapsuleGenerator->SegmentLength = FMath::Max(FMathf::ZeroTolerance, Capsule->GetSegment().GetLength());
+				CapsuleGenerator->NumHemisphereArcSteps = 12;
+				CapsuleGenerator->NumCircleSteps = 12;
 
+				return CapsuleGenerator;
+			}
+
+			break;
+		}
+		case ImplicitObjectType::Convex:
+		{
+			if (const Chaos::FConvex* Convex = InImplicit->template GetObject<Chaos::FConvex>())
 			{
-				FReadScopeLock CacheReadLocLock(GeometryCacheRWLock);
-				if (HasGeometryInCache_AssumesLocked(GeometryKey))
-				{
-					for (TArray<TWeakObjectPtr<UMeshComponent>>::TIterator MeshRemoveIterator = MeshWaitListRemoveIterator.Value().CreateIterator(); MeshRemoveIterator; ++MeshRemoveIterator)
-					{
-						CurrentGeometryProcessedNum++;
-						MeshComponentsWaitingGeometryNum -=1;
-						bCanContinueProcessing = CurrentGeometryProcessedNum < CVarChaosVDGeometryToProcessPerTick->GetInt();
-						ApplyMeshToComponentFromKey(*MeshRemoveIterator, GeometryKey);
-						MeshRemoveIterator.RemoveCurrent();
-					}
+				TSharedPtr<FChaosVDConvexMeshGenerator> ConvexMeshGen = MakeShared<FChaosVDConvexMeshGenerator>();
+				ConvexMeshGen->GenerateFromConvex(*Convex);
+				return ConvexMeshGen;
+			}
 				
-					if (!bCanContinueProcessing)
-					{
-						break;
-					}	
-				}
+			break;
+		}
+		case ImplicitObjectType::TriangleMesh:
+		{
+			if (const Chaos::FTriangleMeshImplicitObject* TriangleMesh = InImplicit->template GetObject<Chaos::FTriangleMeshImplicitObject>())
+			{
+				TSharedPtr<FChaosVDTriMeshGenerator> TriMeshGen = MakeShared<FChaosVDTriMeshGenerator>();
+				TriMeshGen->bReverseOrientation = true;
+				TriMeshGen->GenerateFromTriMesh(*TriangleMesh);
+				return TriMeshGen;
 			}
 
-			if (MeshWaitListRemoveIterator.Value().IsEmpty())
+			break;
+		}
+		case ImplicitObjectType::HeightField:
+		{
+			if (const Chaos::FHeightField* HeightField = InImplicit->template GetObject<Chaos::FHeightField>())
 			{
-				MeshWaitListRemoveIterator.RemoveCurrent();
+				TSharedPtr<FChaosVDHeightFieldMeshGenerator> HeightFieldMeshGen = MakeShared<FChaosVDHeightFieldMeshGenerator>();
+				HeightFieldMeshGen->bReverseOrientation = false;
+				HeightFieldMeshGen->GenerateFromHeightField(*HeightField);
+				return HeightFieldMeshGen;
 			}
+		
+			break;
+		}
+		case ImplicitObjectType::Plane:
+		case ImplicitObjectType::LevelSet:
+		case ImplicitObjectType::TaperedCylinder:
+		case ImplicitObjectType::Cylinder:
+		{
+			//TODO: Implement
+			break;
+		}
+		default:
+			break;
+	}
+
+	return nullptr;
+}
+
+const Chaos::FImplicitObject* FChaosVDGeometryBuilder::UnpackImplicitObject(const Chaos::FImplicitObject* InImplicitObject, Chaos::FRigidTransform3& InOutTransform) const
+{
+	using namespace Chaos;
 
-			if (!bCanContinueProcessing)
+	const EImplicitObjectType InnerType = GetInnerType(InImplicitObject->GetType());
+	switch (InnerType)
+	{
+		case ImplicitObjectType::Convex:
+			{
+				return GetGeometryBasedOnPackedType<FConvex>(InImplicitObject, InOutTransform, InImplicitObject->GetType());
+			}
+		case ImplicitObjectType::TriangleMesh:
+			{
+				return GetGeometryBasedOnPackedType<FTriangleMeshImplicitObject>(InImplicitObject, InOutTransform, InImplicitObject->GetType());
+			}
+		case ImplicitObjectType::HeightField:
+			{
+				return GetGeometryBasedOnPackedType<FHeightField>(InImplicitObject, InOutTransform, InImplicitObject->GetType());
+			}
+	default:
+			ensureMsgf(false, TEXT("Unpacking [%s] is not supported"), *GetImplicitObjectTypeName(InnerType).ToString());
+			break;
+	}
+
+	return nullptr;
+}
+
+void FChaosVDGeometryBuilder::AdjustedTransformForImplicit(const Chaos::FImplicitObject* InImplicit, FTransform& OutAdjustedTransform)
+{
+	using namespace Chaos;
+	switch (GetInnerType(InImplicit->GetType()))
+	{
+		// Currently, only capsules transforms needs to be re-adjusted
+		case ImplicitObjectType::Capsule:
+		{
+			if (const FCapsule* Capsule = InImplicit->template GetObject<FCapsule>())
 			{
-				break;
+				// Re-adjust the location so the pivot is not the center of the capsule, and transform it based on the provided transform
+				const FVector FinalLocation = OutAdjustedTransform.TransformPosition(Capsule->GetCenter() - Capsule->GetAxis() * Capsule->GetSegment().GetLength() * 0.5f);
+				const FQuat Rotation = FRotationMatrix::MakeFromZ(Capsule->GetAxis()).Rotator().Quaternion();
+
+				OutAdjustedTransform.SetRotation(OutAdjustedTransform.GetRotation() * Rotation);
+				OutAdjustedTransform.SetLocation(FinalLocation);
 			}
+			break;
 		}
+		default:
+			break;
+	}
+}
+
+bool FChaosVDGeometryBuilder::ImplicitObjectNeedsUnpacking(const Chaos::FImplicitObject* InImplicitObject) const
+{
+	using namespace Chaos;
+	const EImplicitObjectType InnerType = GetInnerType(InImplicitObject->GetType());
+
+	return InnerType == ImplicitObjectType::Convex || InnerType == ImplicitObjectType::TriangleMesh ||  InnerType == ImplicitObjectType::HeightField;
+}
+
+bool FChaosVDGeometryBuilder::GameThreadTick(float DeltaTime)
+{
+	int32 CurrentGeometryTasksProcessedNum = 0;
+
+	if (MeshComponentsWaitingForGeometry)
+	{
+		MeshComponentsWaitingForGeometry->ProcessWaitingObjects(CurrentGeometryTasksProcessedNum);
 	}
 
 	return true;
 }
 
-void FChaosVDGeometryBuilder::RegisterMeshComponentWaitingForGeometry(uint32 GeometryKey, TWeakObjectPtr<UMeshComponent> MesComponent, const int32 LODsToGenerateNum)
+void FChaosVDGeometryBuilder::AddMeshComponentWaitingForGeometry(uint32 GeometryKey, TWeakObjectPtr<UMeshComponent> MeshComponent) const
 {
-	if (!MesComponent.IsValid())
+	if (!MeshComponent.IsValid())
 	{
+		UE_LOG(LogChaosVDEditor, Verbose, TEXT("[%s] Failed to add mesh component update for geometry key [%d] | Mesh component is invalid"), ANSI_TO_TCHAR(__FUNCTION__), GeometryKey);
 		return;
 	}
 
-	const bool bIsMeshTypeLODCompatible = !MesComponent->IsA(UInstancedStaticMeshComponent::StaticClass()) && !MesComponent->IsA(UDynamicMeshComponent::StaticClass());
-	if (LODsToGenerateNum > 0 && !bIsMeshTypeLODCompatible)
+	if (!ensure(MeshComponentsWaitingForGeometry.IsValid()))
 	{
-		ensureMsgf(false, TEXT("LODs are currently not supported with the specified mesh component [%s] in the CVD tool | [%d] LODs were requested when 0 is expected"), *MesComponent->GetName(), LODsToGenerateNum);
+		UE_LOG(LogChaosVDEditor, Error, TEXT("[%s] Failed to add mesh component update for geometry key [%d] | WaitingListObject is invalid"), ANSI_TO_TCHAR(__FUNCTION__), GeometryKey);
+		return;
 	}
 
+	MeshComponentsWaitingForGeometry->AddObject(GeometryKey, MeshComponent);
+}
+
+void FChaosVDGeometryBuilder::RemoveMeshComponentWaitingForGeometry(uint32 GeometryKey, TWeakObjectPtr<UMeshComponent> MeshComponent) const
+{
+	if (!MeshComponent.IsValid())
+    {
+    	UE_LOG(LogChaosVDEditor, Verbose, TEXT("[%s] Failed to remove mesh component update for geometry key [%d] | Mesh component is invalid"), ANSI_TO_TCHAR(__FUNCTION__), GeometryKey);
+    	return;
+    }
+
+	if (!ensure(MeshComponentsWaitingForGeometry.IsValid()))
 	{
-		FWriteScopeLock WriteLock(GeometryWaitListRWLock);
-	
-		if (TArray<TWeakObjectPtr<UMeshComponent>>* MeshComponentsWaiting = MeshComponentsWaitingForGeometryByKey.Find(GeometryKey))
-		{
-			MeshComponentsWaiting->Add(MesComponent);
-		}
-		else
+		UE_LOG(LogChaosVDEditor, Error, TEXT("[%s] Failed to remove mesh component update for geometry key [%d] | WaitingListObject is invalid"), ANSI_TO_TCHAR(__FUNCTION__), GeometryKey);
+		return;
+	}
+
+	MeshComponentsWaitingForGeometry->RemoveObject(GeometryKey, MeshComponent);
+}
+
+void FChaosVDGeometryBuilder::HandleStaticMeshComponentInstanceIndexUpdated(UInstancedStaticMeshComponent* InComponent, TArrayView<const FInstancedStaticMeshDelegates::FInstanceIndexUpdateData> InIndexUpdates)
+{
+	if (IChaosVDGeometryComponent* DataComponent = Cast<IChaosVDGeometryComponent>(InComponent))
+	{
+		const TArrayView<TSharedPtr<FChaosVDMeshDataInstanceHandle>> MeshDataHandlesView = DataComponent->GetMeshDataInstanceHandles();
+		for (TSharedPtr<FChaosVDMeshDataInstanceHandle>& Handle : MeshDataHandlesView)
 		{
-			MeshComponentsWaitingForGeometryByKey.Add(GeometryKey, {MesComponent});
+			if (Handle)
+			{
+				Handle->HandleInstanceIndexUpdated(InIndexUpdates);
+			}
+			else
+			{
+				UE_LOG(LogChaosVDEditor, Error, TEXT("[%s] Failed to update Instance Index for component [%s] | Handle is in valid"), ANSI_TO_TCHAR(__FUNCTION__), *GetNameSafe(InComponent));
+			}
 		}
-
-		MeshComponentsWaitingGeometryNum += 1;
 	}
 }
diff --git a/Engine/Plugins/Experimental/ChaosVD/Source/ChaosVD/Private/ChaosVDGeometryDataComponent.cpp b/Engine/Plugins/Experimental/ChaosVD/Source/ChaosVD/Private/ChaosVDGeometryDataComponent.cpp
--- a/Engine/Plugins/Experimental/ChaosVD/Source/ChaosVD/Private/ChaosVDGeometryDataComponent.cpp
+++ b/Engine/Plugins/Experimental/ChaosVD/Source/ChaosVD/Private/ChaosVDGeometryDataComponent.cpp
@@ -6,88 +6,256 @@
 #include "ChaosVDGeometryBuilder.h"
 #include "ChaosVDModule.h"
 #include "ChaosVDParticleActor.h"
-#include "ChaosVDScene.h"
 #include "Components/MeshComponent.h"
-#include "MaterialDomain.h"
 #include "Materials/Material.h"
 #include "Materials/MaterialInterface.h"
 #include "Materials/MaterialInstanceDynamic.h"
 
-void FChaosVDGeometryDataComponentBase::UpdateVisibility_Internal(const FChaosVDShapeCollisionData& InCollisionData, UMeshComponent* MeshComponent)
+FChaosVDMeshDataInstanceHandle::FChaosVDMeshDataInstanceHandle(int32 InInstanceIndex, UMeshComponent* InMeshComponent, int32 InParticleID, int32 InSolverID)
 {
-	if (!InCollisionData.bIsValid)
+	MeshComponent = InMeshComponent;
+	MeshInstanceIndex = InInstanceIndex;
+	OwningParticleID = InParticleID;
+	OwningSolverID = InSolverID;
+
+	if (Cast<UInstancedStaticMeshComponent>(MeshComponent))
+	{
+		MeshComponentType = EChaosVDMeshComponent::InstancedStatic;
+	}
+	else if (Cast<UStaticMeshComponent>(MeshComponent))
+	{
+		MeshComponentType = EChaosVDMeshComponent::Static;
+	}
+	else
+	{
+		MeshComponentType = EChaosVDMeshComponent::Dynamic;
+	}
+}
+
+void FChaosVDMeshDataInstanceHandle::SetWorldTransform(const FTransform& InTransform)
+{
+	if (!ExtractedGeometryHandle)
 	{
+		UE_LOG(LogChaosVDEditor, Error, TEXT("[%s] Attempted to update the world transform without a valid geometry handle"), ANSI_TO_TCHAR(__FUNCTION__));
 		return;
 	}
 
-	if (MeshComponent)
+	const FTransform ExtractedRelativeTransform = ExtractedGeometryHandle->GetRelativeTransform();
+
+	CurrentWorldTransform.SetLocation(InTransform.TransformPosition(ExtractedRelativeTransform.GetLocation()));
+	CurrentWorldTransform.SetRotation(InTransform.TransformRotation(ExtractedRelativeTransform.GetRotation()));
+	CurrentWorldTransform.SetScale3D(ExtractedGeometryHandle->GetRelativeTransform().GetScale3D());
+	
+	if (IChaosVDGeometryComponent* CVDGeometryComponent = Cast<IChaosVDGeometryComponent>(GetMeshComponent()))
+	{
+		CVDGeometryComponent->UpdateInstanceWorldTransform(AsShared(), CurrentWorldTransform);
+	}
+}
+
+void FChaosVDMeshDataInstanceHandle::SetInstanceColor(const FLinearColor& NewColor)
+{
+	if (CurrentGeometryColor != NewColor)
 	{
-		if (const UChaosVDEditorSettings* EditorSettings = GetDefault<UChaosVDEditorSettings>())
+		if (IChaosVDGeometryComponent* CVDGeometryComponent = Cast<IChaosVDGeometryComponent>(GetMeshComponent()))
 		{
-			const EChaosVDGeometryVisibilityFlags CurrentVisibilityFlags = static_cast<EChaosVDGeometryVisibilityFlags>(EditorSettings->GeometryVisibilityFlags);
-			
-			bool bShouldGeometryBeVisible = false;
+			CVDGeometryComponent->UpdateInstanceColor(AsShared(), NewColor);
+			CurrentGeometryColor = NewColor;
+		}
+	}
+}
+
+void FChaosVDMeshDataInstanceHandle::UpdateMeshComponentForCollisionData(const FChaosVDShapeCollisionData& InCollisionData)
+{
+	if (InCollisionData.bIsValid && CollisionData != InCollisionData)
+	{
+		if (const TSharedPtr<FChaosVDGeometryBuilder> GeometryBuilderPtr = GeometryBuilderInstance.Pin())
+		{
+			EChaosVDMeshAttributesFlags RequiredMeshAttributes = EChaosVDMeshAttributesFlags::None;
+
+			// If this is a query only type of geometry, we need a translucent mesh
+			if (InCollisionData.bQueryCollision && !InCollisionData.bSimCollision)
+			{
+				EnumAddFlags(RequiredMeshAttributes, EChaosVDMeshAttributesFlags::TranslucentGeometry);
+			}
+
+			// Mirrored geometry needs to be on a instanced mesh component with reversed culling
+			if (GeometryBuilderPtr->HasNegativeScale(ExtractedGeometryHandle->GetRelativeTransform()))
+			{
+				EnumAddFlags(RequiredMeshAttributes, EChaosVDMeshAttributesFlags::MirroredGeometry);
+			}
+
+			// If the current mesh component does not meet the required mesh attributes, we need to move to a new mesh component that it does
+			bool bMeshComponentWasUpdated = false;
+			if (IChaosVDGeometryComponent* CVDOldGeometryComponent = Cast<IChaosVDGeometryComponent>(GetMeshComponent()))
+			{
+				if (RequiredMeshAttributes != CVDOldGeometryComponent->GetMeshComponentAttributeFlags())
+				{
+					CVDOldGeometryComponent->RemoveMeshInstance(AsShared());
+
+					if (bIsSelected)
+					{
+						CVDOldGeometryComponent->SetIsSelected(AsShared(), false);
+					}
+
+					GeometryBuilderPtr->UpdateMeshDataInstance<UChaosVDInstancedStaticMeshComponent>(AsShared(), RequiredMeshAttributes);
+
+					bMeshComponentWasUpdated = true;
+				}
+			}
 
-			if (!EnumHasAnyFlags(CurrentVisibilityFlags, EChaosVDGeometryVisibilityFlags::ShowDisabledParticles))
+			if (bMeshComponentWasUpdated)
 			{
-				//TODO: We should use IChaosVDParticleVisualizationDataProvider instead, which AChaosVDParticleActor implements already
-				// but it is not an uinterface
-				if (AChaosVDParticleActor* ParticleActor = Cast<AChaosVDParticleActor>(MeshComponent->GetOwner()))
+				if (IChaosVDGeometryComponent* CVDNewGeometryComponent = Cast<IChaosVDGeometryComponent>(GetMeshComponent()))
 				{
-					if (const FChaosVDParticleDataWrapper* ParticleData = ParticleActor->GetParticleData())
+					if (CurrentGeometryColor != FLinearColor(ForceInitToZero))
 					{
-						if (ParticleData->ParticleDynamicsMisc.HasValidData() && ParticleData->ParticleDynamicsMisc.bDisabled)
-						{
-							MeshComponent->SetVisibility(bShouldGeometryBeVisible);
-							return;
-						}
+						CVDNewGeometryComponent->UpdateInstanceColor(AsShared(), CurrentGeometryColor);
 					}
+						
+					CVDNewGeometryComponent->UpdateInstanceVisibility(AsShared(), bIsVisible);
+					CVDNewGeometryComponent->SetIsSelected(AsShared(), bIsSelected);
 				}
 			}
+		}
+	}
+}
+
+void FChaosVDMeshDataInstanceHandle::SetGeometryCollisionData(const FChaosVDShapeCollisionData& InCollisionData)
+{
+	// If this is a static mesh component, we can't just update change the material. We need to remove this instance from the current component and move it to a
+	// component that has the correct translucent mesh
+	if (GetMeshComponentType() == EChaosVDMeshComponent::InstancedStatic)
+	{
+		UpdateMeshComponentForCollisionData(InCollisionData);
+	}
 
-			// TODO: Re-visit the way we determine visibility of the meshes.
-			// Now that the options have grown and they will continue to do so, these checks are becoming hard to read and extend
+	CollisionData = InCollisionData;
+}
 
-			const bool bIsHeightfield = GetImplicitObject() && Chaos::GetInnerType(GetImplicitObject()->GetType()) == Chaos::ImplicitObjectType::HeightField;
+void FChaosVDMeshDataInstanceHandle::SetIsSelected(bool bInIsSelected)
+{
+	if (IChaosVDGeometryComponent* CVDGeometryComponent = Cast<IChaosVDGeometryComponent>(GetMeshComponent()))
+	{
+		CVDGeometryComponent->SetIsSelected(AsShared(), bInIsSelected);
+	}
 
-			if (bIsHeightfield && EnumHasAnyFlags(CurrentVisibilityFlags, EChaosVDGeometryVisibilityFlags::ShowHeightfields))
+	bIsSelected = bInIsSelected;
+}
+
+void FChaosVDMeshDataInstanceHandle::SetVisibility(bool bInIsVisible)
+{
+	if (IChaosVDGeometryComponent* CVDGeometryComponent = Cast<IChaosVDGeometryComponent>(GetMeshComponent()))
+	{
+		CVDGeometryComponent->UpdateInstanceVisibility(AsShared(), bInIsVisible);
+	}
+
+	bIsVisible = bInIsVisible;
+}
+
+UMaterialInstanceDynamic* FChaosVDMeshDataInstanceHandle::GetCachedMaterialInstance(EChaosVDMaterialType Type)
+{
+	const UChaosVDEditorSettings* EditorSettings = GetDefault<UChaosVDEditorSettings>();
+	if (!EditorSettings)
+	{
+		return nullptr;
+	}
+
+	if (const TStrongObjectPtr<UMaterialInstanceDynamic>* MaterialInstance = MaterialInstancesByID.Find(Type))
+	{
+		return MaterialInstance->Get();
+	}
+	else
+	{
+		UMaterialInterface* MaterialToCreate = nullptr;
+		switch(Type)
+		{
+		case EChaosVDMaterialType::QueryOnlyMaterial:
 			{
-				bShouldGeometryBeVisible = true;
+				MaterialToCreate = EditorSettings->QueryOnlyMeshesMaterial.Get();
+				break;
 			}
-			else
+		case EChaosVDMaterialType::SimOnlyMaterial:
 			{
-				// Complex vs Simple takes priority although this is subject to change
-				const bool bShouldBeVisibleIfComplex = InCollisionData.bIsComplex && EnumHasAnyFlags(CurrentVisibilityFlags, EChaosVDGeometryVisibilityFlags::Complex);
-				const bool bShouldBeVisibleIfSimple = !InCollisionData.bIsComplex && EnumHasAnyFlags(CurrentVisibilityFlags, EChaosVDGeometryVisibilityFlags::Simple);
-			
-				if (bShouldBeVisibleIfComplex || bShouldBeVisibleIfSimple)
+				MaterialToCreate = EditorSettings->SimOnlyMeshesMaterial.Get();
+				break;
+			}
+		case EChaosVDMaterialType::Instanced:
+			{
+				MaterialToCreate = EditorSettings->InstancedMeshesMaterial.Get();
+				break;
+			}
+		case EChaosVDMaterialType::InstancedQueryOnly:
+			{
+				MaterialToCreate = EditorSettings->InstancedMeshesQueryOnlyMaterial.Get();
+				break;
+			}
+		}
+		
+		if (MaterialToCreate)
+		{
+			UMaterialInstanceDynamic* DynamicMaterial = UMaterialInstanceDynamic::Create(MaterialToCreate, nullptr);
+			MaterialInstancesByID.Add(Type, TStrongObjectPtr<UMaterialInstanceDynamic>(DynamicMaterial));
+
+			return DynamicMaterial;
+		}
+	}
+
+	return nullptr;
+}
+
+void FChaosVDMeshDataInstanceHandle::HandleInstanceIndexUpdated(TArrayView<const FInstancedStaticMeshDelegates::FInstanceIndexUpdateData> InIndexUpdates)
+{
+	// When an index changes, we receive an array with all the indexes that were modified. We need to only act upon the update of for the index we are tracking in this handle
+	for (const FInstancedStaticMeshDelegates::FInstanceIndexUpdateData& IndexUpdateData : InIndexUpdates)
+	{
+		switch (IndexUpdateData.Type)
+		{
+			case FInstancedStaticMeshDelegates::EInstanceIndexUpdateType::Added:
+				break; // We don't need to process 'Added' updates as they can't affect existing IDs
+			case FInstancedStaticMeshDelegates::EInstanceIndexUpdateType::Relocated:
 				{
-					bShouldGeometryBeVisible = (InCollisionData.bSimCollision && EnumHasAnyFlags(CurrentVisibilityFlags, EChaosVDGeometryVisibilityFlags::Simulated))
-					|| (InCollisionData.bQueryCollision && EnumHasAnyFlags(CurrentVisibilityFlags, EChaosVDGeometryVisibilityFlags::Query));
+					if (MeshInstanceIndex == IndexUpdateData.OldIndex)
+					{
+						MeshInstanceIndex = IndexUpdateData.Index;
+					}
+					break;
 				}
-			}
 
-			MeshComponent->SetVisibility(bShouldGeometryBeVisible);
+			case FInstancedStaticMeshDelegates::EInstanceIndexUpdateType::Removed:
+			case FInstancedStaticMeshDelegates::EInstanceIndexUpdateType::Cleared:
+			case FInstancedStaticMeshDelegates::EInstanceIndexUpdateType::Destroyed:
+				{
+					if (MeshInstanceIndex == IndexUpdateData.Index)
+					{
+						MeshInstanceIndex = INDEX_NONE;
+					}
+					break;
+				}
+			default:
+				break;
 		}
 	}
 }
 
-void FChaosVDGeometryDataComponentBase::UpdateDataFromShapeArray_Internal(const TArray<FChaosVDShapeCollisionData>& InShapeArray, FChaosVDShapeCollisionData& CollisionDataToUpdate, UMeshComponent* MeshComponent)
+void FChaosVDGeometryComponentUtils::UpdateCollisionDataFromShapeArray(const TArray<FChaosVDShapeCollisionData>& InShapeArray, const TSharedPtr<FChaosVDMeshDataInstanceHandle>& InInstanceHandle)
 {
-	if (!ensureMsgf(RootImplicitObject.IsValid(), TEXT("Tried to Update Collision Data without a valid Implicit Object")))
+	if (!ensureMsgf(InInstanceHandle && InInstanceHandle->GetGeometryHandle() && InInstanceHandle->GetGeometryHandle()->GetRootImplicitObject(), TEXT("Tried to Update Collision Data without a valid Implicit Object")))
 	{
 		return;
 	}
 
-	RootImplicitObject.GetReference()->VisitObjects([this, &CollisionDataToUpdate, &InShapeArray] (const Chaos::FImplicitObject* ImplicitA, const Chaos::FRigidTransform3& RelativeTransformA, const int32 RootObjectIndexA, const int32 ObjectIndex, const int32 LeafObjectIndexA)
+	FChaosVDShapeCollisionData CollisionDataToUpdate = InInstanceHandle->GetGeometryCollisionData();
+
+	InInstanceHandle->GetGeometryHandle()->GetRootImplicitObject()->VisitObjects([InInstanceHandle, &CollisionDataToUpdate, &InShapeArray] (const Chaos::FImplicitObject* ImplicitA, const Chaos::FRigidTransform3& RelativeTransformA, const int32 RootObjectIndexA, const int32 ObjectIndex, const int32 LeafObjectIndexA)
 	{
 		if (!InShapeArray.IsValidIndex(LeafObjectIndexA))
 		{
 			return true;
 		}
 
-		if (ImplicitA == ImplicitObject)
+		if (ImplicitA == InInstanceHandle->GetGeometryHandle()->GetImplicitObject())
 		{
+			
 			CollisionDataToUpdate = InShapeArray[LeafObjectIndexA];
 			CollisionDataToUpdate.bIsComplex = FChaosVDGeometryBuilder::DoesImplicitContainType(ImplicitA, Chaos::ImplicitObjectType::HeightField) || FChaosVDGeometryBuilder::DoesImplicitContainType(ImplicitA, Chaos::ImplicitObjectType::TriangleMesh);
 			CollisionDataToUpdate.bIsValid = true;
@@ -96,159 +264,137 @@ void FChaosVDGeometryDataComponentBase::UpdateDataFromShapeArray_Internal(const
 		return true;
 	});
 
-	// If our collision data was successfully updated and we have a valid mesh component, set the correct material
-	if (CollisionDataToUpdate.bIsValid && MeshComponent)
+	InInstanceHandle->SetGeometryCollisionData(CollisionDataToUpdate);
+}
+
+void FChaosVDGeometryComponentUtils::UpdateMeshColor(const TSharedPtr<FChaosVDMeshDataInstanceHandle>& InInstanceHandle, const FChaosVDParticleDataWrapper& InParticleData, bool bIsServer)
+{
+	const FChaosVDShapeCollisionData ShapeData = InInstanceHandle->GetGeometryCollisionData();
+	const bool bIsQueryOnly = ShapeData.bQueryCollision && !ShapeData.bSimCollision;
+
+	if (ShapeData.bIsValid)
 	{
-		if (const UChaosVDEditorSettings* EditorSettings = GetDefault<UChaosVDEditorSettings>())
-		{
-			const bool bIsQueryOnly = CollisionDataToUpdate.bQueryCollision && !CollisionDataToUpdate.bSimCollision;
+		FLinearColor ColorToApply = GetGeometryParticleColor(InInstanceHandle->GetGeometryHandle(), InParticleData, bIsServer);
 
-			UMaterialInterface* MaterialToApply = bIsQueryOnly ? GetCachedMaterialInstance(EChaosVDMaterialType::QueryOnlyMaterial) : GetCachedMaterialInstance(EChaosVDMaterialType::SimOnlyMaterial);
+		constexpr float QueryOnlyShapeOpacity = 0.6f;
+		ColorToApply.A = bIsQueryOnly ? QueryOnlyShapeOpacity : 1.0f;
 
-			if (MaterialToApply)
-			{
-				MeshComponent->SetMaterial(0, MaterialToApply);
-			}
-			else
-			{
-				ensure(false);
-				UE_LOG(LogChaosVDEditor, Error, TEXT("[%s] Failed to get Query Only material for, applying the default mesh to all geometry"), ANSI_TO_TCHAR(__FUNCTION__));
-				
-				MeshComponent->SetMaterial(0, UMaterial::GetDefaultMaterial(EMaterialDomain::MD_Surface));
-			}
-		}
+		InInstanceHandle->SetInstanceColor(ColorToApply);
 	}
 }
 
-void FChaosVDGeometryDataComponentBase::UpdateColors_Internal(UMeshComponent* MeshComponent)
+void FChaosVDGeometryComponentUtils::UpdateMeshVisibility(const TSharedPtr<FChaosVDMeshDataInstanceHandle>& InInstanceHandle, const FChaosVDParticleDataWrapper& InParticleData, bool bIsActive)
 {
-	if (!MeshComponent)
-	{
-		return;
-	}
-
-	//TODO: Remove this direct dependency to AChaosVDParticleACtor
-	AChaosVDParticleActor* ParticleActor = Cast<AChaosVDParticleActor>(MeshComponent->GetOwner());
-	if (!ParticleActor)
+	if (!InInstanceHandle || !InInstanceHandle->GetGeometryHandle())
 	{
 		return;
 	}
 
-	const FChaosVDParticleDataWrapper* ParticleData = ParticleActor->GetParticleData();
-	if (!ParticleData)
+	if (!bIsActive)
 	{
+		InInstanceHandle->SetVisibility(bIsActive);
 		return;
 	}
 
-	const UChaosVDEditorSettings* EditorSettings = GetDefault<UChaosVDEditorSettings>();
-	if (!EditorSettings)
+	if (const UChaosVDEditorSettings* EditorSettings = GetDefault<UChaosVDEditorSettings>())
 	{
-		return;
-	}
-
-	constexpr FLinearColor DefaultColor(0.088542f, 0.088542f, 0.088542f);
-	FLinearColor ColorToApply = DefaultColor;
+		const EChaosVDGeometryVisibilityFlags CurrentVisibilityFlags = static_cast<EChaosVDGeometryVisibilityFlags>(EditorSettings->GeometryVisibilityFlags);
+		
+		bool bShouldGeometryBeVisible = false;
 
-	switch(EditorSettings->ParticleColorMode)
-	{
-		case EChaosVDParticleDebugColorMode::ShapeType:
-			{
-				ColorToApply = GetImplicitObject() ? EditorSettings->ColorsByShapeType.GetColorFromShapeType(Chaos::GetInnerType(GetImplicitObject()->GetType())) : DefaultColor;
-				break;
-			}
-		case EChaosVDParticleDebugColorMode::State:
+		if (!EnumHasAnyFlags(CurrentVisibilityFlags, EChaosVDGeometryVisibilityFlags::ShowDisabledParticles))
+		{
+			if (InParticleData.ParticleDynamicsMisc.HasValidData() && InParticleData.ParticleDynamicsMisc.bDisabled)
 			{
-				if (ParticleData->Type == EChaosVDParticleType::Static)
-				{
-					ColorToApply = EditorSettings->ColorsByParticleState.GetColorFromState(EChaosVDObjectStateType::Static);
-				}
-				else
-				{
-					ColorToApply = EditorSettings->ColorsByParticleState.GetColorFromState(ParticleData->ParticleDynamicsMisc.MObjectState);
-				}
-				break;
+				InInstanceHandle->SetVisibility(bShouldGeometryBeVisible);
+				return;
 			}
-		case EChaosVDParticleDebugColorMode::ClientServer:
+		}
+
+		// TODO: Re-visit the way we determine visibility of the meshes.
+		// Now that the options have grown and they will continue to do so, these checks are becoming hard to read and extend
+
+		const bool bIsHeightfield = InInstanceHandle->GetGeometryHandle()->GetImplicitObject() && Chaos::GetInnerType(InInstanceHandle->GetGeometryHandle()->GetImplicitObject()->GetType()) == Chaos::ImplicitObjectType::HeightField;
+
+		if (bIsHeightfield && EnumHasAnyFlags(CurrentVisibilityFlags, EChaosVDGeometryVisibilityFlags::ShowHeightfields))
+		{
+			bShouldGeometryBeVisible = true;
+		}
+		else
+		{
+			const FChaosVDShapeCollisionData InstanceShapeData = InInstanceHandle->GetGeometryCollisionData();
+
+			if (InstanceShapeData.bIsValid)
 			{
-				const TSharedPtr<FChaosVDScene> Scene = ParticleActor->GetScene().Pin();
-				const bool bIsServer = (Scene.IsValid()) ? Scene->IsSolverForServer(ParticleData->SolverID) : false;
-				if (ParticleData->Type == EChaosVDParticleType::Static)
-				{
-					ColorToApply = EditorSettings->ColorsByClientServer.GetColorFromState(bIsServer, EChaosVDObjectStateType::Static);
-				}
-				else
+				// Complex vs Simple takes priority although this is subject to change
+				const bool bShouldBeVisibleIfComplex = InstanceShapeData.bIsComplex && EnumHasAnyFlags(CurrentVisibilityFlags, EChaosVDGeometryVisibilityFlags::Complex);
+				const bool bShouldBeVisibleIfSimple = !InstanceShapeData.bIsComplex && EnumHasAnyFlags(CurrentVisibilityFlags, EChaosVDGeometryVisibilityFlags::Simple);
+		
+				if (bShouldBeVisibleIfComplex || bShouldBeVisibleIfSimple)
 				{
-					ColorToApply = EditorSettings->ColorsByClientServer.GetColorFromState(bIsServer, ParticleData->ParticleDynamicsMisc.MObjectState);
+					bShouldGeometryBeVisible = (InstanceShapeData.bSimCollision && EnumHasAnyFlags(CurrentVisibilityFlags, EChaosVDGeometryVisibilityFlags::Simulated))
+					|| (InstanceShapeData.bQueryCollision && EnumHasAnyFlags(CurrentVisibilityFlags, EChaosVDGeometryVisibilityFlags::Query));
 				}
-				break;
 			}
+		}
 
-		case EChaosVDParticleDebugColorMode::None:
-		default:
-			// Nothing to do here. Color to apply is already set to the default
-			break;
-	}
-
-	if (CurrentGeometryColor == ColorToApply)
-	{
-		return;
-	}
-	
-	if (UMaterialInstanceDynamic* DynamicMaterial = Cast<UMaterialInstanceDynamic>(MeshComponent->GetMaterial(0)))
-	{
-		DynamicMaterial->SetVectorParameterValue(TEXT("BaseColor"), ColorToApply);
-		CurrentGeometryColor = ColorToApply;
+		InInstanceHandle->SetVisibility(bShouldGeometryBeVisible);
 	}
 }
 
-void FChaosVDGeometryDataComponentBase::SetImplicitObject_Internal(const Chaos::FImplicitObject* InImplicitObject)
+FLinearColor FChaosVDGeometryComponentUtils::GetGeometryParticleColor(const TSharedPtr<FChaosVDExtractedGeometryDataHandle>& InGeometryHandle, const FChaosVDParticleDataWrapper& InParticleData, bool bIsServer)
 {
-	ImplicitObject = InImplicitObject;
-}
+	constexpr FLinearColor DefaultColor(0.088542f, 0.088542f, 0.088542f);
+	FLinearColor ColorToApply = DefaultColor;
 
-UMaterialInstanceDynamic* FChaosVDGeometryDataComponentBase::GetCachedMaterialInstance(EChaosVDMaterialType Type)
-{
-	const UChaosVDEditorSettings* EditorSettings = GetDefault<UChaosVDEditorSettings>();
-	if (!EditorSettings)
+	if (!InGeometryHandle)
 	{
-		return nullptr;
+		return ColorToApply;
 	}
 
-	if (const TStrongObjectPtr<UMaterialInstanceDynamic>* MaterialInstance = MaterialInstancesByID.Find(Type))
+	const UChaosVDEditorSettings* EditorSettings = GetDefault<UChaosVDEditorSettings>();
+	if (!EditorSettings)
 	{
-		return MaterialInstance->Get();
+		return ColorToApply;
 	}
-	else
+
+	switch (EditorSettings->ParticleColorMode)
 	{
-		UMaterialInterface* MaterialToCreate = nullptr;
-		switch(Type)
+	case EChaosVDParticleDebugColorMode::ShapeType:
 		{
-			case EChaosVDMaterialType::QueryOnlyMaterial:
-				{
-					MaterialToCreate = EditorSettings->QueryOnlyMeshesMaterial.Get();
-					break;
-				}
-			case EChaosVDMaterialType::SimOnlyMaterial:
-				{
-					MaterialToCreate = EditorSettings->SimOnlyMeshesMaterial.Get();
-					break;
-				}
+			ColorToApply = InGeometryHandle->GetImplicitObject() ? EditorSettings->ColorsByShapeType.GetColorFromShapeType(Chaos::GetInnerType(InGeometryHandle->GetImplicitObject()->GetType())) : DefaultColor;
+			break;
 		}
-		
-		if (MaterialToCreate)
+	case EChaosVDParticleDebugColorMode::State:
 		{
-			UMaterialInstanceDynamic* DynamicMaterial = UMaterialInstanceDynamic::Create(MaterialToCreate, nullptr);
-			MaterialInstancesByID.Add(Type, TStrongObjectPtr<UMaterialInstanceDynamic>(DynamicMaterial));
-
-			return DynamicMaterial;
+			if (InParticleData.Type == EChaosVDParticleType::Static)
+			{
+				ColorToApply = EditorSettings->ColorsByParticleState.GetColorFromState(EChaosVDObjectStateType::Static);
+			}
+			else
+			{
+				ColorToApply = EditorSettings->ColorsByParticleState.GetColorFromState(InParticleData.ParticleDynamicsMisc.MObjectState);
+			}
+			break;
+		}
+	case EChaosVDParticleDebugColorMode::ClientServer:
+		{
+			if (InParticleData.Type == EChaosVDParticleType::Static)
+			{
+				ColorToApply = EditorSettings->ColorsByClientServer.GetColorFromState(bIsServer, EChaosVDObjectStateType::Static);
+			}
+			else
+			{
+				ColorToApply = EditorSettings->ColorsByClientServer.GetColorFromState(bIsServer, InParticleData.ParticleDynamicsMisc.MObjectState);
+			}
+			break;
 		}
-	}
 
-	return nullptr;
-}
+	case EChaosVDParticleDebugColorMode::None:
+	default:
+		// Nothing to do here. Color to apply is already set to the default
+		break;
+	}
 
-const Chaos::FImplicitObject* FChaosVDGeometryDataComponentBase::GetImplicitObject() const
-{
-	// If the root object is no longer valid, the implicit object ptr we have is probably garbage
-	return RootImplicitObject.IsValid() ? ImplicitObject : nullptr;
+	return ColorToApply;
 }
diff --git a/Engine/Plugins/Experimental/ChaosVD/Source/ChaosVD/Private/ChaosVDMeshComponentPool.cpp b/Engine/Plugins/Experimental/ChaosVD/Source/ChaosVD/Private/ChaosVDMeshComponentPool.cpp
new file mode 100644
--- /dev/null
+++ b/Engine/Plugins/Experimental/ChaosVD/Source/ChaosVD/Private/ChaosVDMeshComponentPool.cpp
@@ -0,0 +1,89 @@
+﻿// Copyright Epic Games, Inc. All Rights Reserved.
+
+#include "ChaosVDMeshComponentPool.h"
+
+#include "ChaosVDGeometryDataComponent.h"
+#include "ChaosVDModule.h"
+#include "UObject/Package.h"
+
+bool FChaosVDMeshComponentPool::bUseComponentsPool = true;
+FAutoConsoleVariableRef FChaosVDMeshComponentPool::CVarUseComponentsPool(TEXT("p.Chaos.VD.Tool.UseComponentsPool"), FChaosVDMeshComponentPool::bUseComponentsPool, TEXT("Set to false to disable the use of a pool system for Mesh Components."));
+
+void FChaosVDMeshComponentPool::AddReferencedObjects(FReferenceCollector& Collector)
+{
+	Collector.AddReferencedObjects(PooledDynamicMeshComponent);
+	Collector.AddReferencedObjects(PooledInstancedStaticMeshComponent);
+	Collector.AddReferencedObjects(PooledStaticMeshComponent);
+}
+
+void FChaosVDMeshComponentPool::DisposeMeshComponent(UMeshComponent* MeshComponent)
+{
+	if (!MeshComponent)
+	{
+		return;
+	}
+
+	if (!bUseComponentsPool)
+	{
+		MeshComponent->DestroyComponent();
+		return;
+	}
+
+	ResetMeshComponent(MeshComponent);
+
+	if (UInstancedStaticMeshComponent* AsISMC = Cast<UInstancedStaticMeshComponent>(MeshComponent))
+	{
+		PooledInstancedStaticMeshComponent.Add(AsISMC);
+	}
+	else if (UStaticMeshComponent* AsSMC = Cast<UStaticMeshComponent>(MeshComponent))
+	{
+		PooledStaticMeshComponent.Add(AsSMC);
+	}
+	else if (UDynamicMeshComponent* AsDMC = Cast<UDynamicMeshComponent>(MeshComponent))
+	{
+		PooledDynamicMeshComponent.Add(AsDMC);
+	}
+	else
+	{
+		UE_LOG(LogChaosVDEditor, Warning, TEXT("[%s] Unable top dispose Mesh component [%s] | Incompatible mesh type used [%s]"), ANSI_TO_TCHAR(__FUNCTION__), *GetNameSafe(MeshComponent), *GetNameSafe(MeshComponent->GetClass()));
+	}
+
+	const FString NewName = MeshComponent->GetName() + TEXT("_POOLED_") + FGuid::NewGuid().ToString();
+	MeshComponent->Rename(*NewName, GetTransientPackage(), REN_NonTransactional | REN_DoNotDirty | REN_ForceNoResetLoaders | REN_SkipGeneratedClasses | REN_DontCreateRedirectors);
+}
+
+void FChaosVDMeshComponentPool::ResetMeshComponent(UMeshComponent* MeshComponent)
+{
+	if (!MeshComponent)
+	{
+		return;
+	}
+
+	if (UDynamicMeshComponent* DynamicMeshComponent = Cast<UDynamicMeshComponent>(MeshComponent))
+	{
+		DynamicMeshComponent->SetDynamicMesh(nullptr);
+	}
+	else if (UInstancedStaticMeshComponent* InstancedStaticMeshComponent = Cast<UInstancedStaticMeshComponent>(MeshComponent))
+	{
+		InstancedStaticMeshComponent->SetStaticMesh(nullptr);
+		InstancedStaticMeshComponent->ClearInstances();
+	}
+	else if (UStaticMeshComponent* StaticMeshComponent = Cast<UStaticMeshComponent>(MeshComponent))
+	{
+		StaticMeshComponent->SetStaticMesh(nullptr);
+	}
+	
+	MeshComponent->SetRelativeTransform(FTransform::Identity);
+	
+	if (IChaosVDGeometryComponent* DataComponent = Cast<IChaosVDGeometryComponent>(MeshComponent))
+	{
+		DataComponent->Reset();
+	}
+
+	MeshComponent->UnregisterComponent();
+
+	if (AActor* Owner = MeshComponent->GetOwner())
+	{
+		Owner->RemoveOwnedComponent(MeshComponent);
+	}
+}
diff --git a/Engine/Plugins/Experimental/ChaosVD/Source/ChaosVD/Private/ChaosVDModule.cpp b/Engine/Plugins/Experimental/ChaosVD/Source/ChaosVD/Private/ChaosVDModule.cpp
--- a/Engine/Plugins/Experimental/ChaosVD/Source/ChaosVD/Private/ChaosVDModule.cpp
+++ b/Engine/Plugins/Experimental/ChaosVD/Source/ChaosVD/Private/ChaosVDModule.cpp
@@ -19,6 +19,15 @@
 
 DEFINE_LOG_CATEGORY(LogChaosVDEditor);
 
+FAutoConsoleCommand ChaosVDSpawnNewCVDInstance(
+	TEXT("p.Chaos.VD.SpawnNewCVDInstance"),
+	TEXT("Opens a new CVD windows wothout closing an existing one"),
+	FConsoleCommandWithArgsDelegate::CreateLambda([](const TArray<FString>& Args)
+	{
+		FChaosVDModule::Get().SpawnCVDTab();
+	})
+);
+
 FChaosVDModule& FChaosVDModule::Get()
 {
 	return FModuleManager::Get().LoadModuleChecked<FChaosVDModule>(TEXT("ChaosVD"));
@@ -47,6 +56,11 @@ void FChaosVDModule::ShutdownModule()
 	FChaosVDStyle::Shutdown();
 
 	FGlobalTabmanager::Get()->UnregisterNomadTabSpawner(FChaosVDTabID::ChaosVisualDebuggerTab);
+
+	for (const FName TabID : CreatedExtraTabSpawnersIDs)
+	{
+		FGlobalTabmanager::Get()->UnregisterNomadTabSpawner(TabID);
+	}
 	
 	FCoreDelegates::OnEnginePreExit.RemoveAll(this);
 
@@ -62,6 +76,22 @@ void FChaosVDModule::RegisterClassesCustomDetails() const
 	PropertyModule.RegisterCustomPropertyTypeLayout("ChaosVDParticleDataWrapper", FOnGetPropertyTypeCustomizationInstance::CreateStatic(&FChaosVDParticleDataWrapperCustomization::MakeInstance));
 }
 
+void FChaosVDModule::SpawnCVDTab()
+{
+	// Registering new tab spawners with random names is not the best idea to spawn new tabs, but it is good enough to test we can run multiple instances of CVD withing the editor.
+	// This is also why spawning new tabs it is only exposed via console commands for now.
+	// When this feature is deemed stable and exposed in the UI I will investigate a more correct way of implementing this
+	const FName NewTabID = FName(FChaosVDTabID::ChaosVisualDebuggerTab.ToString() + FGuid::NewGuid().ToString());
+
+	FGlobalTabmanager::Get()->RegisterNomadTabSpawner(NewTabID, FOnSpawnTab::CreateRaw(this, &FChaosVDModule::SpawnMainTab))
+						.SetDisplayName(LOCTEXT("VisualDebuggerTabTitle", "Chaos Visual Debugger"))
+						.SetTooltipText(LOCTEXT("VisualDebuggerTabDesc", "Opens the Chaos Visual Debugger window"));
+
+	FGlobalTabmanager::Get()->TryInvokeTab(NewTabID);
+
+	CreatedExtraTabSpawnersIDs.Add(NewTabID);
+}
+
 TSharedRef<SDockTab> FChaosVDModule::SpawnMainTab(const FSpawnTabArgs& Args)
 {
 	TSharedRef<SDockTab> MainTabInstance =
diff --git a/Engine/Plugins/Experimental/ChaosVD/Source/ChaosVD/Private/ChaosVDParticleActor.cpp b/Engine/Plugins/Experimental/ChaosVD/Source/ChaosVD/Private/ChaosVDParticleActor.cpp
--- a/Engine/Plugins/Experimental/ChaosVD/Source/ChaosVD/Private/ChaosVDParticleActor.cpp
+++ b/Engine/Plugins/Experimental/ChaosVD/Source/ChaosVD/Private/ChaosVDParticleActor.cpp
@@ -18,7 +18,6 @@
 #include "Visualizers/ChaosVDParticleDataVisualizer.h"
 #include "Visualizers/ChaosVDSolverCollisionDataComponentVisualizer.h"
 
-
 AChaosVDParticleActor::AChaosVDParticleActor(const FObjectInitializer& ObjectInitializer) : Super(ObjectInitializer)
 {
 	RootComponent = CreateDefaultSubobject<USceneComponent>(TEXT("SceneComponent0"));
@@ -28,114 +27,248 @@ AChaosVDParticleActor::AChaosVDParticleActor(const FObjectInitializer& ObjectIni
 	CreateVisualizers();
 }
 
-void AChaosVDParticleActor::UpdateFromRecordedParticleData(const FChaosVDParticleDataWrapper& InRecordedData, const Chaos::FRigidTransform3& SimulationTransform)
+void AChaosVDParticleActor::UpdateFromRecordedParticleData(const TSharedPtr<FChaosVDParticleDataWrapper>& InRecordedData, const Chaos::FRigidTransform3& SimulationTransform)
 {
-	//TODO: Make the simulation transform be cached on the CVD Scene, so we can query from it when needed
-	// Copying it to each particle actor is not efficient
-	CachedSimulationTransform = SimulationTransform;
-
-	if (InRecordedData.ParticlePositionRotation.HasValidData())
+	if (!ensure(InRecordedData.IsValid()))
 	{
-		const FVector TargetLocation = SimulationTransform.TransformPosition(InRecordedData.ParticlePositionRotation.MX);
-		if (GetActorLocation() != TargetLocation)
-		{
-			SetActorLocation(InRecordedData.ParticlePositionRotation.MX);
-		}
+		return;	
+	}
 
-		const FQuat TargetRotation = SimulationTransform.GetRotation() * InRecordedData.ParticlePositionRotation.MR;
-		if (GetActorRotation() != TargetRotation.Rotator())
+	if (TSharedPtr<FChaosVDScene> ScenePtr = SceneWeakPtr.Pin())
+	{
+		if (InRecordedData->ParticleCluster.HasValidData())
 		{
-			SetActorRotation(TargetRotation);
+			if (AChaosVDParticleActor* ParentParticle = ScenePtr->GetParticleActor(InRecordedData->SolverID, InRecordedData->ParticleCluster.ParentParticleID))
+			{
+				AttachToActor(ParentParticle, FAttachmentTransformRules::KeepWorldTransform);
+			}
+			else if (AActor* CurrentParent = GetAttachParentActor())
+			{
+				DetachFromActor(FDetachmentTransformRules::KeepWorldTransform);
+			}
 		}
 	}
 
-	if (ParticleDataViewer.GeometryHash != InRecordedData.GeometryHash)
+	//TODO: Make the simulation transform be cached on the CVD Scene, so we can query from it when needed
+	// Copying it to each particle actor is not efficient
+	CachedSimulationTransform = SimulationTransform;
+
+	FTransform NewParticleTransform;
+	bool bHasNewTransform = false;
+	if (InRecordedData->ParticlePositionRotation.HasValidData())
 	{
-		UpdateGeometry(InRecordedData.GeometryHash, EChaosVDActorGeometryUpdateFlags::ForceUpdate);
+		const FVector TargetLocation = SimulationTransform.TransformPosition(InRecordedData->ParticlePositionRotation.MX);
+		const FVector CurrentLocation = SimulationTransform.TransformPosition(ParticleDataPtr && ParticleDataPtr->ParticlePositionRotation.HasValidData() ? ParticleDataPtr->ParticlePositionRotation.MX : FVector::ZeroVector);
+
+		const FQuat TargetRotation = SimulationTransform.GetRotation() * InRecordedData->ParticlePositionRotation.MR;
+		const FQuat CurrentRotation = SimulationTransform.GetRotation() * (ParticleDataPtr && ParticleDataPtr->ParticlePositionRotation.HasValidData() ? ParticleDataPtr->ParticlePositionRotation.MR : FQuat::Identity);
+
+		NewParticleTransform.SetLocation(TargetLocation);
+		NewParticleTransform.SetRotation(TargetRotation);
+		NewParticleTransform.SetScale3D(FVector(1.0f,1.0f,1.0f));
+
+		bHasNewTransform = CurrentRotation != TargetRotation || CurrentLocation != TargetLocation;
+
 	}
 
 	// This is iterating and comparing each element of the array,
 	// We might need to find a faster way of determine if the data changed, but for now this is faster than assuming it changed
-	const bool bShapeDataIsDirty = ParticleDataViewer.CollisionDataPerShape != InRecordedData.CollisionDataPerShape;
+	const bool bShapeDataIsDirty = !ParticleDataPtr || (ParticleDataPtr->CollisionDataPerShape != InRecordedData->CollisionDataPerShape);
+	const bool bDisabledStateChanged = ParticleDataPtr && (ParticleDataPtr->ParticleDynamicsMisc.bDisabled != InRecordedData->ParticleDynamicsMisc.bDisabled);
+	const bool bHasNewGeometry = !ParticleDataPtr || (ParticleDataPtr->GeometryHash != InRecordedData->GeometryHash);
+
+	ParticleDataPtr = InRecordedData;
+
+	if (bHasNewGeometry)
+	{
+		UpdateGeometry(InRecordedData->GeometryHash, EChaosVDActorGeometryUpdateFlags::ForceUpdate);
+	}
 
-	// TODO: We should store a ptr to the data and in our custom details panel draw it
-	ParticleDataViewer = InRecordedData;
+	if (bHasNewTransform)
+	{
+		PerformTaskOnGeometry([NewParticleTransform](TSharedPtr<FChaosVDMeshDataInstanceHandle>& MeshDataHandle)
+		{
+			MeshDataHandle->SetWorldTransform(NewParticleTransform);
+		});
+	}
 
 	// Now that we have updated particle data, update the Shape data and visibility as needed
-	if (bShapeDataIsDirty)
+	if (bShapeDataIsDirty || bHasNewGeometry)
 	{
 		UpdateShapeDataComponents();
 		UpdateGeometryComponentsVisibility();
 	}
+	else if (bDisabledStateChanged)
+	{
+		UpdateGeometryComponentsVisibility();
+	}
 
 	UpdateGeometryColors();
+
+	OnParticleDataUpdated().ExecuteIfBound();
+}
+
+void AChaosVDParticleActor::ProcessUpdatedAndRemovedHandles(TArray<TSharedPtr<FChaosVDExtractedGeometryDataHandle>>& OutExtractedGeometryDataHandles)
+{
+	for (TArray<TSharedPtr<FChaosVDMeshDataInstanceHandle>>::TIterator MeshDataHandleRemoveIterator = MeshDataHandles.CreateIterator(); MeshDataHandleRemoveIterator; ++MeshDataHandleRemoveIterator)
+	{
+		TSharedPtr<FChaosVDMeshDataInstanceHandle>& ExistingMeshDataHandle = *MeshDataHandleRemoveIterator;
+		if (ExistingMeshDataHandle.IsValid() && ExistingMeshDataHandle->GetGeometryHandle())
+		{
+			bool bExists = false;
+
+			// TODO: This search is n2, but I didn't see this as bottleneck. We should check if it is worth adding this to a TSet, or implementing the < operator so we can sort the array and do a binary search
+			// (avoiding the need to allocate a new container) 
+			for (TArray<TSharedPtr<FChaosVDExtractedGeometryDataHandle>>::TIterator HandleRemoveIterator = OutExtractedGeometryDataHandles.CreateIterator(); HandleRemoveIterator; ++HandleRemoveIterator)
+			{
+				const TSharedPtr<FChaosVDExtractedGeometryDataHandle> GeometryDataHandle = *HandleRemoveIterator;
+				const TSharedPtr<FChaosVDExtractedGeometryDataHandle> ExistingComponentGeometryDataHandle = ExistingMeshDataHandle->GetGeometryHandle();
+
+				const bool bBothHandlesAreValid = GeometryDataHandle && ExistingComponentGeometryDataHandle;
+				if (bBothHandlesAreValid && *GeometryDataHandle == *ExistingMeshDataHandle->GetGeometryHandle())
+				{
+					bExists = true;
+
+					// If we have a CVD Geometry Component for this handle, just remove it from the list as it means we don't need to re-create it
+					HandleRemoveIterator.RemoveCurrent();
+					break;
+				}
+			}
+
+			if (!bExists)
+			{
+				if (IChaosVDGeometryComponent* AsGeometryComponent = Cast<IChaosVDGeometryComponent>(ExistingMeshDataHandle->GetMeshComponent()))
+				{
+					AsGeometryComponent->RemoveMeshInstance(ExistingMeshDataHandle);
+				}
+
+				MeshDataHandleRemoveIterator.RemoveCurrent();
+			}
+		}		
+	}
 }
 
+
 void AChaosVDParticleActor::UpdateGeometry(const Chaos::FConstImplicitObjectPtr& InImplicitObject, EChaosVDActorGeometryUpdateFlags OptionsFlags)
 {
+	if (EnumHasAnyFlags(OptionsFlags, EChaosVDActorGeometryUpdateFlags::ForceUpdate))
+	{
+		bIsGeometryDataGenerationStarted = false;
+	}
+
+	if (bIsGeometryDataGenerationStarted)
+	{
+		return;
+	}
+
+	if (!ParticleDataPtr)
+	{
+		return;
+	}
+
 	if (!InImplicitObject.IsValid())
 	{
 		return;
 	}
-	
-	if (EnumHasAnyFlags(OptionsFlags, EChaosVDActorGeometryUpdateFlags::ForceUpdate))
+
+	const TSharedPtr<FChaosVDScene> ScenePtr = SceneWeakPtr.Pin();
+	if (!ScenePtr.IsValid())
 	{
-		bIsGeometryDataGenerationStarted = false;
+		return;
+	}
+
+	const TSharedPtr<FChaosVDGeometryBuilder> GeometryGenerator = ScenePtr->GetGeometryGenerator();
+	if (!GeometryGenerator.IsValid())
+	{
+		return;
+	}
 
-		for (TWeakObjectPtr<UMeshComponent>& MeshComponent : MeshComponents)
+	const int32 ObjectsToGenerateNum = InImplicitObject->CountLeafObjectsInHierarchyImpl();
+
+	// If the new implicit object is empty, then we can just clear all the mesh components and early out
+	if (ObjectsToGenerateNum == 0)
+	{
+		for (const TSharedPtr<FChaosVDMeshDataInstanceHandle>& MeshDataHandle : MeshDataHandles)
 		{
-			if (MeshComponent.IsValid())
+			if (!MeshDataHandle.IsValid())
+			{
+				continue;
+			}
+
+			if (IChaosVDGeometryComponent* AsGeometryComponent = Cast<IChaosVDGeometryComponent>(MeshDataHandle->GetMeshComponent()))
 			{
-				MeshComponent->DestroyComponent();
+				AsGeometryComponent->RemoveMeshInstance(MeshDataHandle);
 			}
 		}
 
-		MeshComponents.Reset();
+		MeshDataHandles.Reset();
+		return;
 	}
+	
+	TArray<TSharedPtr<FChaosVDExtractedGeometryDataHandle>> OutExtractedGeometryDataHandles;
+	OutExtractedGeometryDataHandles.Reserve(ObjectsToGenerateNum);
 
-	if (bIsGeometryDataGenerationStarted)
+	// Heightfields need to be created as Static meshes and use normal Static Mesh components because we need LODs for them due to their high triangle count
+	const bool bHasToUseStaticMeshComponent = FChaosVDGeometryBuilder::DoesImplicitContainType(InImplicitObject, Chaos::ImplicitObjectType::HeightField);
+	constexpr int32 LODsToGenerateNum = 3;
+	constexpr int32 LODsToGenerateNumForInstancedStaticMesh = 0;
+
+	GeometryGenerator->CreateMeshesFromImplicitObject<UStaticMesh>(InImplicitObject, this, OutExtractedGeometryDataHandles, bHasToUseStaticMeshComponent ? LODsToGenerateNum : LODsToGenerateNumForInstancedStaticMesh);
+
+	// This should not happen in theory, but there might be some valid situations where it does. Adding an ensure to catch them and then evaluate if it is really an issue (if it is not I will remove the ensure later on). 
+	if (!ensure(ObjectsToGenerateNum == OutExtractedGeometryDataHandles.Num()))
 	{
-		return;
+		UE_LOG(LogChaosVDEditor, Warning, TEXT("[%s] Geometry objects being generated doesn't match the number of objects in the implicit object | Expected [%d] | Being generated [%d] | Particle Actor [%s]"), ANSI_TO_TCHAR(__FUNCTION__), ObjectsToGenerateNum, OutExtractedGeometryDataHandles.Num(), *GetName());
 	}
 
-	if (const TSharedPtr<FChaosVDScene>& ScenePtr = SceneWeakPtr.Pin())
+	// Figure out what geometry was removed, and destroy their components as needed. Also, if a geometry is already generated an active, remove it from the geometry to generate list
+	ProcessUpdatedAndRemovedHandles(OutExtractedGeometryDataHandles);
+
+	if (OutExtractedGeometryDataHandles.Num() > 0)
 	{
-		if (const TSharedPtr<FChaosVDGeometryBuilder>& GeometryGenerator = ScenePtr->GetGeometryGenerator())
+		for (const TSharedPtr<FChaosVDExtractedGeometryDataHandle>& ExtractedGeometryDataHandle : OutExtractedGeometryDataHandles)
 		{
-			TArray<TWeakObjectPtr<UMeshComponent>> OutGeneratedMeshComponents;
-			Chaos::FRigidTransform3 Transform;
+			//TODO: Time Slice component creation
+			TSharedPtr<FChaosVDMeshDataInstanceHandle> MeshDataInstance;
 
-			// Heightfields need to be created as Static meshes and use normal Static Mesh components because we need LODs for them due to their high triangle count
-			if (FChaosVDGeometryBuilder::DoesImplicitContainType(InImplicitObject, Chaos::ImplicitObjectType::HeightField))
+			if (bHasToUseStaticMeshComponent)
 			{
-				constexpr int32 LODsToGenerateNum = 3;
-				constexpr int32 StartingMeshComponentIndex = 0;
-				GeometryGenerator->CreateMeshComponentsFromImplicit<UStaticMesh, UChaosVDStaticMeshComponent>(InImplicitObject, this, OutGeneratedMeshComponents, Transform, StartingMeshComponentIndex, LODsToGenerateNum);
+				MeshDataInstance = GeometryGenerator->CreateMeshDataInstance<UChaosVDStaticMeshComponent>(*ParticleDataPtr.Get(), ExtractedGeometryDataHandle);
 			}
 			else
 			{
-				GeometryGenerator->CreateMeshComponentsFromImplicit<UStaticMesh, UChaosVDInstancedStaticMeshComponent>(InImplicitObject, this, OutGeneratedMeshComponents, Transform);
+				MeshDataInstance = GeometryGenerator->CreateMeshDataInstance<UChaosVDInstancedStaticMeshComponent>(*ParticleDataPtr.Get(), ExtractedGeometryDataHandle);
 			}
 
-			if (OutGeneratedMeshComponents.Num() > 0)
+			const UMeshComponent* CreatedMeshComponent = MeshDataInstance->GetMeshComponent();
+			if (!CreatedMeshComponent)
 			{
-				MeshComponents.Append(OutGeneratedMeshComponents);
+				UE_LOG(LogChaosVDEditor, Error, TEXT("[%s] Failed To Create mesh component for [%s]"), ANSI_TO_TCHAR(__FUNCTION__), *GetName());
+				continue;
+			}
 
-				for (TWeakObjectPtr<UMeshComponent> MeshComponent : MeshComponents)
-				{
-					if (IChaosVDGeometryDataComponent* DataComponent = Cast<IChaosVDGeometryDataComponent>(MeshComponent.Get()))
-					{
-						DataComponent->SetRootImplicitObject(InImplicitObject);
-						DataComponent->UpdateDataFromShapeArray(ParticleDataViewer.CollisionDataPerShape);
-					}
-				}
+			// If we have a valid transform data, we need to update our instance with it as the mesh component is not part of this actor (and event if it is, we don't use the actor transform anymore)
+			if (ParticleDataPtr && ParticleDataPtr->ParticlePositionRotation.HasValidData())
+			{
+				const FVector TargetLocation = CachedSimulationTransform.TransformPosition(ParticleDataPtr->ParticlePositionRotation.MX);
+				const FQuat TargetRotation = CachedSimulationTransform.GetRotation() * ParticleDataPtr->ParticlePositionRotation.MR;
 
-				UpdateGeometryComponentsVisibility();
-				UpdateGeometryColors();
+				FTransform ParticleTransform;
+				ParticleTransform.SetLocation(TargetLocation);
+				ParticleTransform.SetRotation(TargetRotation);
 
-				bIsGeometryDataGenerationStarted = true;
+				MeshDataInstance->SetWorldTransform(ParticleTransform);
 			}
+
+			MeshDataHandles.Add(MeshDataInstance);
 		}
+
+		// Ensure that visibility and colorization is up to date after updating this Particle's Geometry
+		UpdateGeometryComponentsVisibility();
+		
+		UpdateGeometryColors();
+
+		bIsGeometryDataGenerationStarted = true;
 	}
 }
 
@@ -158,9 +291,9 @@ void AChaosVDParticleActor::SetScene(TWeakPtr<FChaosVDScene> InScene)
 	{
 		GeometryUpdatedDelegate = ScenePtr->OnNewGeometryAvailable().AddWeakLambda(this, [this](const Chaos::FConstImplicitObjectPtr& ImplicitObject, const uint32 ID)
 		{
-			if (ParticleDataViewer.GeometryHash == ID)
+			if (ParticleDataPtr && ParticleDataPtr->GeometryHash == ID)
 			{
-				UpdateGeometry(ImplicitObject);
+				UpdateGeometry(ImplicitObject, EChaosVDActorGeometryUpdateFlags::ForceUpdate);
 			}
 		});
 	}
@@ -180,7 +313,7 @@ void AChaosVDParticleActor::GetVisualizationContext(FChaosVDVisualizationContext
 {
 	OutVisualizationContext.SpaceTransform = CachedSimulationTransform;
 	OutVisualizationContext.CVDScene = SceneWeakPtr;
-	OutVisualizationContext.SolverID = ParticleDataViewer.SolverID;
+	OutVisualizationContext.SolverID = ParticleDataPtr ? ParticleDataPtr->SolverID : INDEX_NONE;
 }
 
 void AChaosVDParticleActor::CreateVisualizers()
@@ -194,7 +327,7 @@ bool AChaosVDParticleActor::IsSelectedInEditor() const
 {
 	// The implementation of this method in UObject, used a global edit callback,
 	// but as we don't use the global editor selection system, we need to re-route it.
-	if (TSharedPtr<FChaosVDScene> ScenePtr = SceneWeakPtr.Pin())
+	if (const TSharedPtr<FChaosVDScene> ScenePtr = SceneWeakPtr.Pin())
 	{
 		return ScenePtr->IsObjectSelected(this);
 	}
@@ -206,16 +339,6 @@ void AChaosVDParticleActor::PostEditChangeProperty(FPropertyChangedEvent& Proper
 {
 	Super::PostEditChangeProperty(PropertyChangedEvent);
 
-	if (PropertyChangedEvent.GetMemberPropertyName() == GET_MEMBER_NAME_STRING_CHECKED(AChaosVDParticleActor, ParticleDataViewer))
-	{
-		// Not particularly useful for now. This is a test code verifying we can react to changes in the data.
-		// In the future this could be part of the Re-simulation feature. When data is changed here, it can be propagated to the evolution instance that will be re-simulated
-		if (PropertyChangedEvent.GetPropertyName() == GET_MEMBER_NAME_STRING_CHECKED(FChaosVDParticleDataWrapper, GeometryHash))
-		{
-			UpdateGeometry(ParticleDataViewer.GeometryHash, EChaosVDActorGeometryUpdateFlags::ForceUpdate);
-		}
-	}
-
 	if (PropertyChangedEvent.GetMemberPropertyName() == GET_MEMBER_NAME_STRING_CHECKED(AChaosVDParticleActor, LocalParticleDataVisualizationFlags))
 	{
 		if (TUniquePtr<FChaosVDDataVisualizerBase>* ParticleDataVisualizer = CVDVisualizers.Find(FChaosVDParticleDataVisualizer::VisualizerID))
@@ -228,12 +351,43 @@ void AChaosVDParticleActor::PostEditChangeProperty(FPropertyChangedEvent& Proper
 #if WITH_EDITOR
 void AChaosVDParticleActor::SetIsTemporarilyHiddenInEditor(bool bIsHidden)
 {
-	const bool bShouldBeHidden = bIsActive ? bIsHidden : true;
+	Super::SetIsTemporarilyHiddenInEditor(bIsHidden);
 
-	Super::SetIsTemporarilyHiddenInEditor(bShouldBeHidden);
+	PerformTaskOnGeometry([this, bIsHidden](TSharedPtr<FChaosVDMeshDataInstanceHandle>& MeshDataHandle)
+	{
+		FChaosVDGeometryComponentUtils::UpdateMeshVisibility(MeshDataHandle, ParticleDataPtr ? *ParticleDataPtr.Get() : FChaosVDParticleDataWrapper(), IsActive() && !bIsHidden);
+	});
 }
 #endif //WITH_EDITOR
 
+FBox AChaosVDParticleActor::GetComponentsBoundingBox(bool bNonColliding, bool bIncludeFromChildActors) const
+{
+	FBox BoundingBox = FBox(ForceInitToZero);
+	if (ParticleDataPtr)
+	{
+		FBoxSphereBounds::Builder BoundsBuilder;
+
+		for (const TSharedPtr<FChaosVDMeshDataInstanceHandle>& MeshDataHandle : MeshDataHandles)
+		{
+			if (MeshDataHandle)
+			{
+				if (const UStaticMeshComponent* StaticMeshComponent = Cast<UStaticMeshComponent>(MeshDataHandle->GetMeshComponent()))
+				{
+					if (const UStaticMesh* Mesh = StaticMeshComponent->GetStaticMesh())
+					{
+						BoundsBuilder += Mesh->GetBounds().TransformBy(MeshDataHandle->GetWorldTransform());
+					}
+				}
+			}
+		}
+
+		const FBoxSphereBounds SphereBounds= BoundsBuilder;
+		BoundingBox = SphereBounds.GetBox();
+	}
+
+	return BoundingBox;
+}
+
 void AChaosVDParticleActor::GetCollisionData(TArray<TSharedPtr<FChaosVDCollisionDataFinder>>& OutCollisionDataFound)
 {
 	if (const TArray<TSharedPtr<FChaosVDParticlePairMidPhase>>* MidPhases = GetCollisionMidPhasesArray())
@@ -262,80 +416,77 @@ bool AChaosVDParticleActor::HasCollisionData()
 	return false;
 }
 
-FName AChaosVDParticleActor::GetName()
+FName AChaosVDParticleActor::GetProviderName()
 {
 	return GetFName();
 }
 
+void AChaosVDParticleActor::PushSelectionToProxies()
+{
+	Super::PushSelectionToProxies();
+
+	PerformTaskOnGeometry([this](TSharedPtr<FChaosVDMeshDataInstanceHandle>& MeshDataHandle)
+	{
+		const bool bIsSelectedInEditor = IsSelectedInEditor();
+		MeshDataHandle->SetIsSelected(bIsSelectedInEditor);
+	});
+}
+
 const TArray<TSharedPtr<FChaosVDParticlePairMidPhase>>* AChaosVDParticleActor::GetCollisionMidPhasesArray() const
 {
+	if (!ParticleDataPtr.IsValid())
+	{
+		return nullptr;
+	}
+
 	const TSharedPtr<FChaosVDScene> ScenePtr = SceneWeakPtr.Pin();
 	if (!ScenePtr.IsValid())
 	{
 		return nullptr;
 	}
 
-	if (AChaosVDSolverInfoActor* SolverInfoActor = ScenePtr->GetSolverInfoActor(ParticleDataViewer.SolverID))
+	if (AChaosVDSolverInfoActor* SolverInfoActor = ScenePtr->GetSolverInfoActor(ParticleDataPtr->SolverID))
 	{
 		if (const UChaosVDSolverCollisionDataComponent* CollisionDataComponent = SolverInfoActor->GetCollisionDataComponent())
 		{
-			return CollisionDataComponent->GetMidPhasesForParticle(ParticleDataViewer.ParticleIndex, EChaosVDCollisionParticlePairSlot::Any);
+			return CollisionDataComponent->GetMidPhasesForParticle(ParticleDataPtr->ParticleIndex, EChaosVDCollisionParticlePairSlot::Any);
 		}
 	}
 
 	return nullptr;
 }
 
-void AChaosVDParticleActor::PerformTaskOnGeometryComponents(TFunction<void(IChaosVDGeometryDataComponent& InDataComponent)> TaskToPerform)
+void AChaosVDParticleActor::UpdateShapeDataComponents()
 {
-	if (!ensure(TaskToPerform))
-	{
-		UE_LOG(LogChaosVDEditor, Error, TEXT("[%s] Called with an invalid task callback..."), ANSI_TO_TCHAR(__FUNCTION__))
-		return;
-	}
-
-	for (TWeakObjectPtr<UMeshComponent> MeshComponent : MeshComponents)
+	PerformTaskOnGeometry([this](TSharedPtr<FChaosVDMeshDataInstanceHandle>& MeshDataHandle)
 	{
-		if (IChaosVDGeometryDataComponent* DataComponent = Cast<IChaosVDGeometryDataComponent>(MeshComponent.Get()))
+		if (ParticleDataPtr)
 		{
-			// We need wait until we have a valid mesh
-			if (DataComponent->IsMeshReady())
-			{
-				TaskToPerform(*DataComponent);
-			}
-			else if (FChaosVDMeshReadyDelegate* MeshReadyDelegate = DataComponent->OnMeshReady())
-			{
-				// TODO: Guard against this being multiple times for a specific task.
-				// These tasks should probably be implemented on the components themselves
-				// and they decide if they need to wait and how to do so.
-				MeshReadyDelegate->AddWeakLambda(MeshComponent.Get(), [TaskToPerform](IChaosVDGeometryDataComponent& GeometryDataComponent)
-				{
-					TaskToPerform(GeometryDataComponent);
-				});
-			}
+			FChaosVDGeometryComponentUtils::UpdateCollisionDataFromShapeArray(ParticleDataPtr->CollisionDataPerShape, MeshDataHandle);
 		}
-	}
+	});
 }
 
-void AChaosVDParticleActor::UpdateShapeDataComponents()
+void AChaosVDParticleActor::UpdateGeometryComponentsVisibility()
 {
-	for (TWeakObjectPtr<UMeshComponent> MeshComponent : MeshComponents)
+	PerformTaskOnGeometry([this](TSharedPtr<FChaosVDMeshDataInstanceHandle>& MeshDataHandle)
 	{
-		if (IChaosVDGeometryDataComponent* DataComponent = Cast<IChaosVDGeometryDataComponent>(MeshComponent.Get()))
+		if (ParticleDataPtr)
 		{
-			DataComponent->UpdateDataFromShapeArray(ParticleDataViewer.CollisionDataPerShape);
+			FChaosVDGeometryComponentUtils::UpdateMeshVisibility(MeshDataHandle, *ParticleDataPtr.Get(), IsActive());
 		}
-	}
-}
-
-void AChaosVDParticleActor::UpdateGeometryComponentsVisibility()
-{
-	PerformTaskOnGeometryComponents([](IChaosVDGeometryDataComponent& GeometryDataComponent){ GeometryDataComponent.UpdateVisibility(); });
+	});
 }
 
 void AChaosVDParticleActor::UpdateGeometryColors()
 {
-	PerformTaskOnGeometryComponents([](IChaosVDGeometryDataComponent& GeometryDataComponent){ GeometryDataComponent.UpdateColors(); });
+	PerformTaskOnGeometry([this](TSharedPtr<FChaosVDMeshDataInstanceHandle>& MeshDataHandle)
+	{
+		if (ParticleDataPtr)
+		{
+			FChaosVDGeometryComponentUtils::UpdateMeshColor(MeshDataHandle, *ParticleDataPtr.Get(), GetIsServerParticle());
+		}
+	});
 }
 
 void AChaosVDParticleActor::SetIsActive(bool bNewActive)
@@ -349,7 +500,9 @@ void AChaosVDParticleActor::SetIsActive(bool bNewActive)
 		// We need to add a way to unlist inactive particle actors without a full hierarchy rebuild, which would be too costly
 		bEditable = bNewActive;
 		bListedInSceneOutliner = bNewActive;
-		SetIsTemporarilyHiddenInEditor(!bIsActive);
+
+		UpdateGeometryComponentsVisibility();
+
 #endif
 
 		if (const TSharedPtr<FChaosVDScene> ScenePtr = SceneWeakPtr.Pin())
diff --git a/Engine/Plugins/Experimental/ChaosVD/Source/ChaosVD/Private/ChaosVDParticleActorCustomization.cpp b/Engine/Plugins/Experimental/ChaosVD/Source/ChaosVD/Private/ChaosVDParticleActorCustomization.cpp
--- a/Engine/Plugins/Experimental/ChaosVD/Source/ChaosVD/Private/ChaosVDParticleActorCustomization.cpp
+++ b/Engine/Plugins/Experimental/ChaosVD/Source/ChaosVD/Private/ChaosVDParticleActorCustomization.cpp
@@ -2,10 +2,16 @@
 
 #include "ChaosVDParticleActorCustomization.h"
 
+#include "ChaosVDModule.h"
+#include "ChaosVDParticleActor.h"
 #include "DetailCategoryBuilder.h"
 #include "DetailLayoutBuilder.h"
+#include "DetailWidgetRow.h"
+#include "IStructureDetailsView.h"
+#include "DataWrappers/ChaosVDParticleDataWrapper.h"
 #include "DetailsCustomizations/ChaosVDDetailsCustomizationUtils.h"
 
+#define LOCTEXT_NAMESPACE "ChaosVisualDebugger"
 
 FChaosVDParticleActorCustomization::FChaosVDParticleActorCustomization()
 {
@@ -13,6 +19,14 @@ FChaosVDParticleActorCustomization::FChaosVDParticleActorCustomization()
 	AllowedCategories.Add(FChaosVDParticleActorCustomization::ChaosVDVisualizationCategoryName);
 }
 
+FChaosVDParticleActorCustomization::~FChaosVDParticleActorCustomization()
+{
+	if (CurrentObservedActor.Get())
+	{
+		CurrentObservedActor->OnParticleDataUpdated().Unbind();
+	}
+}
+
 TSharedRef<IDetailCustomization> FChaosVDParticleActorCustomization::MakeInstance()
 {
 	return MakeShareable( new FChaosVDParticleActorCustomization );
@@ -24,5 +38,54 @@ void FChaosVDParticleActorCustomization::CustomizeDetails(IDetailLayoutBuilder&
 
 	FChaosVDDetailsCustomizationUtils::HideAllCategories(DetailBuilder, AllowedCategories);
 
-	DetailBuilder.EditCategory(ChaosVDCategoryName).InitiallyCollapsed(false);
+	// We keep the particle data we need to visualize as a shared ptr because copying it each frame we advance/rewind to to an struct that lives in the particle actor it is not cheap.
+	// Having a struct details view to which we set that pointer data each time the data in the particle is updated (meaning we assigned another ptr from the recording)
+	// seems to be more expensive because it has to rebuild the entire layout from scratch.
+	// So a middle ground I found is to have a Particle Data struct in this customization instance, which we add as external property. Then each time the particle data is updated we copy the data over.
+	// This allow us to only perform the copy just for the particle that is being inspected and not every particle updated in that frame.
+
+	IDetailCategoryBuilder& CVDMainCategoryBuilder = DetailBuilder.EditCategory(ChaosVDCategoryName).InitiallyCollapsed(false);
+
+	TArray<TWeakObjectPtr<UObject>> SelectedObjects;
+	DetailBuilder.GetObjectsBeingCustomized(SelectedObjects);
+	if (SelectedObjects.Num() > 0)
+	{
+		//TODO: Add support for multi-selection.
+		if (!ensure(SelectedObjects.Num() == 1))
+		{
+			UE_LOG(LogChaosVDEditor, Warning, TEXT("[%s] [%d] objects were selectioned but this customization panel only support single object selection."), ANSI_TO_TCHAR(__FUNCTION__), SelectedObjects.Num())
+		}
+
+		if (AChaosVDParticleActor* CurrentActor = CurrentObservedActor.Get())
+		{
+			CurrentParticleDataCopy = FChaosVDParticleDataWrapper();
+			CurrentActor->OnParticleDataUpdated().Unbind();
+			CurrentActor = nullptr;
+		}
+		
+		if (AChaosVDParticleActor* ParticleActor = Cast<AChaosVDParticleActor>(SelectedObjects[0]))
+		{
+			CurrentObservedActor = ParticleActor;
+			CurrentObservedActor->OnParticleDataUpdated().BindRaw(this, &FChaosVDParticleActorCustomization::HandleParticleDataUpdated);
+
+			HandleParticleDataUpdated();
+
+			const TSharedPtr<FStructOnScope> ParticleDataView = MakeShared<FStructOnScope>(FChaosVDParticleDataWrapper::StaticStruct(), reinterpret_cast<uint8*>(&CurrentParticleDataCopy));
+			TArray<TSharedPtr<IPropertyHandle>> Handles = CVDMainCategoryBuilder.AddAllExternalStructureProperties(ParticleDataView.ToSharedRef(), EPropertyLocation::Default, nullptr);
+		}
+	}
 }
+
+void FChaosVDParticleActorCustomization::HandleParticleDataUpdated()
+{	
+	if (const FChaosVDParticleDataWrapper* ParticleDataPtr = CurrentObservedActor.Get() ? CurrentObservedActor->GetParticleData() : nullptr)
+	{
+		CurrentParticleDataCopy = *ParticleDataPtr;	
+	}
+	else
+	{
+		CurrentParticleDataCopy = FChaosVDParticleDataWrapper();
+	}
+}
+
+#undef LOCTEXT_NAMESPACE
diff --git a/Engine/Plugins/Experimental/ChaosVD/Source/ChaosVD/Private/ChaosVDPlaybackViewportClient.cpp b/Engine/Plugins/Experimental/ChaosVD/Source/ChaosVD/Private/ChaosVDPlaybackViewportClient.cpp
--- a/Engine/Plugins/Experimental/ChaosVD/Source/ChaosVD/Private/ChaosVDPlaybackViewportClient.cpp
+++ b/Engine/Plugins/Experimental/ChaosVD/Source/ChaosVD/Private/ChaosVDPlaybackViewportClient.cpp
@@ -18,6 +18,7 @@
 #include "UnrealWidget.h"
 #include "Actors/ChaosVDSolverInfoActor.h"
 #include "Components/ChaosVDSolverCollisionDataComponent.h"
+#include "Components/InstancedStaticMeshComponent.h"
 #include "Visualizers/ChaosVDDebugDrawUtils.h"
 #include "Widgets/SChaosVDMainTab.h"
 
@@ -32,8 +33,10 @@ FChaosVDPlaybackViewportClient::FChaosVDPlaybackViewportClient(const TSharedPtr<
 
 	if (UChaosVDEditorSettings* Settings = GetMutableDefault<UChaosVDEditorSettings>())
 	{
-		OverrideFarClipPlane(Settings->FarClippingOverride);
-		Settings->OnFarClippingOverrideChanged().AddRaw(this, &FChaosVDPlaybackViewportClient::HandleFarClippingOverrideSettingsChanged);
+		Settings->OnFarClippingOverrideChanged().AddRaw(this, &FChaosVDPlaybackViewportClient::HandleViewportSettingsChanged);
+		Settings->OnVisibilitySettingsChanged().AddRaw(this, &FChaosVDPlaybackViewportClient::HandleViewportSettingsChanged);
+
+		HandleViewportSettingsChanged(Settings);
 	}
 }
 
@@ -55,6 +58,7 @@ FChaosVDPlaybackViewportClient::~FChaosVDPlaybackViewportClient()
 	if (UChaosVDEditorSettings* Settings = GetMutableDefault<UChaosVDEditorSettings>())
 	{
 		Settings->OnFarClippingOverrideChanged().RemoveAll(this);
+		Settings->OnVisibilitySettingsChanged().RemoveAll(this);
 	}
 }
 
@@ -89,6 +93,32 @@ void FChaosVDPlaybackViewportClient::ProcessClick(FSceneView& View, HHitProxy* H
 			bClickHandled = Visualizer->VisProxyHandleClick(this, ComponentVisProxy, Click);
 		}
 
+		const IChaosVDGeometryComponent* AsCVDGeometryComponent = nullptr;
+		int32 MeshInstanceIndex = INDEX_NONE;
+
+		if (const HInstancedStaticMeshInstance* InstancedStaticMeshProxy = HitProxyCast<HInstancedStaticMeshInstance>(HitProxy))
+		{
+			AsCVDGeometryComponent = Cast<IChaosVDGeometryComponent>(InstancedStaticMeshProxy->Component);
+			MeshInstanceIndex = InstancedStaticMeshProxy->InstanceIndex;
+		}
+		else if (const HActor* ActorHitProxy = HitProxyCast<HActor>(HitProxy))
+		{
+			AsCVDGeometryComponent = Cast<IChaosVDGeometryComponent>(ActorHitProxy->PrimComponent.Get());
+			MeshInstanceIndex = 0;
+		}
+
+		if (AsCVDGeometryComponent && MeshInstanceIndex != INDEX_NONE)
+		{
+			if (const TSharedPtr<FChaosVDMeshDataInstanceHandle> MeshDataHandle = AsCVDGeometryComponent->GetMeshDataInstanceHandle(MeshInstanceIndex))
+			{
+				if (AChaosVDParticleActor* ClickedActor = ScenePtr->GetParticleActor(MeshDataHandle->GetOwningSolverID(), MeshDataHandle->GetOwningParticleID()))
+				{
+					ScenePtr->SetSelectedObject(ClickedActor);
+					bClickHandled = true;
+				}
+			}
+		}
+
 		if (bClickHandled)
 		{
 			return;
@@ -138,11 +168,13 @@ void FChaosVDPlaybackViewportClient::HandleActorMoving(AActor* MovedActor) const
 	}
 }
 
-void FChaosVDPlaybackViewportClient::HandleFarClippingOverrideSettingsChanged(UChaosVDEditorSettings* SettingsObject)
+void FChaosVDPlaybackViewportClient::HandleViewportSettingsChanged(UChaosVDEditorSettings* SettingsObject)
 {
 	if (SettingsObject)
 	{
 		OverrideFarClipPlane(SettingsObject->FarClippingOverride);
+		EngineShowFlags.SetMeshEdges(EnumHasAnyFlags(static_cast<EChaosVDGeometryVisibilityFlags>(SettingsObject->GeometryVisibilityFlags), EChaosVDGeometryVisibilityFlags::ShowTriangleEdges));
+		Invalidate();
 	}
 }
 
diff --git a/Engine/Plugins/Experimental/ChaosVD/Source/ChaosVD/Private/ChaosVDScene.cpp b/Engine/Plugins/Experimental/ChaosVD/Source/ChaosVD/Private/ChaosVDScene.cpp
--- a/Engine/Plugins/Experimental/ChaosVD/Source/ChaosVD/Private/ChaosVDScene.cpp
+++ b/Engine/Plugins/Experimental/ChaosVD/Source/ChaosVD/Private/ChaosVDScene.cpp
@@ -51,6 +51,8 @@ void FChaosVDScene::Initialize()
 	PhysicsVDWorld = CreatePhysicsVDWorld();
 
 	GeometryGenerator = MakeShared<FChaosVDGeometryBuilder>();
+
+	GeometryGenerator->Initialize(AsWeak());
 	
 	StreamableManager = MakeShared<FStreamableManager>();
 
@@ -60,6 +62,8 @@ void FChaosVDScene::Initialize()
 		// Jira for tracking UE-191639
 		StreamableManager->RequestSyncLoad(Settings->QueryOnlyMeshesMaterial.ToSoftObjectPath());
 		StreamableManager->RequestSyncLoad(Settings->SimOnlyMeshesMaterial.ToSoftObjectPath());
+		StreamableManager->RequestSyncLoad(Settings->InstancedMeshesMaterial.ToSoftObjectPath());
+		StreamableManager->RequestSyncLoad(Settings->InstancedMeshesQueryOnlyMaterial.ToSoftObjectPath());
 		
 		Settings->OnVisibilitySettingsChanged().AddRaw(this, &FChaosVDScene::HandleVisibilitySettingsChanged);
 		Settings->OnColorSettingsChanged().AddRaw(this, &FChaosVDScene::HandleColorSettingsChanged);
@@ -133,7 +137,7 @@ void FChaosVDScene::UpdateFromRecordedStepData(const int32 SolverID, const FStri
 		UpdatingSceneSlowTask.MakeDialogDelayed(ChaosVDSceneUIOptions::DelayToShowProgressDialogThreshold, ChaosVDSceneUIOptions::bShowCancelButton, ChaosVDSceneUIOptions::bAllowInPIE);
 	
 		// Go over existing Particle VD Instances and update them or create them if needed 
-		for (const FChaosVDParticleDataWrapper& Particle : InRecordedStepData.RecordedParticlesData)
+		for (const TSharedPtr<FChaosVDParticleDataWrapper>& Particle : InRecordedStepData.RecordedParticlesData)
 		{
 			const int32 ParticleVDInstanceID = GetIDForRecordedParticleData(Particle);
 			ParticlesIDsInRecordedStepData.Add(ParticleVDInstanceID);
@@ -301,6 +305,10 @@ Chaos::FConstImplicitObjectPtr FChaosVDScene::GetUpdatedGeometry(int32 GeometryI
 		{
 			return *Geometry;
 		}
+		else
+		{
+			UE_LOG(LogChaosVDEditor, Warning, TEXT("Geometry for key [%d] is not loaded in the recording yet"), GeometryID);
+		}
 	}
 
 	return nullptr;
@@ -337,33 +345,28 @@ bool FChaosVDScene::IsSolverForServer(int32 SolverID) const
 	return false;
 }
 
-AChaosVDParticleActor* FChaosVDScene::SpawnParticleFromRecordedData(const FChaosVDParticleDataWrapper& InParticleData, const FChaosVDSolverFrameData& InFrameData)
+AChaosVDParticleActor* FChaosVDScene::SpawnParticleFromRecordedData(const TSharedPtr<FChaosVDParticleDataWrapper>& InParticleData, const FChaosVDSolverFrameData& InFrameData)
 {
 	using namespace Chaos;
 
+	if (!InParticleData.IsValid())
+	{
+		return nullptr;
+	}
+
 	FActorSpawnParameters Params;
-	Params.Name = *InParticleData.DebugName;
+	Params.Name = *InParticleData->DebugName;
 	Params.NameMode = FActorSpawnParameters::ESpawnActorNameMode::Requested;
 
 	if (AChaosVDParticleActor* NewActor = PhysicsVDWorld->SpawnActor<AChaosVDParticleActor>(Params))
 	{
 		NewActor->SetIsActive(true);
+		NewActor->SetScene(AsShared());
+		NewActor->SetIsServerParticle(IsSolverForServer(InParticleData->SolverID));
 		NewActor->UpdateFromRecordedParticleData(InParticleData, InFrameData.SimulationTransform);
 
-		if (!InParticleData.DebugName.IsEmpty())
-		{
-			NewActor->SetActorLabel(InParticleData.DebugName);
-		}
-
-		NewActor->SetScene(AsShared());
-		
-		if (ensure(LoadedRecording.IsValid()))
-		{
-			if (const Chaos::FConstImplicitObjectPtr* Geometry = LoadedRecording->GetGeometryMap().Find(InParticleData.GeometryHash))
-			{
-				NewActor->UpdateGeometry(*Geometry);
-			}
-		}
+		const bool bHasDebugName = !InParticleData->DebugName.IsEmpty();
+		NewActor->SetActorLabel(bHasDebugName ? InParticleData->DebugName : TEXT("Unnamed Particle - ID : ") + FString::FromInt(InParticleData->ParticleIndex));
 
 		return NewActor;
 	}
@@ -371,9 +374,9 @@ AChaosVDParticleActor* FChaosVDScene::SpawnParticleFromRecordedData(const FChaos
 	return nullptr;
 }
 
-int32 FChaosVDScene::GetIDForRecordedParticleData(const FChaosVDParticleDataWrapper& InParticleData) const
+int32 FChaosVDScene::GetIDForRecordedParticleData(const TSharedPtr<FChaosVDParticleDataWrapper>& InParticleData) const
 {
-	return InParticleData.ParticleIndex;
+	return InParticleData ? InParticleData->ParticleIndex : INDEX_NONE;
 }
 
 void FChaosVDScene::CreateBaseLights(UWorld* TargetWorld) const
@@ -412,6 +415,16 @@ void FChaosVDScene::CreateBaseLights(UWorld* TargetWorld) const
 	}
 }
 
+AActor* FChaosVDScene::CreateMeshComponentsContainer(UWorld* TargetWorld)
+{
+	const FName GeometryFolderPath("ChaosVisualDebugger/GeneratedMeshComponents");
+
+	MeshComponentContainerActor = TargetWorld->SpawnActor<AActor>();
+	MeshComponentContainerActor->SetFolderPath(GeometryFolderPath);
+
+	return MeshComponentContainerActor;
+}
+
 UWorld* FChaosVDScene::CreatePhysicsVDWorld()
 {
 	const FName UniqueWorldName = FName(FGuid::NewGuid().ToString());
@@ -433,6 +446,7 @@ UWorld* FChaosVDScene::CreatePhysicsVDWorld()
 	);
 
 	CreateBaseLights(NewWorld);
+	CreateMeshComponentsContainer(NewWorld);
 
 	ActorDestroyedHandle = NewWorld->AddOnActorDestroyedHandler(FOnActorDestroyed::FDelegate::CreateRaw(this, &FChaosVDScene::HandleActorDestroyed));
 	
@@ -522,13 +536,16 @@ void FChaosVDScene::InitializeSelectionSets()
 	SelectionSet = NewObject<UTypedElementSelectionSet>(GetTransientPackage(), NAME_None, RF_Transactional);
 	SelectionSet->AddToRoot();
 
-	ActorSelection = USelection::CreateActorSelection(GetTransientPackage(), TEXT("CVDSelectedActors"), RF_Transactional);
+	FString ActorSelectionObjectName = FString::Printf(TEXT("CVDSelectedActors-%s"), *FGuid::NewGuid().ToString());
+	ActorSelection = USelection::CreateActorSelection(GetTransientPackage(), *ActorSelectionObjectName, RF_Transactional);
 	ActorSelection->SetElementSelectionSet(SelectionSet);
 
-	ComponentSelection = USelection::CreateComponentSelection(GetTransientPackage(), TEXT("CVDSelectedComponents"), RF_Transactional);
+	FString ComponentSelectionObjectName = FString::Printf(TEXT("CVDSelectedComponents-%s"), *FGuid::NewGuid().ToString());
+	ComponentSelection = USelection::CreateComponentSelection(GetTransientPackage(), *ComponentSelectionObjectName, RF_Transactional);
 	ComponentSelection->SetElementSelectionSet(SelectionSet);
 
-	ObjectSelection = USelection::CreateObjectSelection(GetTransientPackage(), TEXT("CVDSelectedObjects"), RF_Transactional);
+	FString ObjectSelectionObjectName = FString::Printf(TEXT("CVDSelectedObjects-%s"), *FGuid::NewGuid().ToString());
+	ObjectSelection = USelection::CreateObjectSelection(GetTransientPackage(), *ObjectSelectionObjectName, RF_Transactional);
 	ObjectSelection->SetElementSelectionSet(SelectionSet);
 
 	SelectionSet->OnPreChange().AddRaw(this, &FChaosVDScene::HandlePreSelectionChange);
diff --git a/Engine/Plugins/Experimental/ChaosVD/Source/ChaosVD/Private/Components/ChaosVDInstancedStaticMeshComponent.cpp b/Engine/Plugins/Experimental/ChaosVD/Source/ChaosVD/Private/Components/ChaosVDInstancedStaticMeshComponent.cpp
--- a/Engine/Plugins/Experimental/ChaosVD/Source/ChaosVD/Private/Components/ChaosVDInstancedStaticMeshComponent.cpp
+++ b/Engine/Plugins/Experimental/ChaosVD/Source/ChaosVD/Private/Components/ChaosVDInstancedStaticMeshComponent.cpp
@@ -2,28 +2,262 @@
 
 #include "Components/ChaosVDInstancedStaticMeshComponent.h"
 
-void UChaosVDInstancedStaticMeshComponent::UpdateVisibility()
+#include "ChaosVDModule.h"
+#include "Materials/MaterialInterface.h"
+#include "Materials/MaterialInstanceDynamic.h"
+
+void UChaosVDInstancedStaticMeshComponent::UpdateInstanceVisibility(const TSharedPtr<FChaosVDMeshDataInstanceHandle>& InInstanceHandle, bool bIsVisible)
+{
+	if (!InInstanceHandle.IsValid())
+	{
+		UE_LOG(LogChaosVDEditor, Warning, TEXT("[%s] Attempted to update a mesh instance using an invalid handle. No instances were updated"), ANSI_TO_TCHAR(__FUNCTION__));
+		return;
+	}
+
+	if (InInstanceHandle->GetMeshComponent() != this)
+	{
+		UE_LOG(LogChaosVDEditor, Error, TEXT("[%s] Attempted to update a mesh instance using a handle from another component. No instances were updated | Handle Component [%s] | Current Component [%s]"), ANSI_TO_TCHAR(__FUNCTION__), *GetNameSafe(InInstanceHandle->GetMeshComponent()), *GetNameSafe(this));
+		return;
+	}
+
+	const int32 InstanceIndex = InInstanceHandle->GetMeshInstanceIndex();
+
+	FTransform Transform = InInstanceHandle->GetWorldTransform();
+	if (!bIsVisible)
+	{
+		// Setting the scale to 0 will hide this instance while keeping it on the component
+		Transform.SetScale3D(FVector::ZeroVector);
+	}
+
+	constexpr bool bIsWorldSpaceTransform = true;
+	constexpr bool bMarkRenderDirty = true;
+	constexpr bool bTeleport = true;
+
+	UpdateInstanceTransform(InstanceIndex, Transform, bIsWorldSpaceTransform, bMarkRenderDirty, bTeleport);
+}
+
+void UChaosVDInstancedStaticMeshComponent::SetIsSelected(const TSharedPtr<FChaosVDMeshDataInstanceHandle>& InInstanceHandle, bool bIsSelected)
+{
+	if (!InInstanceHandle.IsValid())
+	{
+		UE_LOG(LogChaosVDEditor, Warning, TEXT("[%s] Attempted to update a mesh instance using an invalid handle. No instances were updated"), ANSI_TO_TCHAR(__FUNCTION__));
+		return;
+	}
+
+	if (InInstanceHandle->GetMeshComponent() != this)
+	{
+		UE_LOG(LogChaosVDEditor, Error, TEXT("[%s] Attempted to update a mesh instance using a handle from another component. No instances were updated | Handle Component [%s] | Current Component [%s]"), ANSI_TO_TCHAR(__FUNCTION__), *GetNameSafe(InInstanceHandle->GetMeshComponent()), *GetNameSafe(this));
+		return;
+	}
+
+	const int32 InstanceIndex = InInstanceHandle->GetMeshInstanceIndex();
+
+	NotifySMInstanceSelectionChanged({this, InstanceIndex}, bIsSelected);
+}
+
+void UChaosVDInstancedStaticMeshComponent::UpdateInstanceColor(const TSharedPtr<FChaosVDMeshDataInstanceHandle>& InInstanceHandle, FLinearColor NewColor)
+{
+	if (CurrentMaterial == nullptr)
+	{
+		CurrentMaterial = GetMaterial(0);
+
+		bool bIsSolidColor = FMath::IsNearlyEqual(NewColor.A, 1.0f);
+		UMaterialInterface* DesiredMaterial = InInstanceHandle->GetCachedMaterialInstance( bIsSolidColor ? EChaosVDMaterialType::Instanced : EChaosVDMaterialType::InstancedQueryOnly);
+
+		if (CurrentMaterial != DesiredMaterial)
+		{
+			SetMaterial(0, DesiredMaterial);
+			CurrentMaterial = DesiredMaterial;
+		}
+	}
+	
+	SetNumCustomDataFloats(4);
+	SetCustomData(InInstanceHandle->GetMeshInstanceIndex(), TArrayView<float>(reinterpret_cast<float*>(&NewColor), 4));
+}
+
+void UChaosVDInstancedStaticMeshComponent::UpdateInstanceWorldTransform(const TSharedPtr<FChaosVDMeshDataInstanceHandle>& InInstanceHandle, const FTransform& InTransform)
+{
+	if (!InInstanceHandle.IsValid())
+	{
+		UE_LOG(LogChaosVDEditor, Warning, TEXT("[%s] Attempted to update a mesh instance using an invalid handle. No instances were updated"), ANSI_TO_TCHAR(__FUNCTION__));
+		return;
+	}
+
+	if (InInstanceHandle->GetMeshComponent() != this)
+	{
+		UE_LOG(LogChaosVDEditor, Error, TEXT("[%s] Attempted to update a mesh instance using a handle from another component. No instances were updated | Handle Component [%s] | Current Component [%s]"), ANSI_TO_TCHAR(__FUNCTION__), *GetNameSafe(InInstanceHandle->GetMeshComponent()), *GetNameSafe(this));
+		return;
+	}
+
+	const int32 InstanceIndex = InInstanceHandle->GetMeshInstanceIndex();
+
+	FTransform NewTransform = InTransform;
+
+	if (!InInstanceHandle->GetVisibility())
+	{
+		NewTransform.SetScale3D(FVector::ZeroVector);
+	}
+
+	constexpr bool bIsWorldSpaceTransform = true;
+	constexpr bool bMarkRenderDirty = true;
+	constexpr bool bTeleport = true;
+
+	UpdateInstanceTransform(InstanceIndex, NewTransform, bIsWorldSpaceTransform, bMarkRenderDirty, bTeleport);
+}
+
+TArrayView<TSharedPtr<FChaosVDMeshDataInstanceHandle>> UChaosVDInstancedStaticMeshComponent::GetMeshDataInstanceHandles()
+{
+	return CurrentInstanceHandles;
+}
+
+void UChaosVDInstancedStaticMeshComponent::Reset()
+{
+	bIsMeshReady = false;
+	MeshReadyDelegate = FChaosVDMeshReadyDelegate();
+	ComponentEmptyDelegate = FChaosVDMeshComponentEmptyDelegate();
+	CurrentMaterial = nullptr;
+
+	CurrentInstanceHandles.Reset();
+	CurrentGeometryKey = 0;
+}
+
+bool UChaosVDInstancedStaticMeshComponent::UpdateGeometryKey(const uint32 NewHandleGeometryKey)
 {
-	UpdateVisibility_Internal(CollisionData, Cast<UMeshComponent>(this));
+	if (CurrentGeometryKey != 0 && CurrentGeometryKey != NewHandleGeometryKey)
+	{
+		ensure(false);
+
+		UE_LOG(LogChaosVDEditor, Warning, TEXT("[%s] Attempted to add a mesh instance belonging to another geometry key. No instance was added | CurrentKey [%u] | New Key [%u]"), ANSI_TO_TCHAR(__FUNCTION__), CurrentGeometryKey, NewHandleGeometryKey);
+		return false;
+	}
+	else
+	{
+		CurrentGeometryKey = NewHandleGeometryKey;
+	}
+	
+	return true;
 }
 
-void UChaosVDInstancedStaticMeshComponent::UpdateDataFromShapeArray(const TArray<FChaosVDShapeCollisionData>& InShapeArray)
+TSharedPtr<FChaosVDMeshDataInstanceHandle> UChaosVDInstancedStaticMeshComponent::AddMeshInstance(const FTransform InstanceTransform, bool bIsWorldSpace, const TSharedPtr<FChaosVDExtractedGeometryDataHandle>& InGeometryHandle, int32 ParticleID, int32 SolverID)
 {
-	UpdateDataFromShapeArray_Internal(InShapeArray, CollisionData, this);
+	int32 InstanceIndex = AddInstance(InstanceTransform, bIsWorldSpace);
+
+	check(InstanceIndex != INDEX_NONE);
+
+	const TSharedPtr<FChaosVDMeshDataInstanceHandle> InstanceHandle = MakeShared<FChaosVDMeshDataInstanceHandle>(InstanceIndex, this, ParticleID, SolverID);
+	InstanceHandle->SetGeometryHandle(InGeometryHandle);
+
+	const uint32 NewHandleGeometryKey = InGeometryHandle->GetGeometryKey();
+
+	if (!UpdateGeometryKey(NewHandleGeometryKey))
+	{
+		return nullptr;
+	}
+	
+	CurrentInstanceHandles.Add(InstanceHandle);
+
+	return InstanceHandle;
 }
 
-void UChaosVDInstancedStaticMeshComponent::SetImplicitObject(const Chaos::FImplicitObject* InImplicitObject)
+void UChaosVDInstancedStaticMeshComponent::AddMeshInstanceForHandle(TSharedPtr<FChaosVDMeshDataInstanceHandle> MeshDataHandle, const FTransform InstanceTransform, bool bIsWorldSpace, const TSharedPtr<FChaosVDExtractedGeometryDataHandle>& InGeometryHandle, int32 ParticleID, int32 SolverID)
 {
-	SetImplicitObject_Internal(InImplicitObject);
+	const int32 InstanceIndex = AddInstance(InstanceTransform, bIsWorldSpace);
+
+	check(InstanceIndex != INDEX_NONE);
+
+	const uint32 NewHandleGeometryKey = InGeometryHandle->GetGeometryKey();
+
+	if (!UpdateGeometryKey(NewHandleGeometryKey))
+	{
+		return;
+	}
+
+	MeshDataHandle->SetMeshInstanceIndex(InstanceIndex);
+	MeshDataHandle->SetMeshComponent(this);
+	
+	CurrentInstanceHandles.Add(MeshDataHandle);
 }
 
-void UChaosVDInstancedStaticMeshComponent::UpdateColors()
+void UChaosVDInstancedStaticMeshComponent::RemoveMeshInstance(TSharedPtr<FChaosVDMeshDataInstanceHandle> InHandleToRemove)
 {
-	UpdateColors_Internal(this);
+	if (!InHandleToRemove.IsValid())
+	{
+		UE_LOG(LogChaosVDEditor, Warning, TEXT("[%s] Attempted to remove a mesh instace using an invalid handle. No instanced were removed"), ANSI_TO_TCHAR(__FUNCTION__));
+		return;
+	}
+
+	if (InHandleToRemove->GetMeshComponent() != this)
+	{
+		UE_LOG(LogChaosVDEditor, Error, TEXT("[%s] Attempted to remove a mesh instace using a handle from another component. No instanced were removed | Handle Component [%s] | Current Component [%s]"), ANSI_TO_TCHAR(__FUNCTION__), *GetNameSafe(InHandleToRemove->GetMeshComponent()), *GetNameSafe(this));
+		return;
+	}
+
+	if (!ensure(InHandleToRemove->GetMeshInstanceIndex() != INDEX_NONE))
+	{
+		return;
+	}
+
+	{
+		const int32 CurrentInstanceCount =  GetInstanceCount();
+
+		if (InHandleToRemove->GetMeshInstanceIndex() > CurrentInstanceCount - 1)
+		{
+			ensure(false);
+
+			UE_LOG(LogChaosVDEditor, Error, TEXT("[%s] Mesh Data Handle has an invalid instance index. No instanced were removed | Handle Instance Index [%d] | Current Instance Count [%d]"), ANSI_TO_TCHAR(__FUNCTION__), InHandleToRemove->GetMeshInstanceIndex(), CurrentInstanceCount);
+			return;
+		}
+	}
+
+	RemoveInstance(InHandleToRemove->GetMeshInstanceIndex());
+
+	auto FindPredicate = [InHandleToRemove](const TSharedPtr<FChaosVDMeshDataInstanceHandle>& InstanceHandle)
+	{
+		return InstanceHandle->GetMeshInstanceIndex() == InHandleToRemove->GetMeshInstanceIndex();
+	};
+
+	// TODO: Evaluate if keeping the instance handles array sorted on any change is worth the cost to allow a binary search here. Same question for a Set
+	const int32 InstanceHandleIndex = CurrentInstanceHandles.IndexOfByPredicate(FindPredicate);
+
+	if (InstanceHandleIndex != INDEX_NONE)
+	{
+		CurrentInstanceHandles.RemoveAtSwap(InstanceHandleIndex);
+	}
+
+	// If the component becomes "empty", allow it to go back to the pool
+	if (GetInstanceCount() == 0)
+	{
+		ComponentEmptyDelegate.Broadcast(this);
+	}
+}
+
+bool UChaosVDInstancedStaticMeshComponent::Modify(bool bAlwaysMarkDirty)
+{
+	// CVD Mesh Components are not saved to any assets or require undo
+	return false;
+}
+
+bool UChaosVDInstancedStaticMeshComponent::IsNavigationRelevant() const
+{
+	return false;
 }
 
-void UChaosVDInstancedStaticMeshComponent::SetRootImplicitObject(const Chaos::FConstImplicitObjectPtr& InImplicitObject)
+uint32 UChaosVDInstancedStaticMeshComponent::GetGeometryKey() const
 {
-	RootImplicitObject = InImplicitObject;
+	return CurrentGeometryKey;
 }
 
+TSharedPtr<FChaosVDMeshDataInstanceHandle> UChaosVDInstancedStaticMeshComponent::GetMeshDataInstanceHandle(int32 InstanceIndex) const
+{
+	auto FindPredicate = [InstanceIndex](const TSharedPtr<FChaosVDMeshDataInstanceHandle>& InstanceHandle)
+	{
+		return InstanceHandle->GetMeshInstanceIndex() == InstanceIndex;
+	};
+
+	if (const TSharedPtr<FChaosVDMeshDataInstanceHandle>* FoundHandle = CurrentInstanceHandles.FindByPredicate(FindPredicate))
+	{
+		return *FoundHandle;
+	}
+
+	return nullptr;
+}
diff --git a/Engine/Plugins/Experimental/ChaosVD/Source/ChaosVD/Private/Components/ChaosVDStaticMeshComponent.cpp b/Engine/Plugins/Experimental/ChaosVD/Source/ChaosVD/Private/Components/ChaosVDStaticMeshComponent.cpp
--- a/Engine/Plugins/Experimental/ChaosVD/Source/ChaosVD/Private/Components/ChaosVDStaticMeshComponent.cpp
+++ b/Engine/Plugins/Experimental/ChaosVD/Source/ChaosVD/Private/Components/ChaosVDStaticMeshComponent.cpp
@@ -2,28 +2,170 @@
 
 #include "Components/ChaosVDStaticMeshComponent.h"
 
-void UChaosVDStaticMeshComponent::UpdateVisibility()
+#include "ChaosVDModule.h"
+#include "MaterialDomain.h"
+#include "Materials/Material.h"
+#include "Materials/MaterialInstanceDynamic.h"
+
+uint32 UChaosVDStaticMeshComponent::GetGeometryKey() const
 {
-	UpdateVisibility_Internal(CollisionData, Cast<UMeshComponent>(this));
+	return CurrentGeometryKey;
 }
 
-void UChaosVDStaticMeshComponent::UpdateDataFromShapeArray(const TArray<FChaosVDShapeCollisionData>& InShapeArray)
+void UChaosVDStaticMeshComponent::UpdateInstanceVisibility(const TSharedPtr<FChaosVDMeshDataInstanceHandle>& InInstanceHandle, bool bIsVisible)
 {
-	UpdateDataFromShapeArray_Internal(InShapeArray, CollisionData, this);
+	if (!InInstanceHandle.IsValid())
+	{
+		UE_LOG(LogChaosVDEditor, Warning, TEXT("[%s] Attempted to update a mesh instance using an invalid handle. No instances were updated"), ANSI_TO_TCHAR(__FUNCTION__));
+		return;
+	}
+
+	if (InInstanceHandle->GetMeshComponent() != this)
+	{
+		UE_LOG(LogChaosVDEditor, Error, TEXT("[%s] Attempted to update a mesh instance using a handle from another component. No instances were updated | Handle Component [%s] | Current Component [%s]"), ANSI_TO_TCHAR(__FUNCTION__), *GetNameSafe(InInstanceHandle->GetMeshComponent()), *GetNameSafe(this));
+		return;
+	}
+
+	SetVisibility(bIsVisible);
 }
 
-void UChaosVDStaticMeshComponent::SetImplicitObject(const Chaos::FImplicitObject* InImplicitObject)
+void UChaosVDStaticMeshComponent::SetIsSelected(const TSharedPtr<FChaosVDMeshDataInstanceHandle>& InInstanceHandle, bool bIsSelected)
 {
-	SetImplicitObject_Internal(InImplicitObject);
+	bIsOwningParticleSelected = bIsSelected;
+
+	PushSelectionToProxy();
+}
+
+bool UChaosVDStaticMeshComponent::ShouldRenderSelected() const
+{
+	return bIsOwningParticleSelected;
+}
+
+void UChaosVDStaticMeshComponent::UpdateInstanceColor(const TSharedPtr<FChaosVDMeshDataInstanceHandle>& InInstanceHandle, FLinearColor NewColor)
+{
+	if (!InInstanceHandle.IsValid())
+	{
+		UE_LOG(LogChaosVDEditor, Warning, TEXT("[%s] Attempted to update a mesh instance using an invalid handle. No instances were updated"), ANSI_TO_TCHAR(__FUNCTION__));
+		return;
+	}
+
+	if (InInstanceHandle->GetMeshComponent() != this)
+	{
+		UE_LOG(LogChaosVDEditor, Error, TEXT("[%s] Attempted to update a mesh instance using a handle from another component. No instances were updated | Handle Component [%s] | Current Component [%s]"), ANSI_TO_TCHAR(__FUNCTION__), *GetNameSafe(InInstanceHandle->GetMeshComponent()), *GetNameSafe(this));
+		return;
+	}
+
+	bool bIsSolidColor = FMath::IsNearlyEqual(NewColor.A, 1.0f);
+
+	if (UMaterialInstanceDynamic* MaterialToApply = bIsSolidColor ? InInstanceHandle->GetCachedMaterialInstance(EChaosVDMaterialType::SimOnlyMaterial) : InInstanceHandle->GetCachedMaterialInstance(EChaosVDMaterialType::QueryOnlyMaterial))
+	{
+		MaterialToApply->SetVectorParameterValue(TEXT("BaseColor"), NewColor);
+		InInstanceHandle->GetMeshComponent()->SetMaterial(0, MaterialToApply);
+	}
+	else
+	{
+		ensure(false);
+		UE_LOG(LogChaosVDEditor, Error, TEXT("[%s] Failed to get Query Only material for, applying the default mesh to all geometry"), ANSI_TO_TCHAR(__FUNCTION__));
+			
+		InInstanceHandle->GetMeshComponent()->SetMaterial(0, UMaterial::GetDefaultMaterial(EMaterialDomain::MD_Surface));
+	}
 }
 
-void UChaosVDStaticMeshComponent::UpdateColors()
+void UChaosVDStaticMeshComponent::UpdateInstanceWorldTransform(const TSharedPtr<FChaosVDMeshDataInstanceHandle>& InInstanceHandle, const FTransform& InTransform)
 {
-	UpdateColors_Internal(this);
+	if (!InInstanceHandle.IsValid())
+	{
+		UE_LOG(LogChaosVDEditor, Warning, TEXT("[%s] Attempted to update a mesh instance using an invalid handle. No instances were updated"), ANSI_TO_TCHAR(__FUNCTION__));
+		return;
+	}
+
+	if (InInstanceHandle->GetMeshComponent() != this)
+	{
+		UE_LOG(LogChaosVDEditor, Error, TEXT("[%s] Attempted to update a mesh instance using a handle from another component. No instances were updated | Handle Component [%s] | Current Component [%s]"), ANSI_TO_TCHAR(__FUNCTION__), *GetNameSafe(InInstanceHandle->GetMeshComponent()), *GetNameSafe(this));
+		return;
+	}
+
+	SetWorldTransform(InTransform);
 }
 
-void UChaosVDStaticMeshComponent::SetRootImplicitObject(const Chaos::FConstImplicitObjectPtr& InImplicitObject)
+void UChaosVDStaticMeshComponent::Reset()
 {
-	RootImplicitObject = InImplicitObject;
+	bIsMeshReady = false;
+	MeshReadyDelegate = FChaosVDMeshReadyDelegate();
+	ComponentEmptyDelegate = FChaosVDMeshComponentEmptyDelegate();
+
+	CurrentMeshDataHandle = nullptr;
+	CurrentGeometryKey = 0;
 }
 
+TSharedPtr<FChaosVDMeshDataInstanceHandle> UChaosVDStaticMeshComponent::AddMeshInstance(const FTransform InstanceTransform, bool bIsWorldSpace, const TSharedPtr<FChaosVDExtractedGeometryDataHandle>& InGeometryHandle, int32 ParticleID, int32 SolverID)
+{
+	// Static meshes only have one instance
+	constexpr int32 StaticMeshInstanceIndex = 0;
+
+	CurrentMeshDataHandle = MakeShared<FChaosVDMeshDataInstanceHandle>(StaticMeshInstanceIndex, this, ParticleID, SolverID);
+	CurrentMeshDataHandle->SetGeometryHandle(InGeometryHandle);
+
+	const uint32 NewHandleGeometryKey = InGeometryHandle->GetGeometryKey();
+
+	if (!UpdateGeometryKey(NewHandleGeometryKey))
+	{
+		return nullptr;
+	}
+
+	CurrentMeshDataHandle->SetWorldTransform(InstanceTransform);
+
+	return CurrentMeshDataHandle;
+}
+
+void UChaosVDStaticMeshComponent::AddMeshInstanceForHandle(TSharedPtr<FChaosVDMeshDataInstanceHandle> MeshDataHandle, const FTransform InstanceTransform, bool bIsWorldSpace, const TSharedPtr<FChaosVDExtractedGeometryDataHandle>& InGeometryHandle, int32 ParticleID, int32 SolverID)
+{
+	// Static meshes only have one instance
+	constexpr int32 StaticMeshInstanceIndex = 0;
+
+	const uint32 NewHandleGeometryKey = InGeometryHandle->GetGeometryKey();
+	if (!UpdateGeometryKey(NewHandleGeometryKey))
+	{
+		return;
+	}
+
+	MeshDataHandle->SetMeshInstanceIndex(StaticMeshInstanceIndex);
+	MeshDataHandle->SetMeshComponent(this);
+	
+	CurrentMeshDataHandle->SetWorldTransform(InstanceTransform);
+}
+
+void UChaosVDStaticMeshComponent::RemoveMeshInstance(TSharedPtr<FChaosVDMeshDataInstanceHandle> InHandleToRemove)
+{
+	SetStaticMesh(nullptr);
+	CurrentMeshDataHandle = nullptr;
+	ComponentEmptyDelegate.Broadcast(this);
+}
+
+bool UChaosVDStaticMeshComponent::UpdateGeometryKey(uint32 NewHandleGeometryKey)
+{
+	if (CurrentGeometryKey != 0 && CurrentGeometryKey != NewHandleGeometryKey)
+	{
+		ensure(false);
+
+		UE_LOG(LogChaosVDEditor, Warning, TEXT("[%s] Attempted to add a mesh instance belonging to another geometry key. No instance was added | CurrentKey [%u] | New Key [%u]"), ANSI_TO_TCHAR(__FUNCTION__), CurrentGeometryKey, NewHandleGeometryKey);
+		return false;
+	}
+	else
+	{
+		CurrentGeometryKey = NewHandleGeometryKey;
+	}
+
+	return true;
+}
+
+
+TSharedPtr<FChaosVDMeshDataInstanceHandle> UChaosVDStaticMeshComponent::GetMeshDataInstanceHandle(int32 InstanceIndex) const
+{
+	return CurrentMeshDataHandle;
+}
+
+TArrayView<TSharedPtr<FChaosVDMeshDataInstanceHandle>> UChaosVDStaticMeshComponent::GetMeshDataInstanceHandles()
+{
+	return TArrayView<TSharedPtr<FChaosVDMeshDataInstanceHandle>>(&CurrentMeshDataHandle, 1);
+}
diff --git a/Engine/Plugins/Experimental/ChaosVD/Source/ChaosVD/Private/DetailsCustomizations/ChaosVDParticleDataWrapperCustomization.cpp b/Engine/Plugins/Experimental/ChaosVD/Source/ChaosVD/Private/DetailsCustomizations/ChaosVDParticleDataWrapperCustomization.cpp
--- a/Engine/Plugins/Experimental/ChaosVD/Source/ChaosVD/Private/DetailsCustomizations/ChaosVDParticleDataWrapperCustomization.cpp
+++ b/Engine/Plugins/Experimental/ChaosVD/Source/ChaosVD/Private/DetailsCustomizations/ChaosVDParticleDataWrapperCustomization.cpp
@@ -29,6 +29,7 @@ void FChaosVDParticleDataWrapperCustomization::CustomizeChildren(TSharedRef<IPro
 	ParticleDataViewersNames.Add(GET_MEMBER_NAME_CHECKED(FChaosVDParticleDataWrapper, ParticleDynamics));
 	ParticleDataViewersNames.Add(GET_MEMBER_NAME_CHECKED(FChaosVDParticleDataWrapper, ParticleDynamicsMisc));
 	ParticleDataViewersNames.Add(GET_MEMBER_NAME_CHECKED(FChaosVDParticleDataWrapper, ParticleMassProps));
+	ParticleDataViewersNames.Add(GET_MEMBER_NAME_CHECKED(FChaosVDParticleDataWrapper, ParticleCluster));
 	
 	for (uint32 ChildIndex = 0; ChildIndex < NumChildren; ++ChildIndex)
 	{
diff --git a/Engine/Plugins/Experimental/ChaosVD/Source/ChaosVD/Private/ObjectsWaitingGeometryList.cpp b/Engine/Plugins/Experimental/ChaosVD/Source/ChaosVD/Private/ObjectsWaitingGeometryList.cpp
new file mode 100644
--- /dev/null
+++ b/Engine/Plugins/Experimental/ChaosVD/Source/ChaosVD/Private/ObjectsWaitingGeometryList.cpp
@@ -0,0 +1,9 @@
+﻿// Copyright Epic Games, Inc. All Rights Reserved.
+
+#include "ObjectsWaitingGeometryList.h"
+
+int32 FChaosVDWailingObjectListCVars::NumOfGTTaskToProcessPerTick = 200;
+FAutoConsoleVariableRef FChaosVDWailingObjectListCVars::CVarChaosVDGeometryToProcessPerTick(
+	TEXT("p.Chaos.VD.Tool.GeometryToProcessPerTick"),
+	FChaosVDWailingObjectListCVars::NumOfGTTaskToProcessPerTick,
+	TEXT("Number of generated geometry to process each tick when loading a teace file in the CVD tool."));
diff --git a/Engine/Plugins/Experimental/ChaosVD/Source/ChaosVD/Private/Trace/DataProcessors/ChaosVDTraceParticleDataProcessor.cpp b/Engine/Plugins/Experimental/ChaosVD/Source/ChaosVD/Private/Trace/DataProcessors/ChaosVDTraceParticleDataProcessor.cpp
--- a/Engine/Plugins/Experimental/ChaosVD/Source/ChaosVD/Private/Trace/DataProcessors/ChaosVDTraceParticleDataProcessor.cpp
+++ b/Engine/Plugins/Experimental/ChaosVD/Source/ChaosVD/Private/Trace/DataProcessors/ChaosVDTraceParticleDataProcessor.cpp
@@ -24,15 +24,15 @@ bool FChaosVDTraceParticleDataProcessor::ProcessRawData(const TArray<uint8>& InD
 
 	Chaos::FChaosArchive Ar(MemeReader);
 
-	FChaosVDParticleDataWrapper ParticleData;
-	ParticleData.Serialize(Ar);
+	TSharedPtr<FChaosVDParticleDataWrapper> ParticleData = MakeShared<FChaosVDParticleDataWrapper>();
+	ParticleData->Serialize(Ar);
 
 	// This can be null if the recording started Mid-Frame. In this case we just discard the data for now
-	if (FChaosVDSolverFrameData* FrameData = ProviderSharedPtr->GetCurrentSolverFrame(ParticleData.SolverID))
+	if (FChaosVDSolverFrameData* FrameData = ProviderSharedPtr->GetCurrentSolverFrame(ParticleData->SolverID))
 	{
 		if (ensureMsgf(FrameData->SolverSteps.Num() > 0, TEXT("A particle was traced without a valid step scope")))
 		{
-			FrameData->SolverSteps.Last().RecordedParticlesData.Add(MoveTemp(ParticleData));
+			FrameData->SolverSteps.Last().RecordedParticlesData.Add(ParticleData);
 		}
 	}
 
diff --git a/Engine/Plugins/Experimental/ChaosVD/Source/ChaosVD/Private/Visualizers/ChaosVDDebugDrawUtils.cpp b/Engine/Plugins/Experimental/ChaosVD/Source/ChaosVD/Private/Visualizers/ChaosVDDebugDrawUtils.cpp
--- a/Engine/Plugins/Experimental/ChaosVD/Source/ChaosVD/Private/Visualizers/ChaosVDDebugDrawUtils.cpp
+++ b/Engine/Plugins/Experimental/ChaosVD/Source/ChaosVD/Private/Visualizers/ChaosVDDebugDrawUtils.cpp
@@ -148,6 +148,27 @@ void FChaosVDDebugDrawUtils::DrawBox(FPrimitiveDrawInterface* PDI, const FVector
 	}
 }
 
+void FChaosVDDebugDrawUtils::DrawLine(FPrimitiveDrawInterface* PDI, const FVector& InStartPosition, const FVector& InEndPosition, const FColor& InColor, FStringView DebugText, ESceneDepthPriorityGroup DepthPriority)
+{
+	if (!PDI)
+	{
+		return;
+	}
+
+	constexpr float Thickness = 2.0f;
+	constexpr float DepthBias = 0;
+	constexpr bool bScreenSpace = Thickness > 0;
+
+	PDI->DrawLine(InStartPosition, InEndPosition, InColor, DepthPriority, Thickness, DepthBias, bScreenSpace);
+
+	if (!DebugText.IsEmpty())
+	{
+		// Draw the text in the middle of the line
+		const FVector TextWorldPosition = InStartPosition + ((InEndPosition - InStartPosition)  * 0.5f);
+		DrawText(DebugText, TextWorldPosition, InColor);
+	}
+}
+
 void FChaosVDDebugDrawUtils::DrawCanvas(FViewport& InViewport, FSceneView& View, FCanvas& Canvas)
 {
 	if (!GEngine)
diff --git a/Engine/Plugins/Experimental/ChaosVD/Source/ChaosVD/Private/Visualizers/ChaosVDParticleDataVisualizer.cpp b/Engine/Plugins/Experimental/ChaosVD/Source/ChaosVD/Private/Visualizers/ChaosVDParticleDataVisualizer.cpp
--- a/Engine/Plugins/Experimental/ChaosVD/Source/ChaosVD/Private/Visualizers/ChaosVDParticleDataVisualizer.cpp
+++ b/Engine/Plugins/Experimental/ChaosVD/Source/ChaosVD/Private/Visualizers/ChaosVDParticleDataVisualizer.cpp
@@ -4,6 +4,8 @@
 
 #include "ChaosVDEditorSettings.h"
 #include "ChaosVDModule.h"
+#include "ChaosVDParticleActor.h"
+#include "ChaosVDScene.h"
 #include "DataWrappers/ChaosVDParticleDataWrapper.h"
 #include "Visualizers/ChaosVDDebugDrawUtils.h"
 
@@ -79,4 +81,29 @@ void FChaosVDParticleDataVisualizer::DrawVisualization(const FSceneView* View, F
 			FChaosVDDebugDrawUtils::DrawArrowVector(PDI, OwnerCoMLocation, OwnerCoMLocation + AngImpScale * ParticleDataViewer->ParticleDynamics.MAngularImpulseVelocity, TEXT("Angular Implulse Velocity"), FColor::Emerald);
 		}
 	}
+
+	// TODO: This is a Proof of concept to test how debug draw connectivity data will look
+	// This will get re-implemented when we move particle data visualization to the new visualization system currently used for Collision Data
+	if (ParticleDataViewer->ParticleCluster.HasValidData())
+	{
+		if (EnumHasAnyFlags(VisualizationFlagsToUse, EChaosVDParticleDataVisualizationFlags::ClusterConnectivityEdge))
+		{
+			for (const FChaosVDConnectivityEdge& ConnectivityEdge : ParticleDataViewer->ParticleCluster.ConnectivityEdges)
+			{
+				if (TSharedPtr<FChaosVDScene> ScenePtr = VisualizationContext.CVDScene.Pin())
+				{
+					if (AChaosVDParticleActor* SiblingParticle = ScenePtr->GetParticleActor(VisualizationContext.SolverID, ConnectivityEdge.SiblingParticleID))
+					{
+						if (const FChaosVDParticleDataWrapper* SiblingParticleData = SiblingParticle->GetParticleData())
+						{
+							FVector BoxExtents(2,2,2);
+							FTransform BoxTransform(ParticleDataViewer->ParticlePositionRotation.MR, ParticleDataViewer->ParticlePositionRotation.MX);
+							FChaosVDDebugDrawUtils::DrawBox(PDI, BoxExtents, FColor::Black, BoxTransform, nullptr, ESceneDepthPriorityGroup::SDPG_Foreground);
+							FChaosVDDebugDrawUtils::DrawLine(PDI, ParticleDataViewer->ParticlePositionRotation.MX, SiblingParticleData->ParticlePositionRotation.MX, FColor::Blue, FString::Printf(TEXT("Strain [%f]"), ConnectivityEdge.Strain), ESceneDepthPriorityGroup::SDPG_Foreground);
+						}
+					}	
+				}
+			}
+		}
+	}
 }
diff --git a/Engine/Plugins/Experimental/ChaosVD/Source/ChaosVD/Private/Visualizers/ChaosVDSolverCollisionDataComponentVisualizer.cpp b/Engine/Plugins/Experimental/ChaosVD/Source/ChaosVD/Private/Visualizers/ChaosVDSolverCollisionDataComponentVisualizer.cpp
--- a/Engine/Plugins/Experimental/ChaosVD/Source/ChaosVD/Private/Visualizers/ChaosVDSolverCollisionDataComponentVisualizer.cpp
+++ b/Engine/Plugins/Experimental/ChaosVD/Source/ChaosVD/Private/Visualizers/ChaosVDSolverCollisionDataComponentVisualizer.cpp
@@ -207,8 +207,10 @@ void FChaosVDSolverCollisionDataComponentVisualizer::DrawMidPhaseData(const UAct
 
 				PDI->SetHitProxy(new HChaosVDContactPointProxy(Component, HitPRoxyDataFinder));
 
-				const FTransform& WorldActorTransform1 = CVDParticleActor1->GetTransform();
-				const FTransform& WorldActorTransform0 = CVDParticleActor0->GetTransform();
+				const FChaosVDParticleDataWrapper* ParticleDataActor1 = CVDParticleActor1->GetParticleData();
+				const FChaosVDParticleDataWrapper* ParticleDataActor0 = CVDParticleActor0->GetParticleData();
+				const FTransform WorldActorTransform1 = ParticleDataActor1 && ParticleDataActor1->ParticlePositionRotation.HasValidData() ? FTransform(ParticleDataActor1->ParticlePositionRotation.MR, ParticleDataActor1->ParticlePositionRotation.MX) : FTransform();
+				const FTransform WorldActorTransform0 = ParticleDataActor0 && ParticleDataActor0->ParticlePositionRotation.HasValidData() ? FTransform(ParticleDataActor0->ParticlePositionRotation.MR, ParticleDataActor0->ParticlePositionRotation.MX) : FTransform();
 
 				constexpr int32 ContactPlaneOwner = 1;
 				constexpr int32 ContactPointOwner = 1 - ContactPlaneOwner;
@@ -299,7 +301,6 @@ void FChaosVDSolverCollisionDataComponentVisualizer::DrawMidPhaseData(const UAct
 					if (ManifoldPoint.InitialPhi != 0)
 					{
 						FChaosVDDebugDrawUtils::DrawCircle(PDI, WorldPlaneLocation + ManifoldPoint.InitialPhi * WorldPlaneNormal, 0.25f * DebugDrawSettings.ContactCircleRadius, CircleSegments, InitialPhiColor, LinesThickness, Axes.GetUnitAxis(EAxis::Y), Axes.GetUnitAxis(EAxis::Z), TEXT("Contact : Manifold Initial Phi"), DebugDrawSettings.DepthPriority);
-
 					}
 				}
 
diff --git a/Engine/Plugins/Experimental/ChaosVD/Source/ChaosVD/Private/Widgets/SChaosVDCollisionDataInspector.cpp b/Engine/Plugins/Experimental/ChaosVD/Source/ChaosVD/Private/Widgets/SChaosVDCollisionDataInspector.cpp
--- a/Engine/Plugins/Experimental/ChaosVD/Source/ChaosVD/Private/Widgets/SChaosVDCollisionDataInspector.cpp
+++ b/Engine/Plugins/Experimental/ChaosVD/Source/ChaosVD/Private/Widgets/SChaosVDCollisionDataInspector.cpp
@@ -115,7 +115,7 @@ void SChaosVDCollisionDataInspector::SetCollisionDataProviderObjectToInspect(ICh
 		return;
 	}
 
-	CurrentObjectBeingInspectedName = CollisionDataProvider->GetName();
+	CurrentObjectBeingInspectedName = CollisionDataProvider->GetProviderName();
 
 	TArray<TSharedPtr<FName>> NewCollisionDataEntriesNameList;
 	TArray<TSharedPtr<FChaosVDCollisionDataFinder>> FoundCollisionData;
diff --git a/Engine/Plugins/Experimental/ChaosVD/Source/ChaosVD/Private/Widgets/SChaosVDDetailsView.cpp b/Engine/Plugins/Experimental/ChaosVD/Source/ChaosVD/Private/Widgets/SChaosVDDetailsView.cpp
--- a/Engine/Plugins/Experimental/ChaosVD/Source/ChaosVD/Private/Widgets/SChaosVDDetailsView.cpp
+++ b/Engine/Plugins/Experimental/ChaosVD/Source/ChaosVD/Private/Widgets/SChaosVDDetailsView.cpp
@@ -2,11 +2,14 @@
 
 #include "Widgets/SChaosVDDetailsView.h"
 
+#include "ChaosVDParticleActor.h"
 #include "GameFramework/Actor.h"
 #include "PropertyEditorModule.h"
 #include "SSubobjectInstanceEditor.h"
 #include "SSubobjectEditor.h"
 #include "SSubobjectEditorModule.h"
+#include "Visualizers/IChaosVDParticleVisualizationDataProvider.h"
+
 
 void SChaosVDDetailsView::Construct(const FArguments& InArgs)
 {
@@ -76,6 +79,7 @@ void SChaosVDDetailsView::SetSelectedObject(UObject* NewObject)
 	}
 
 	CurrentObjectInView = NewObject;
+
 	DetailsView->SetObject(NewObject, true);
 	SubobjectEditor->UpdateTree();
 }
diff --git a/Engine/Plugins/Experimental/ChaosVD/Source/ChaosVD/Public/ChaosVDCollisionDataProviderInterface.h b/Engine/Plugins/Experimental/ChaosVD/Source/ChaosVD/Public/ChaosVDCollisionDataProviderInterface.h
--- a/Engine/Plugins/Experimental/ChaosVD/Source/ChaosVD/Public/ChaosVDCollisionDataProviderInterface.h
+++ b/Engine/Plugins/Experimental/ChaosVD/Source/ChaosVD/Public/ChaosVDCollisionDataProviderInterface.h
@@ -38,5 +38,5 @@ public:
 	/**
 	 * Returns the name of the object providing the collision data
 	 */
-	virtual FName GetName() PURE_VIRTUAL(IChaosVDCollisionDataProviderInterface::HasCollisionData, return FName();)
+	virtual FName GetProviderName() PURE_VIRTUAL(IChaosVDCollisionDataProviderInterface::GetProviderName, return FName();)
 };
diff --git a/Engine/Plugins/Experimental/ChaosVD/Source/ChaosVD/Public/ChaosVDEditorSettings.h b/Engine/Plugins/Experimental/ChaosVD/Source/ChaosVD/Public/ChaosVDEditorSettings.h
--- a/Engine/Plugins/Experimental/ChaosVD/Source/ChaosVD/Public/ChaosVDEditorSettings.h
+++ b/Engine/Plugins/Experimental/ChaosVD/Source/ChaosVD/Public/ChaosVDEditorSettings.h
@@ -43,6 +43,7 @@ enum class EChaosVDGeometryVisibilityFlags : uint8
 	Complex = 1 << 4,
 	ShowHeightfields = 1 << 5, // Selecting this will show heightfields even if complex is not selected
 	ShowDisabledParticles = 1 << 6,
+	ShowTriangleEdges = 1 << 7,
 };
 ENUM_CLASS_FLAGS(EChaosVDGeometryVisibilityFlags)
 
@@ -215,6 +216,12 @@ public:
 	UPROPERTY(Config)
 	TSoftObjectPtr<UMaterial> SimOnlyMeshesMaterial;
 
+	UPROPERTY(Config)
+	TSoftObjectPtr<UMaterial> InstancedMeshesMaterial;
+
+	UPROPERTY(Config)
+	TSoftObjectPtr<UMaterial> InstancedMeshesQueryOnlyMaterial;
+
 	UPROPERTY(Config)
 	FSoftClassPath SkySphereActorClass;
 
diff --git a/Engine/Plugins/Experimental/ChaosVD/Source/ChaosVD/Public/ChaosVDExtractedGeometryDataHandle.h b/Engine/Plugins/Experimental/ChaosVD/Source/ChaosVD/Public/ChaosVDExtractedGeometryDataHandle.h
new file mode 100644
--- /dev/null
+++ b/Engine/Plugins/Experimental/ChaosVD/Source/ChaosVD/Public/ChaosVDExtractedGeometryDataHandle.h
@@ -0,0 +1,68 @@
+﻿// Copyright Epic Games, Inc. All Rights Reserved.
+
+#pragma once
+
+//#include "ChaosVDGeometryBuilder.h"
+#include "Math/Transform.h"
+#include "UObject/NameTypes.h"
+#include "UObject/WeakObjectPtr.h"
+
+class FChaosVDGeometryBuilder;
+class UDynamicMesh;
+class UStaticMesh;
+
+namespace Chaos
+{
+	class FImplicitObject;
+}
+
+/** Structure that acts as a handle to access the generated static mesh of a implicit object, plus other data extracted from it */
+struct FChaosVDExtractedGeometryDataHandle
+{
+	/** Key to be used to match this geometry with recorded Shape Data */
+	uint32 GetDataComponentKey() const { return DataComponentKey; }
+	void SetDataComponentKey(const uint32 Key) { DataComponentKey = Key; }
+
+	/** Key representing the generated static mesh - Handles for different implicit objects can point to the same geometry key, if after being unpacked they result in essentially the same mesh */
+	uint32 GetGeometryKey() const { return GeometryKey; }
+	void SetGeometryKey(const uint32 Key) { GeometryKey = Key; }
+
+	/** Name to be used to create labels in editor for this data */
+	FName GetName() const;
+
+	/** Extracted transform of the generated implicit object in component space */
+	const FTransform& GetRelativeTransform() const { return Transform; }
+	void SetTransform(const FTransform& InTransform) { Transform = InTransform; }
+
+	/** Ptr to the implicit object this handle provides access to and from which the mesh was generated */
+	const Chaos::FImplicitObject* GetImplicitObject() const { return ImplicitObject; }
+
+	/** If the object was an union and this was a leaf, this getter returns the root union object, otherwise it will be the same as GetImplicitObject */
+	const Chaos::FImplicitObject* GetRootImplicitObject() const { return RootImplicitObject; }
+	
+	/** Sets the Implicit Object Ptr from which the data was extracted and the mesh was generated */
+	void SetImplicitObject(const Chaos::FImplicitObject* InImplicitObject) { ImplicitObject = InImplicitObject; }
+	
+	/** Sets the Implicit Object Ptr from which the data was extracted and the mesh was generated */
+	void SetRootImplicitObject(const Chaos::FImplicitObject* InImplicitObject) { RootImplicitObject = InImplicitObject; }
+
+	bool operator==(const FChaosVDExtractedGeometryDataHandle& Other) const
+	{
+		return Other.GeometryKey == GeometryKey && Other.DataComponentKey == DataComponentKey && Other.Transform.Equals(Transform);
+	}
+
+	friend uint32 GetTypeHash(const FChaosVDExtractedGeometryDataHandle& Handle)
+	{
+		return HashCombine(HashCombine(Handle.GeometryKey, Handle.DataComponentKey), GetTypeHash(Handle.Transform));
+	}
+
+private:
+	uint32 GeometryKey = 0;
+
+	uint32 DataComponentKey = 0;
+
+	FTransform Transform;
+
+	const Chaos::FImplicitObject* ImplicitObject = nullptr;
+	const Chaos::FImplicitObject* RootImplicitObject = nullptr;
+};
diff --git a/Engine/Plugins/Experimental/ChaosVD/Source/ChaosVD/Public/ChaosVDGeometryBuilder.h b/Engine/Plugins/Experimental/ChaosVD/Source/ChaosVD/Public/ChaosVDGeometryBuilder.h
--- a/Engine/Plugins/Experimental/ChaosVD/Source/ChaosVD/Public/ChaosVDGeometryBuilder.h
+++ b/Engine/Plugins/Experimental/ChaosVD/Source/ChaosVD/Public/ChaosVDGeometryBuilder.h
@@ -2,27 +2,30 @@
 
 #pragma once
 
-#include "AsyncCompilationHelpers.h"
-#include "ChaosVDConvexMeshGenerator.h"
+#include "ChaosVDExtractedGeometryDataHandle.h"
 #include "ChaosVDGeometryDataComponent.h"
 #include "ChaosVDHeightfieldMeshGenerator.h"
-#include "ChaosVDTriMeshGenerator.h"
-#include "Chaos/HeightField.h"
+#include "ChaosVDMeshComponentPool.h"
+#include "ChaosVDModule.h"
+#include "ChaosVDParticleActor.h"
+#include "ChaosVDScene.h"
+#include "ObjectsWaitingGeometryList.h"
 #include "Chaos/ImplicitObjectScaled.h"
 #include "Chaos/ImplicitObjectType.h"
 #include "Containers/Ticker.h"
 #include "Components/DynamicMeshComponent.h"
 #include "Components/InstancedStaticMeshComponent.h"
 #include "Engine/StaticMesh.h"
-#include "Generators/CapsuleGenerator.h"
-#include "Generators/MinimalBoxMeshGenerator.h"
-#include "Generators/SphereGenerator.h"
+
 #include "UDynamicMesh.h"
-#include "Framework/Notifications/NotificationManager.h"
+#include "Components/ChaosVDInstancedStaticMeshComponent.h"
+#include "UObject/WeakObjectPtr.h"
 #include "Tasks/Task.h"
+#include "Templates/SharedPointer.h"
 #include "UObject/GCObject.h"
 #include "UObject/UObjectGlobals.h"
 
+class UChaosVDStaticMeshComponent;
 class AActor;
 class UDynamicMesh;
 class UDynamicMeshComponent;
@@ -36,7 +39,8 @@ namespace UE
 	}
 }
 
-typedef TMap<uint32, TArray<TWeakObjectPtr<UMeshComponent>>> FChaosVDWaitListMeshMap;
+typedef TWeakObjectPtr<UMeshComponent> FMeshComponentWeakPtr;
+typedef TSharedPtr<FChaosVDExtractedGeometryDataHandle> FExtractedGeometryHandle;
 
 /*
  * Generates Dynamic mesh components and dynamic meshes based on Chaos implicit object data
@@ -45,15 +49,17 @@ class FChaosVDGeometryBuilder : public FGCObject, public TSharedFromThis<FChaosV
 {
 public:
 
-	FChaosVDGeometryBuilder() : GeometryGenerationNotification(NSLOCTEXT("ChaosVisualDebugger", "GeometryGenNotification","Particle Geometry"))
+	FChaosVDGeometryBuilder()
 	{
-		GameThreadTickDelegate = FTSTicker::GetCoreTicker().AddTicker(FTickerDelegate::CreateRaw(this, &FChaosVDGeometryBuilder::GameThreadTick));
+		FInstancedStaticMeshDelegates::OnInstanceIndexUpdated.AddRaw(this, &FChaosVDGeometryBuilder::HandleStaticMeshComponentInstanceIndexUpdated);
 	}
 
 	virtual ~FChaosVDGeometryBuilder() override
 	{
 		FTSTicker::GetCoreTicker().RemoveTicker(GameThreadTickDelegate);
 
+		FInstancedStaticMeshDelegates::OnInstanceIndexUpdated.RemoveAll(this);
+
 		for (const TPair<uint32, TObjectPtr<UStaticMesh>>& StaticMeshByKey : StaticMeshCacheMap)
 		{
 			if (StaticMeshByKey.Value)
@@ -72,16 +78,24 @@ public:
 			}
 		}
 	}
+	
+	void Initialize(const TWeakPtr<FChaosVDScene>& ChaosVDScene);
 
 	/** Creates Dynamic Mesh components for each object within the provided Implicit object
 	 *	@param InImplicitObject : Implicit object to process
 	 *	@param Owner Actor who will own the generated components
-	 *	@param OutMeshComponents Array containing all the generated components
-	 *	@param Index Index of the current component being processed. This is useful when this method is called recursively
-	 *	@param Transform to apply to the generated components/geometry
+	 *	@param OutMeshDataHandles Array containing all the generated components
+	 *	@param DesiredLODCount Number of LODs to generate for the mesh.
+	 *	@param MeshIndex Index of the current component being processed. This is useful when this method is called recursively
+	 *	@param InTransform to apply to the generated components/geometry
 	 */
-	template<typename MeshType, typename ComponentType>
-	void CreateMeshComponentsFromImplicit(const Chaos::FImplicitObject* InImplicitObject, AActor* Owner, TArray<TWeakObjectPtr<UMeshComponent>>& OutMeshComponents, Chaos::FRigidTransform3& Transform, const int32 Index = 0, const int32 DesiredLODCount = 0);
+	
+	template<typename MeshType>
+	void CreateMeshesFromImplicitObject(const Chaos::FImplicitObject* InImplicitObject, AActor* Owner, TArray<TSharedPtr<FChaosVDExtractedGeometryDataHandle>>& OutMeshDataHandles, const int32 DesiredLODCount = 0, const Chaos::FRigidTransform3& InTransform = Chaos::FRigidTransform3(), const int32 MeshIndex = 0)
+	{
+		// To start set the leaf and the root to the same ptr. If the object is an union, in the subsequent recursive call the leaf will be set correctly
+		CreateMeshesFromImplicit_Internal<MeshType>(InImplicitObject, InImplicitObject, Owner, OutMeshDataHandles, DesiredLODCount, InTransform, MeshIndex);
+	}
 	
 	virtual void AddReferencedObjects(FReferenceCollector& Collector) override;
 	virtual FString GetReferencerName() const override
@@ -97,14 +111,18 @@ public:
 	 */
 	static bool DoesImplicitContainType(const Chaos::FImplicitObject* InImplicitObject, const Chaos::EImplicitObjectType ImplicitTypeToCheck);
 
-private:
-
 	/**
 	 * Evaluates the provided transform's scale, and returns true if the scale has a negative component
 	 * @param InTransform Transform to evaluate
 	 */
-	bool HasNegativeScale(const Chaos::FRigidTransform3& InTransform) const;
+	static bool HasNegativeScale(const Chaos::FRigidTransform3& InTransform);
 
+private:
+	
+	template<typename MeshType>
+	void CreateMeshesFromImplicit_Internal(const Chaos::FImplicitObject* InRootImplicitObject,const Chaos::FImplicitObject* InLeafImplicitObject, AActor* Owner, TArray<TSharedPtr<FChaosVDExtractedGeometryDataHandle>>& OutMeshDataHandles, const int32 DesiredLODCount = 0, const Chaos::FRigidTransform3& InTransform = Chaos::FRigidTransform3(), const int32 MeshIndex = 0);
+	
+public:
 	/**
 	 * Return true if we have cached geometry for the provided Geometry Key
 	 * @param GeometryKey Cache key for the geometry we are looking for
@@ -112,40 +130,113 @@ private:
 	bool HasGeometryInCache(uint32 GeometryKey);
 	bool HasGeometryInCache_AssumesLocked(uint32 GeometryKey) const;
 
-	/** Creates a Dynamic Mesh for the provided Implicit object and generator, and then caches it to be reused later
-	 * @param GeometryCacheKey Key to be used to find this geometry in the cache
-	 * @param MeshGenerator Generator class with the data and rules to create the mesh
-	 */
-	UDynamicMesh* CreateAndCacheDynamicMesh(const uint32 GeometryCacheKey, UE::Geometry::FMeshShapeGenerator& MeshGenerator);
-
 	/** Returns an already mesh for the provided implicit object if exists, otherwise returns null
 	 * @param GeometryCacheKey Key to be used to find this geometry in the cache
 	 */
 	template<typename MeshType>
 	MeshType* GetCachedMeshForImplicit(const uint32 GeometryCacheKey);
 
+private:
+	/** Creates a Dynamic Mesh for the provided Implicit object and generator, and then caches it to be reused later
+	 * @param GeometryCacheKey Key to be used to find this geometry in the cache
+	 * @param MeshGenerator Generator class with the data and rules to create the mesh
+	 */
+	UDynamicMesh* CreateAndCacheDynamicMesh(const uint32 GeometryCacheKey, UE::Geometry::FMeshShapeGenerator& MeshGenerator);
+
 	/** Creates a Dynamic Mesh for the provided Implicit object and generator, and then caches it to be reused later
 	 * @param GeometryCacheKey Key to be used to find this geometry in the cache
 	 * @param MeshGenerator Generator class with the data and rules to create the mesh
+	 * @param LODsToGenerateNum Number of LODs to generate for this static mesh
 	 */
 	UStaticMesh* CreateAndCacheStaticMesh(const uint32 GeometryCacheKey, UE::Geometry::FMeshShapeGenerator& MeshGenerator, const int32 LODsToGenerateNum = 0);
 
+	/** Takes a Mesh component ptr and initializes it to be used with the provided owner
+	 * @param Owner Actor that will own the provided Mesh Component
+	 * @param MeshComponent Component to initialize
+	 */
+	template <class ComponentType>
+    bool InitializeMeshComponent(AActor* Owner, ComponentType* MeshComponent);
+
+public:
+	/**
+	 * Finds or creates a Mesh component for the geometry data handle provided, and add a new instance of that geometry to it
+	 * @param InOwningParticleData Particle data from which the implicit object is from
+	 * @param InExtractedGeometryDataHandle Handle to the extracted geometry data the new component will use
+	 * */
+	template<typename ComponentType>
+	TSharedPtr<FChaosVDMeshDataInstanceHandle> CreateMeshDataInstance(const FChaosVDParticleDataWrapper& InOwningParticleData, const TSharedPtr<FChaosVDExtractedGeometryDataHandle>& InExtractedGeometryDataHandle);
+
+	/**
+	 * Finds or creates a Mesh component compatible with the provided mesh data handle, and updates the handle to use that new component.
+	 * This is used when data in the handle changed and becomes no longer compatible with the mesh component in use.
+	 * @param HandleToUpdate Instance handle we need to update to a new component
+	 * @param MeshAttributes Attributes of the mesh that the new component needs to be compativle with
+	 * */
+	template<typename ComponentType>
+	void UpdateMeshDataInstance(TSharedPtr<FChaosVDMeshDataInstanceHandle> HandleToUpdate, EChaosVDMeshAttributesFlags MeshAttributes);
+
 	/**
-	 * Creates an empty DynamicMeshComponent/Static Mesh Component or Instanced mesh component and adds it to the actor
-	 * @param Owner Actor who will own the component
-	 * @param Name Name of the component. It has to be unique within the components in the owner actor
-	 * @param Transform Transform to apply as Relative Transform in the component after its creating and attachment to the provided actor
-	 * @param DataComponentKey Key to match this component data to a specific geometry. Usually the type hash of the implicit object
+	 * Destroys a Mesh component that will not longer be used.
+	 * If pooling is enabled, the component will be reset and added back to the pool
+	 * @param MeshComponent Component to Destroy
 	 * */
+	void DestroyMeshComponent(UMeshComponent* MeshComponent);
+
+private:
+
+	/** Gets a ptr to a fully initialized Mesh component compatible with the provided geometry handle and mesh attribute flags, ready to accept a new mesh instance
+	 * @param GeometryDataHandle Handle with the data of the geometry to be used for the new instance
+	 * @param MeshAttributes Set of flags that the mesh component needs to be compatible with
+	 */
 	template<typename ComponentType>
-	ComponentType* CreateMeshComponent(AActor* Owner, const FString& Name, const Chaos::FRigidTransform3& Transform, uint32 DataComponentKey) const;
+	ComponentType* GetMeshComponentForNewInstance(const TSharedPtr<FChaosVDExtractedGeometryDataHandle>& GeometryDataHandle, EChaosVDMeshAttributesFlags MeshAttributes);
+
+	/** Gets a reference to the correct instanced static mesh component cache that is compatible with the provided mesh attribute flags
+	 * @param MeshAttributeFlags GeometryHandle this component will render
+	 */
+	TMap<uint32, UChaosVDInstancedStaticMeshComponent*>& GetInstancedStaticMeshComponentCacheMap(EChaosVDMeshAttributesFlags MeshAttributeFlags);
+
+	/** Gets any available instanced static mesh component that is compatible with the provided mesh attributes and component type
+	 * @param InExtractedGeometryDataHandle GeometryHandle this component will render
+	 * @param MeshComponentsContainerActor Actor that will owns the mesh component
+	 * @param MeshComponentAttributeFlags Set of flags that the mesh component needs to be compatible with
+	 * @param bOutIsNewComponent True if the component we are returning is new. False if it is an existing one but that can take a new mesh instance
+	 */
+	template <typename ComponentType>
+	ComponentType* GetAvailableInstancedStaticMeshComponent(const TSharedPtr<FChaosVDExtractedGeometryDataHandle>& InExtractedGeometryDataHandle, AActor* MeshComponentsContainerActor, EChaosVDMeshAttributesFlags MeshComponentAttributeFlags, bool& bOutIsNewComponent);
+
+	/** Gets any available mesh component that is compatible with the provided mesh attributes and component type
+	 * @param InExtractedGeometryDataHandle GeometryHandle this component will render
+	 * @param MeshComponentsContainerActor Actor that will owns the mesh component
+	 * @param MeshComponentAttributeFlags Set of flags that the mesh component needs to be compatible with
+	 * @param bOutIsNewComponent True if the component we are returning is new. False if it is an existing one but that can take a new mesh instance
+	 */
+	template <typename ComponentType>
+	ComponentType* GetAvailableMeshComponent(const TSharedPtr<FChaosVDExtractedGeometryDataHandle>& InExtractedGeometryDataHandle, AActor* MeshComponentsContainerActor, EChaosVDMeshAttributesFlags MeshComponentAttributeFlags, bool& bOutIsNewComponent);
 
 	/**
 	 * Applies a mesh to a mesh component based on its type
 	 * @param MeshComponent Mesh component to apply the mesh to
 	 * @param GeometryKey Key to find the mesh to apply in the cache
 	 */
-	void ApplyMeshToComponentFromKey(TWeakObjectPtr<UMeshComponent> MeshComponent, const uint32 GeometryKey);
+	bool ApplyMeshToComponentFromKey(TWeakObjectPtr<UMeshComponent> MeshComponent, const uint32 GeometryKey);
+
+	/** Creates a mesh generator for the provided Implicit object which will be used to create a Static Mesh or Dynamic Mesh */
+	TSharedPtr<UE::Geometry::FMeshShapeGenerator> CreateMeshGeneratorForImplicitObject(const Chaos::FImplicitObject* InImplicit);
+
+	const Chaos::FImplicitObject* UnpackImplicitObject(const Chaos::FImplicitObject* InImplicitObject, Chaos::FRigidTransform3& InOutTransform) const;
+
+	/** Re-adjust the provided transform if needed, so it can be visualized properly with its generated mesh */
+	void AdjustedTransformForImplicit(const Chaos::FImplicitObject* InImplicit, FTransform& OutAdjustedTransform);
+
+	/** Extracts data from an implicit object in a format CVD can use, and starts the Mesh generation process if needed
+	 * @return Returns a handle to the generated data that can be used to access the generated mesh when ready
+	 */
+	template <typename MeshType>
+	TSharedPtr<FChaosVDExtractedGeometryDataHandle> ExtractGeometryDataForImplicit(const Chaos::FImplicitObject* InImplicitObject, const Chaos::FRigidTransform3& InTransform, const int32 Index);
+
+	/** Returns true if the implicit object if of one of the types we need to unpack before generating a mesh for it */
+	bool ImplicitObjectNeedsUnpacking(const Chaos::FImplicitObject* InImplicitObject) const;
 
 	/**
 	 * Creates a Mesh from the provided Implicit object geometry data. This is a async operation, and the mesh will be assigned to the component once is ready
@@ -167,13 +258,25 @@ private:
 	/** Tick method of this Geometry builder. Used to do everything that needs to be performed in the GT, like applying the generated meshes to mesh component */
 	bool GameThreadTick(float DeltaTime);
 
+public:
 	/**
 	 * Add a mesh component to the waiting list for Geometry. This needs to be called before dispatching a generation job for new Geometry
 	 * @param GeometryKey Key of the geometry that is being generated (or will be generated)
-	 * @param MesComponent Component where the geometry needs to be applied
-	 * @param LODsToGenerateNum How many LODs we expect the generated mesh to have. Used to check we are using the correct mesh component type
+	 * @param MeshComponent Component where the geometry needs to be applied
 	 */
-	void RegisterMeshComponentWaitingForGeometry(uint32 GeometryKey, TWeakObjectPtr<UMeshComponent> MesComponent, const int32 LODsToGenerateNum);
+	void AddMeshComponentWaitingForGeometry(uint32 GeometryKey, TWeakObjectPtr<UMeshComponent> MeshComponent) const;
+
+	/**
+	 * Add a mesh component to the waiting list for Geometry. This needs to be called before dispatching a generation job for new Geometry
+	 * @param GeometryKey Key of the geometry that is being generated (or will be generated)
+	 * @param MeshComponent Component where the geometry needs to be applied
+	 */
+	void RemoveMeshComponentWaitingForGeometry(uint32 GeometryKey, TWeakObjectPtr<UMeshComponent> MeshComponent) const;
+
+private:
+
+	/** Handles any changes to the indexes of created instanced mesh components we are managing, making corrections/updates as needed */
+	void HandleStaticMeshComponentInstanceIndexUpdated(UInstancedStaticMeshComponent* InComponent, TArrayView<const FInstancedStaticMeshDelegates::FInstanceIndexUpdateData> InIndexUpdates);
 
 	/** Map containing already generated dynamic mesh for any given implicit object */
 	TMap<uint32, TObjectPtr<UDynamicMesh>> DynamicMeshCacheMap;
@@ -181,27 +284,40 @@ private:
 	/** Map containing already generated static mesh for any given implicit object */
 	TMap<uint32, TObjectPtr<UStaticMesh>> StaticMeshCacheMap;
 
-	/** Map containing all the meshes component waiting for geometry, by geometry key*/
-	FChaosVDWaitListMeshMap MeshComponentsWaitingForGeometryByKey;
-
 	/** Set of all geometry keys of the Meshes that are being generated but not ready yet */
 	TSet<uint32> GeometryBeingGeneratedByKey;
 	
 	/** Used to lock Read or Writes to the Geometry cache and in flight job tracking containers */
 	FRWLock GeometryCacheRWLock;
 
-	/** Used to lock Read or Writes to the Geometry Waiting list */
-	FRWLock GeometryWaitListRWLock;
-
 	/** Handle to the ticker used to ticker the Geometry Builder in the game thread*/
 	FTSTicker::FDelegateHandle GameThreadTickDelegate;
 
-	FAsyncCompilationNotification GeometryGenerationNotification;
+	/** Object containing all the meshes component waiting for geometry, by geometry key*/
+	TUniquePtr<FObjectsWaitingGeometryList<FMeshComponentWeakPtr>> MeshComponentsWaitingForGeometry;
 
-	std::atomic<int32> MeshComponentsWaitingGeometryNum = 0;
+	/** Object containing all the meshes component waiting for geometry, by geometry key*/
+	TUniquePtr<FObjectsWaitingGeometryList<FExtractedGeometryHandle>> HandlesWaitingForGeometry;
 
-	friend class FGeometryGenerationTask;
+	/** Map containing already initialized Instanced static mesh components for any given geometry key */
+	TMap<uint32, UChaosVDInstancedStaticMeshComponent*> InstancedMeshComponentByGeometryKey;
+
+	/** Map containing already initialized Instanced static mesh components ready to be use with translucent materials, for any given geometry key */
+	TMap<uint32, UChaosVDInstancedStaticMeshComponent*> TranslucentInstancedMeshComponentByGeometryKey;
+
+	/** Map containing already initialized Instanced static mesh components for mesh instances that required a negative scale transform, for any given geometry key */
+	TMap<uint32, UChaosVDInstancedStaticMeshComponent*> MirroredInstancedMeshComponentByGeometryKey;
 
+	/** Map containing already initialized Instanced static mesh components for mesh instances that required a negative scale transform and use a translucent material, for any given geometry key */
+	TMap<uint32, UChaosVDInstancedStaticMeshComponent*> TranslucentMirroredInstancedMeshComponentByGeometryKey;
+
+	/** Instance of uninitialized mesh components pool */
+	FChaosVDMeshComponentPool ComponentMeshPool;
+
+	/** Weak Ptr to the CVD scene owning this geometry builder */
+	TWeakPtr<FChaosVDScene> SceneWeakPtr;
+	
+	friend class FGeometryGenerationTask;
 };
 
 /** Used to execute each individual Geometry Generation task using the data with which was constructed.
@@ -229,309 +345,305 @@ private:
 	int32 LODsToGenerateNum;
 };
 
-template <typename MeshType, typename ComponentType>
-void FChaosVDGeometryBuilder::CreateMeshComponentsFromImplicit(const Chaos::FImplicitObject* InImplicitObject, AActor* Owner, TArray<TWeakObjectPtr<UMeshComponent>>& OutMeshComponents, Chaos::FRigidTransform3& Transform, const int32 Index, const int32 DesiredLODCount)
+template <typename MeshType>
+TSharedPtr<FChaosVDExtractedGeometryDataHandle> FChaosVDGeometryBuilder::ExtractGeometryDataForImplicit(const Chaos::FImplicitObject* InImplicitObject, const Chaos::FRigidTransform3& InTransform, const int32 Index)
 {
-	static_assert(std::is_same_v<MeshType, UStaticMesh> || std::is_same_v<MeshType, UDynamicMesh>, "CreateMeshComponentsFromImplicit Only supports DynamicMesh and Static Mesh");
-	static_assert(std::is_base_of_v<UStaticMeshComponent, ComponentType> || std::is_base_of_v<UInstancedStaticMeshComponent, ComponentType> /*|| std::is_same_v<MeshType, UDynamicMeshComponent>*/, "CreateMeshComponentsFromImplicit Only supports DynamicMeshComponent, Static MeshComponent and Instanced Static Mesh Component");
+	const uint32 ImplicitObjectHash = InImplicitObject->GetTypeHash();
+
+	Chaos::FRigidTransform3 ExtractedTransform = InTransform;
+	const bool bNeedsUnpack = ImplicitObjectNeedsUnpacking(InImplicitObject);
+	if (const Chaos::FImplicitObject* ImplicitObjectToProcess = bNeedsUnpack ? UnpackImplicitObject(InImplicitObject, ExtractedTransform) : InImplicitObject)
+	{
+		TSharedPtr<FChaosVDExtractedGeometryDataHandle> MeshDataHandle = MakeShared<FChaosVDExtractedGeometryDataHandle>();
+
+		const uint32 GeometryKey = ImplicitObjectToProcess->GetTypeHash();		
+		MeshDataHandle->SetGeometryKey(GeometryKey);
+	
+		// For the Component data key, we need the hash of the implicit as it is (packed) because we will need to match it when looking for shape data
+		MeshDataHandle->SetDataComponentKey(bNeedsUnpack ? ImplicitObjectHash : GeometryKey);
+
+		AdjustedTransformForImplicit(ImplicitObjectToProcess, ExtractedTransform);
+		MeshDataHandle->SetTransform(ExtractedTransform);
+
+		if (!HasGeometryInCache(GeometryKey))
+		{
+			if (TSharedPtr<UE::Geometry::FMeshShapeGenerator> MeshGenerator = CreateMeshGeneratorForImplicitObject(ImplicitObjectToProcess))
+			{
+				DispatchCreateAndCacheMeshForImplicitAsync<MeshType>(GeometryKey, MeshGenerator);
+			}
+			else
+			{
+				UE_LOG(LogChaosVDEditor, Error, TEXT("[%s] Failed create geometry data handle | Failed to generate a valid mesh generator from the implicit object. | Geometry Key [%u] "), ANSI_TO_TCHAR(__FUNCTION__), GeometryKey);
+				return nullptr;
+			}
+		}
 
-	// We could have you make the Mesh type as template and infer the component type based on that, but we also want to be able to create Static Meshes with either Instanced or normal static mesh components
-	constexpr bool bHasValidCombinationForStaticMesh =  std::is_same_v<MeshType, UStaticMesh> && (std::is_base_of_v<UStaticMeshComponent, ComponentType> || std::is_base_of_v<UInstancedStaticMeshComponent, ComponentType>);
-	constexpr bool bHasValidCombinationForDynamicMesh =  std::is_same_v<MeshType, UDynamicMesh> && std::is_base_of_v<UDynamicMeshComponent, ComponentType>;
-	static_assert(bHasValidCombinationForStaticMesh || bHasValidCombinationForDynamicMesh , "Incorrect Component type for Mesh type. Did you use a Dynamic Mesh with a Static Mesh component type?.");
+		return MeshDataHandle;
+	}
+
+	return nullptr;
+}
+
+template <typename MeshType>
+void FChaosVDGeometryBuilder::CreateMeshesFromImplicit_Internal(const Chaos::FImplicitObject* InRootImplicitObject, const Chaos::FImplicitObject* InLeafImplicitObject, AActor* Owner, TArray<TSharedPtr<FChaosVDExtractedGeometryDataHandle>>& OutMeshDataHandles, const int32 DesiredLODCount, const Chaos::FRigidTransform3& InTransform, const int32 MeshIndex)
+{
+	TRACE_CPUPROFILER_EVENT_SCOPE(FChaosVDGeometryBuilder::CreateMeshesFromImplicit_Internal);
 
 	using namespace Chaos;
 
-	const EImplicitObjectType InnerType = GetInnerType(InImplicitObject->GetType());
-	const EImplicitObjectType PackedType = InImplicitObject->GetType();
+	const EImplicitObjectType InnerType = GetInnerType(InRootImplicitObject->GetType());
 	
-	if (InnerType == ImplicitObjectType::Union)
+	if (InnerType == ImplicitObjectType::Union || InnerType == ImplicitObjectType::UnionClustered)
 	{
-		if (const FImplicitObjectUnion* Union = InImplicitObject->template AsA<FImplicitObjectUnion>())
+		if (const FImplicitObjectUnion* Union = InLeafImplicitObject->template AsA<FImplicitObjectUnion>())
 		{
-			for (int i = 0; i < Union->GetObjects().Num(); ++i)
+			for (int32 ObjectIndex = 0; ObjectIndex < Union->GetObjects().Num(); ++ObjectIndex)
 			{
-				const FImplicitObjectPtr& UnionImplicit = Union->GetObjects()[i];
+				const FImplicitObjectPtr& UnionImplicit = Union->GetObjects()[ObjectIndex];
 
-				CreateMeshComponentsFromImplicit<MeshType, ComponentType>(UnionImplicit.GetReference(), Owner, OutMeshComponents, Transform, i, DesiredLODCount);	
+				CreateMeshesFromImplicitObject<MeshType>(UnionImplicit.GetReference(), Owner, OutMeshDataHandles, DesiredLODCount, InTransform, ObjectIndex);	
 			}
 		}
 
 		return;
 	}
 
-	if (InnerType ==  ImplicitObjectType::Transformed)
+	if (InnerType == ImplicitObjectType::Transformed)
 	{
-		const TImplicitObjectTransformed<FReal, 3>* Transformed = InImplicitObject->template GetObject<TImplicitObjectTransformed<FReal, 3>>();
-		FRigidTransform3 TransformCopy = Transformed->GetTransform();
-		CreateMeshComponentsFromImplicit<MeshType, ComponentType>(Transformed->GetTransformedObject(), Owner, OutMeshComponents,TransformCopy, Index, DesiredLODCount);
+		if (const TImplicitObjectTransformed<FReal, 3>* Transformed = InLeafImplicitObject->template GetObject<TImplicitObjectTransformed<FReal, 3>>())
+		{
+			CreateMeshesFromImplicitObject<MeshType>(Transformed->GetTransformedObject(), Owner, OutMeshDataHandles, DesiredLODCount, Transformed->GetTransform(), MeshIndex);
+		}
+		
 		return;
 	}
 
-	ComponentType* MeshComponent = nullptr;
+	if (const TSharedPtr<FChaosVDExtractedGeometryDataHandle> MeshDataHandle = ExtractGeometryDataForImplicit<MeshType>(InLeafImplicitObject, InTransform, MeshIndex))
+	{	
+		MeshDataHandle->SetImplicitObject(InLeafImplicitObject);
+		MeshDataHandle->SetRootImplicitObject(InRootImplicitObject);
+
+		OutMeshDataHandles.Add(MeshDataHandle);
+	}
+}
 
-	switch (InnerType)
+template <typename MeshType>
+MeshType* FChaosVDGeometryBuilder::GetCachedMeshForImplicit(const uint32 GeometryCacheKey)
+{
+	if constexpr (std::is_same_v<MeshType, UDynamicMesh>)
 	{
-		case ImplicitObjectType::Sphere:
+		if (const TObjectPtr<UDynamicMesh>* MeshPtrPtr = DynamicMeshCacheMap.Find(GeometryCacheKey))
 		{
-			const Chaos::TSphere<FReal, 3>* Sphere = InImplicitObject->template GetObject<Chaos::TSphere<FReal, 3>>();
-
-			const FString Name = FString::Format(TEXT("{0} - {1}"), {TEXT("Sphere"), FString::FromInt(Index)});
-
-			const uint32 GeometryKey = Sphere->GetTypeHash();
-				
-			MeshComponent = CreateMeshComponent<ComponentType>(Owner, Name, Transform, GeometryKey);
-			RegisterMeshComponentWaitingForGeometry(GeometryKey, MeshComponent, DesiredLODCount);
+			return MeshPtrPtr->Get();
+		}
+	}
+	else if constexpr (std::is_same_v<MeshType, UStaticMesh>)
+	{
+		if (const TObjectPtr<UStaticMesh>* MeshPtrPtr = StaticMeshCacheMap.Find(GeometryCacheKey))
+		{
+			return MeshPtrPtr->Get();
+		}
+	}
 
-			if (!HasGeometryInCache(GeometryKey))
-			{
-				TSharedPtr<UE::Geometry::FSphereGenerator> SphereGen = MakeShared<UE::Geometry::FSphereGenerator>();
-				SphereGen->Radius = Sphere->GetRadius();
-				SphereGen->NumTheta = 50;
-				SphereGen->NumPhi = 50;
-				SphereGen->bPolygroupPerQuad = false;
+	return nullptr;
+}
 
-				DispatchCreateAndCacheMeshForImplicitAsync<MeshType>(GeometryKey, SphereGen);
-			}
+template <typename ComponentType>
+bool FChaosVDGeometryBuilder::InitializeMeshComponent(AActor* Owner, ComponentType* MeshComponent)
+{
+	if (!ensure(MeshComponent))
+	{
+		UE_LOG(LogChaosVDEditor, Error, TEXT("[%s] Failed To Create mesh component | Component Is Null. "), ANSI_TO_TCHAR(__FUNCTION__));
+		return false;
+	}
 
-			break;
-		}
-		case ImplicitObjectType::Box:
-		{
-			const Chaos::TBox<FReal, 3>* Box = InImplicitObject->template GetObject<Chaos::TBox<FReal, 3>>();
+	if (Owner)
+	{
+		Owner->AddOwnedComponent(MeshComponent);
+		Owner->AddInstanceComponent(MeshComponent);
+		MeshComponent->RegisterComponent();
+		MeshComponent->AttachToComponent(Owner->GetRootComponent(), FAttachmentTransformRules::SnapToTargetIncludingScale);
 
-			const FString Name = FString::Format(TEXT("{0} - {1}"), {TEXT("Box"), FString::FromInt(Index)});
-			
-			const uint32 GeometryKey = Box->GetTypeHash();
+		MeshComponent->bSelectable = true;
+	}
+	else
+	{
+		UE_LOG(LogChaosVDEditor, Error, TEXT("[%s] Failed To Register Component | Owner Is Null. "), ANSI_TO_TCHAR(__FUNCTION__));
+		return false;
+	}
 
-			MeshComponent = CreateMeshComponent<ComponentType>(Owner, Name, Transform, GeometryKey);
-			RegisterMeshComponentWaitingForGeometry(GeometryKey, MeshComponent, DesiredLODCount);
+	return true;
+}
 
-			if (!HasGeometryInCache(GeometryKey))
-			{
-				TSharedPtr<UE::Geometry::FMinimalBoxMeshGenerator> BoxGen = MakeShared<UE::Geometry::FMinimalBoxMeshGenerator>();
-				UE::Geometry::FOrientedBox3d OrientedBox;
-				OrientedBox.Frame = UE::Geometry::FFrame3d(Box->Center());
-				OrientedBox.Extents = Box->Extents() * 0.5;
-				BoxGen->Box = OrientedBox;
-				DispatchCreateAndCacheMeshForImplicitAsync<MeshType>(GeometryKey, BoxGen);
-			}
+template <typename ComponentType>
+ComponentType* FChaosVDGeometryBuilder::GetAvailableInstancedStaticMeshComponent(const TSharedPtr<FChaosVDExtractedGeometryDataHandle>& InExtractedGeometryDataHandle, AActor* MeshComponentsContainerActor, EChaosVDMeshAttributesFlags MeshComponentAttributeFlags, bool& bOutIsNewComponent)
+{
+	// Get the correct Instanced Mesh Component from the existing cache
+	TMap<uint32, UChaosVDInstancedStaticMeshComponent*>& InstancedMeshComponentMapToSearch = GetInstancedStaticMeshComponentCacheMap(MeshComponentAttributeFlags);
 
-			break;
-		}
-		case ImplicitObjectType::Plane:
-			break;
-		case ImplicitObjectType::Capsule:
+	if (UChaosVDInstancedStaticMeshComponent** FoundInstancedMeshComponent = InstancedMeshComponentMapToSearch.Find(InExtractedGeometryDataHandle->GetGeometryKey()))
+	{
+		bOutIsNewComponent = false;
+		return Cast<ComponentType>(*FoundInstancedMeshComponent);
+	}
+	else
+	{
+		// If no exiting component meets our requirements, get a new one form the pool
+		ComponentType* Component = ComponentMeshPool.AcquireMeshComponent<ComponentType>(MeshComponentsContainerActor, InExtractedGeometryDataHandle->GetName());
+		if (!InitializeMeshComponent<ComponentType>(MeshComponentsContainerActor, Component))
 		{
-			const FCapsule* Capsule = InImplicitObject->template GetObject<FCapsule>();
-
-			const FString Name = FString::Format(TEXT("{0} - {1}"), {TEXT("Capsule"), FString::FromInt(Index)});
-			const FRigidTransform3 StartingTransform;
+			return nullptr;
+		}
 
-			const uint32 GeometryKey = Capsule->GetTypeHash();
-				
-			MeshComponent = CreateMeshComponent<ComponentType>(Owner, Name, StartingTransform, GeometryKey);
+		bOutIsNewComponent = true;
 
-			// Re-adjust the location so the pivot is not the center of the capsule, and transform it based on the provided transform
-			const FVector FinalLocation = Transform.TransformPosition(Capsule->GetCenter() - Capsule->GetAxis() * Capsule->GetSegment().GetLength() * 0.5f);
-			const FQuat Rotation = FRotationMatrix::MakeFromZ(Capsule->GetAxis()).Rotator().Quaternion();
+		InstancedMeshComponentMapToSearch.Add(InExtractedGeometryDataHandle->GetGeometryKey(), Component);
 
-			MeshComponent->SetRelativeRotation(Transform.GetRotation() * Rotation);
-			MeshComponent->SetRelativeLocation(FinalLocation);
-			MeshComponent->SetRelativeScale3D(Transform.GetScale3D());
+		return Component;
+	}
 
-			RegisterMeshComponentWaitingForGeometry(GeometryKey, MeshComponent, DesiredLODCount);
+	return nullptr;
+}
 
-			if (!HasGeometryInCache(GeometryKey))
-			{
-				TSharedPtr<UE::Geometry::FCapsuleGenerator> CapsuleGenerator = MakeShared<UE::Geometry::FCapsuleGenerator>();
-				CapsuleGenerator->Radius = FMath::Max(FMathf::ZeroTolerance, Capsule->GetRadius());
-				CapsuleGenerator->SegmentLength = FMath::Max(FMathf::ZeroTolerance, Capsule->GetSegment().GetLength());
-				CapsuleGenerator->NumHemisphereArcSteps = 12;
-				CapsuleGenerator->NumCircleSteps = 12;
-				CapsuleGenerator->bPolygroupPerQuad = false;
-	
-				DispatchCreateAndCacheMeshForImplicitAsync<MeshType>(GeometryKey, CapsuleGenerator);
-			}
-			break;
-		}
-		case ImplicitObjectType::LevelSet:
+template <typename ComponentType>
+ComponentType* FChaosVDGeometryBuilder::GetAvailableMeshComponent(const TSharedPtr<FChaosVDExtractedGeometryDataHandle>& InExtractedGeometryDataHandle, AActor* MeshComponentsContainerActor, EChaosVDMeshAttributesFlags MeshComponentAttributeFlags, bool& bOutIsNewComponent)
+{
+	ComponentType* MeshComponent = nullptr;
+	if constexpr (std::is_base_of_v<UInstancedStaticMeshComponent, ComponentType>)
+	{
+		MeshComponent = GetAvailableInstancedStaticMeshComponent<ComponentType>(InExtractedGeometryDataHandle, MeshComponentsContainerActor, MeshComponentAttributeFlags, bOutIsNewComponent);
+	}
+	else
+	{
+		MeshComponent = ComponentMeshPool.AcquireMeshComponent<ComponentType>(MeshComponentsContainerActor, InExtractedGeometryDataHandle->GetName());
+		if (!InitializeMeshComponent<ComponentType>(MeshComponentsContainerActor, MeshComponent))
 		{
-			//TODO: Implement
-			break;
+			return nullptr;
 		}
-		break;
-		case ImplicitObjectType::Convex:
-		{
-			if (const FConvex* Convex = GetGeometryBasedOnPackedType<FConvex>(InImplicitObject, Transform, PackedType))
-			{
-				const FString Name = FString::Format(TEXT("{0} - {1}"), {TEXT("Convex"), FString::FromInt(Index)});
 
-				uint32 DataComponentKey = InImplicitObject->GetTypeHash();
-				MeshComponent = CreateMeshComponent<ComponentType>(Owner, Name, Transform, DataComponentKey);
-
-				// For the Cache key, we need the hash of the geometry itself without scale or transformations, because these will be applied to the
-				// mesh component. For example, we don't want to generate two meshes for a box, because in one instance was scaled. We only one mesh for one box, and scale the component as needed
-				const uint32 GeometryKey = Convex->GetTypeHash();
-				RegisterMeshComponentWaitingForGeometry(GeometryKey, MeshComponent, DesiredLODCount);
-		
-				if (!HasGeometryInCache(GeometryKey))
-				{
-					TSharedPtr<FChaosVDConvexMeshGenerator> ConvexMeshGen = MakeShared<FChaosVDConvexMeshGenerator>();
-					ConvexMeshGen->GenerateFromConvex(*Convex);
+		bOutIsNewComponent = true;
+	}
 
-					DispatchCreateAndCacheMeshForImplicitAsync<MeshType>(GeometryKey ,ConvexMeshGen);
-				}
-			}
+	return MeshComponent;
+}
 
-			break;
-		}
-		case ImplicitObjectType::TaperedCylinder:
-			break;
-		case ImplicitObjectType::Cylinder:
-			break;
-		case ImplicitObjectType::TriangleMesh:
-		{
-			if (const FTriangleMeshImplicitObject* TriangleMesh = GetGeometryBasedOnPackedType<FTriangleMeshImplicitObject>(InImplicitObject, Transform, PackedType))
-			{
-				const FString Name = FString::Format(TEXT("{0} - {1}"), {TEXT("Trimesh"), FString::FromInt(Index)});
 
-				// For the Component data key, we need the hash of the implicit as it is as we will need to match it when looking for shape data
-				uint32 DataComponentKey = InImplicitObject->GetTypeHash();
-				MeshComponent = CreateMeshComponent<ComponentType>(Owner, Name, Transform, DataComponentKey);
+template <typename ComponentType>
+TSharedPtr<FChaosVDMeshDataInstanceHandle> FChaosVDGeometryBuilder::CreateMeshDataInstance(const FChaosVDParticleDataWrapper& InOwningParticleData, const TSharedPtr<FChaosVDExtractedGeometryDataHandle>& InExtractedGeometryDataHandle)
+{
+	static_assert(std::is_base_of_v<UChaosVDStaticMeshComponent, ComponentType> || std::is_base_of_v<UChaosVDInstancedStaticMeshComponent, ComponentType>, "CreateMeshComponentsFromImplicit Only supports CVD versions of Static MeshComponent and Instanced Static Mesh Component");
 
-				// For the Cache key, we need the hash of the geometry itself without scale or transformations, because these will be applied to the
-				// mesh component. For example, we don't want to generate two meshes for a box, because in one instance was scaled. We only one mesh for one box, and scale the component as needed
-				const uint32 GeometryKey = TriangleMesh->GetTypeHash();
-				RegisterMeshComponentWaitingForGeometry(GeometryKey, MeshComponent, DesiredLODCount);
-	
-				if (!HasGeometryInCache(GeometryKey))
-				{
-					TSharedPtr<FChaosVDTriMeshGenerator> TriMeshGen = MakeShared<FChaosVDTriMeshGenerator>();
-					TriMeshGen->bReverseOrientation = true;
-					TriMeshGen->GenerateFromTriMesh(*TriangleMesh);
-
-					DispatchCreateAndCacheMeshForImplicitAsync<MeshType>(GeometryKey, TriMeshGen);
-				}
-			}
-				
-			break;
-		}
-		case ImplicitObjectType::HeightField:
-		{
-			if (const FHeightField* HeightField = GetGeometryBasedOnPackedType<FHeightField>(InImplicitObject, Transform, PackedType))
-			{
-				const FString Name = FString::Format(TEXT("{0} - {1}"), {TEXT("HeightField"), FString::FromInt(Index)});
+	if (!InExtractedGeometryDataHandle.IsValid())
+	{
+		UE_LOG(LogChaosVDEditor, Error, TEXT("[%s] Failed To Create mesh Component | Handle is invalid. "), ANSI_TO_TCHAR(__FUNCTION__));
+		return nullptr;
+	}
 
-				// For the Component data key, we need the hash of the implicit as it is as we will need to match it when looking for shape data
-				uint32 DataComponentKey = InImplicitObject->GetTypeHash();
-				MeshComponent = CreateMeshComponent<ComponentType>(Owner, Name, Transform, DataComponentKey);
+	const FTransform ExtractedGeometryTransform = InExtractedGeometryDataHandle->GetRelativeTransform();
 
-				// For the Cache key, we need the hash of the geometry itself without scale or transformations, because these will be applied to the
-				// mesh component. For example, we don't want to generate two meshes for a box, because in one instance was scaled. We only one mesh for one box, and scale the component as needed
-				const uint32 GeometryKey = HeightField->GetTypeHash();
-				RegisterMeshComponentWaitingForGeometry(GeometryKey, MeshComponent, DesiredLODCount);
+	EChaosVDMeshAttributesFlags MeshComponentAttributeFlags = EChaosVDMeshAttributesFlags::None;
+	if (HasNegativeScale(ExtractedGeometryTransform))
+	{
+		EnumAddFlags(MeshComponentAttributeFlags, EChaosVDMeshAttributesFlags::MirroredGeometry);
+	}
 
-				if (!HasGeometryInCache(GeometryKey))
-				{
-					TSharedPtr<FChaosVDHeightFieldMeshGenerator> HeightFieldMeshGen = MakeShared<FChaosVDHeightFieldMeshGenerator>();
-					HeightFieldMeshGen->bReverseOrientation = false;
-					HeightFieldMeshGen->GenerateFromHeightField(*HeightField);
+	TSharedPtr<FChaosVDMeshDataInstanceHandle> MeshComponentHandle;
+	ComponentType* Component = GetMeshComponentForNewInstance<ComponentType>(InExtractedGeometryDataHandle, MeshComponentAttributeFlags);
 
-					DispatchCreateAndCacheMeshForImplicitAsync<MeshType>(GeometryKey, HeightFieldMeshGen, DesiredLODCount);
-				}
-			}
+	if (IChaosVDGeometryComponent* CVDGeometryComponent = Cast<IChaosVDGeometryComponent>(Component))
+	{
+		constexpr bool bIsWorldSpace = true;
+		const FTransform OwningParticleTransform(InOwningParticleData.ParticlePositionRotation.MR, InOwningParticleData.ParticlePositionRotation.MX);
+		MeshComponentHandle = CVDGeometryComponent->AddMeshInstance(OwningParticleTransform, bIsWorldSpace, InExtractedGeometryDataHandle, InOwningParticleData.ParticleIndex, InOwningParticleData.SolverID);
 		
-			break;
-		}
-		default:
-			break;
-		}
-
-		if (MeshComponent != nullptr)
-		{
-			if (IChaosVDGeometryDataComponent* DataComponent = Cast<IChaosVDGeometryDataComponent>(MeshComponent))
-			{
-				DataComponent->SetImplicitObject(InImplicitObject);
-			}
+		MeshComponentHandle->SetGeometryBuilder(AsWeak());
+	}
 
-			OutMeshComponents.Add(MeshComponent);
-		}
+	return MeshComponentHandle;
 }
 
-template <typename MeshType>
-MeshType* FChaosVDGeometryBuilder::GetCachedMeshForImplicit(const uint32 GeometryCacheKey)
+template <typename ComponentType>
+void FChaosVDGeometryBuilder::UpdateMeshDataInstance(TSharedPtr<FChaosVDMeshDataInstanceHandle> HandleToUpdate, EChaosVDMeshAttributesFlags MeshAttributes)
 {
-	if constexpr (std::is_same_v<MeshType, UDynamicMesh>)
+	static_assert(std::is_base_of_v<UChaosVDStaticMeshComponent, ComponentType> || std::is_base_of_v<UChaosVDInstancedStaticMeshComponent, ComponentType>, "CreateMeshComponentsFromImplicit Only supports CVD versions of Static MeshComponent and Instanced Static Mesh Component");
+
+	if (!HandleToUpdate.IsValid())
 	{
-		if (TObjectPtr<MeshType>* MeshPtrPtr = DynamicMeshCacheMap.Find(GeometryCacheKey))
-		{
-			return *MeshPtrPtr;
-		}
+		UE_LOG(LogChaosVDEditor, Error, TEXT("[%s] Failed To Update Mesh Data Handle | Mesh Data Handle is invalid. "), ANSI_TO_TCHAR(__FUNCTION__));
+		return;
 	}
-	else if constexpr (std::is_same_v<MeshType, UStaticMesh>)
+	
+	TSharedPtr<FChaosVDExtractedGeometryDataHandle> ExtractedGeometryDataHandle = HandleToUpdate->GetGeometryHandle();
+	if (!ExtractedGeometryDataHandle.IsValid())
 	{
-		if (TObjectPtr<MeshType>* MeshPtrPtr = StaticMeshCacheMap.Find(GeometryCacheKey))
-		{
-			return *MeshPtrPtr;
-		}
+		UE_LOG(LogChaosVDEditor, Error, TEXT("[%s] Failed To Update Mesh Data Handle | Geometry Data Handle is invalid. "), ANSI_TO_TCHAR(__FUNCTION__));
+		return;
 	}
 
-	return nullptr;
+	ComponentType* Component = GetMeshComponentForNewInstance<ComponentType>(HandleToUpdate->GetGeometryHandle(), MeshAttributes);
+
+	if (IChaosVDGeometryComponent* CVDGeometryComponent = Cast<IChaosVDGeometryComponent>(Component))
+	{
+		constexpr bool bIsWorldSpace = true;
+		CVDGeometryComponent->AddMeshInstanceForHandle(HandleToUpdate, HandleToUpdate->GetWorldTransform(), bIsWorldSpace, HandleToUpdate->GetGeometryHandle(), HandleToUpdate->GetOwningParticleID(), HandleToUpdate->GetOwningSolverID());
+	}
 }
 
 template <typename ComponentType>
-ComponentType* FChaosVDGeometryBuilder::CreateMeshComponent(AActor* Owner, const FString& Name, const Chaos::FRigidTransform3& Transform, uint32 DataComponentKey) const
+ComponentType* FChaosVDGeometryBuilder::GetMeshComponentForNewInstance(const TSharedPtr<FChaosVDExtractedGeometryDataHandle>& GeometryDataHandle, EChaosVDMeshAttributesFlags MeshAttributes)
 {
-	ComponentType* MeshComponent = NewObject<ComponentType>(Owner, *Name);
-
-	if (Owner)
+	if (!GeometryDataHandle.IsValid())
 	{
-		MeshComponent->RegisterComponent();
-		MeshComponent->AttachToComponent(Owner->GetRootComponent(), FAttachmentTransformRules::SnapToTargetIncludingScale);
-		Owner->AddInstanceComponent(MeshComponent);
+		UE_LOG(LogChaosVDEditor, Error, TEXT("[%s] Failed To obtain Mesh Component | Geometry Data Handle is invalid. "), ANSI_TO_TCHAR(__FUNCTION__));
+		return nullptr;
 	}
 
-	MeshComponent->bSelectable = true;
-
-	constexpr bool bIsStaticMeshComponent = std::is_base_of_v<UStaticMeshComponent, ComponentType> && !std::is_base_of_v<UInstancedStaticMeshComponent, ComponentType>;
-
-	if constexpr (std::is_base_of_v<UDynamicMeshComponent, ComponentType> || bIsStaticMeshComponent)
+	AActor* MeshComponentsContainerActor = nullptr;
+	
+	if (const TSharedPtr<FChaosVDScene> CVDScene = SceneWeakPtr.Pin())
 	{
-		MeshComponent->SetRelativeTransform(Transform);
+		MeshComponentsContainerActor = CVDScene->GetMeshComponentsContainerActor();
 	}
-	else if constexpr (std::is_base_of_v<UInstancedStaticMeshComponent, ComponentType>)
+
+	if (!MeshComponentsContainerActor)
 	{
-		// If we have negative scale we need to force reverse the culling mode in this component otherwise the faces will be inverted
-		MeshComponent->SetReverseCulling(HasNegativeScale(Transform));
-		MeshComponent->AddInstance(Transform);
+		return nullptr;
 	}
 
-	if (IChaosVDGeometryDataComponent* DataComponent = Cast<IChaosVDGeometryDataComponent>(MeshComponent))
+	bool bIsNew = false;
+	ComponentType* Component = GetAvailableMeshComponent<ComponentType>(GeometryDataHandle, MeshComponentsContainerActor, MeshAttributes, bIsNew);
+
+	if (IChaosVDGeometryComponent* CVDGeometryComponent = Cast<IChaosVDGeometryComponent>(Component))
 	{
-		DataComponent->SetGeometryID(DataComponentKey);
+		if (!CVDGeometryComponent->IsMeshReady())
+		{
+			AddMeshComponentWaitingForGeometry(GeometryDataHandle->GetGeometryKey(), Component);
+		}
+
+		if (bIsNew)
+		{
+			CVDGeometryComponent->OnComponentEmpty()->AddRaw(this, &FChaosVDGeometryBuilder::DestroyMeshComponent);			
+			CVDGeometryComponent->SetMeshComponentAttributeFlags(MeshAttributes);
+		}	
 	}
 
-	return MeshComponent;
+	return Component;
 }
 
 template <typename MeshType>
 void FChaosVDGeometryBuilder::DispatchCreateAndCacheMeshForImplicitAsync(const uint32 GeometryKey, TSharedPtr<UE::Geometry::FMeshShapeGenerator> MeshGenerator, const int32 LODsToGenerateNum)
 {
 	{
-		FReadScopeLock ReadLock(GeometryCacheRWLock);
+		FWriteScopeLock WriteLock(GeometryCacheRWLock);
 		if (GeometryBeingGeneratedByKey.Contains(GeometryKey))
 		{
 			return;
 		}
-	}
 
-	{
-		FWriteScopeLock WriteLock(GeometryCacheRWLock);
 		GeometryBeingGeneratedByKey.Add(GeometryKey);
 	}
 	
-	TSharedPtr<FGeometryGenerationTask> GenerationTask = MakeShared<FGeometryGenerationTask>(AsWeak(), MeshGenerator, GeometryKey,LODsToGenerateNum);
+	TSharedPtr<FGeometryGenerationTask> GenerationTask = MakeShared<FGeometryGenerationTask>(AsWeak(), MeshGenerator, GeometryKey, LODsToGenerateNum);
 	UE::Tasks::Launch(TEXT("GeometryGeneration"),
 		[GenerationTask]()
 		{
diff --git a/Engine/Plugins/Experimental/ChaosVD/Source/ChaosVD/Public/ChaosVDGeometryDataComponent.h b/Engine/Plugins/Experimental/ChaosVD/Source/ChaosVD/Public/ChaosVDGeometryDataComponent.h
--- a/Engine/Plugins/Experimental/ChaosVD/Source/ChaosVD/Public/ChaosVDGeometryDataComponent.h
+++ b/Engine/Plugins/Experimental/ChaosVD/Source/ChaosVD/Public/ChaosVDGeometryDataComponent.h
@@ -2,6 +2,9 @@
 
 #pragma once
 
+#include "ChaosVDEditorSettings.h"
+#include "ChaosVDExtractedGeometryDataHandle.h"
+#include "InstancedStaticMeshDelegates.h"
 #include "Chaos/ImplicitFwd.h"
 #include "Containers/Array.h"
 #include "Containers/Map.h"
@@ -10,11 +13,16 @@
 #include "HAL/Platform.h"
 #include "UObject/Interface.h"
 #include "UObject/StrongObjectPtr.h"
+#include "Templates/SharedPointer.h"
+#include "Components/MeshComponent.h"
 
 #include "ChaosVDGeometryDataComponent.generated.h"
 
-class IChaosVDGeometryDataComponent;
-class UMeshComponent;
+class UChaosVDStaticMeshComponent;
+class FChaosVDGeometryBuilder;
+class UChaosVDInstancedStaticMeshComponent;
+struct FChaosVDExtractedGeometryDataHandle;
+class IChaosVDGeometryComponent;
 class UMaterialInstanceDynamic;
 
 namespace Chaos
@@ -22,10 +30,154 @@ namespace Chaos
 	class FImplicitObject;
 }
 
-DECLARE_MULTICAST_DELEGATE_OneParam(FChaosVDMeshReadyDelegate, IChaosVDGeometryDataComponent&)
+DECLARE_MULTICAST_DELEGATE_OneParam(FChaosVDMeshReadyDelegate, IChaosVDGeometryComponent&)
+DECLARE_MULTICAST_DELEGATE_OneParam(FChaosVDMeshComponentEmptyDelegate, UMeshComponent*)
+
+UENUM()
+enum class EChaosVDMaterialType
+{
+	SimOnlyMaterial,
+	QueryOnlyMaterial,
+	Instanced,
+	InstancedQueryOnly
+};
+
+UENUM()
+enum class EChaosVDMeshComponent
+{
+	Invalid,
+	Static,
+	InstancedStatic,
+	Dynamic
+};
+
+/** Handle that provides access to a specific mesh instance on a CVD Mesh component (instanced or static)*/
+class FChaosVDMeshDataInstanceHandle : public TSharedFromThis<FChaosVDMeshDataInstanceHandle>
+{
+public:
+
+	explicit FChaosVDMeshDataInstanceHandle(int32 InInstanceIndex, UMeshComponent* InMeshComponent, int32 InParticleID, int32 InSolverID);
+
+	/** Returns the Particle ID of the particle owning this mesh instance */
+	int32 GetOwningParticleID() const { return OwningParticleID; }
+
+	/** Returns the Solver ID of the particle owning this mesh instance */
+	int32 GetOwningSolverID() const { return OwningSolverID; }
+
+	/** Applies the provided world transform to the mesh instance this handle represents
+	 * @param InTransform Transform to apply
+	 */
+	void SetWorldTransform(const FTransform& InTransform);
+
+	/** Returns the world transform of the mesh instance this handle represents */
+	const FTransform& GetWorldTransform() const { return CurrentWorldTransform; }
+
+	/** Sets the Geometry Handle used to create the mesh instance this handle represents */
+	void SetGeometryHandle(const TSharedPtr<FChaosVDExtractedGeometryDataHandle>& InHandle) { ExtractedGeometryHandle = InHandle; };
+	
+	/** Returns the geometry handle used to create the mesh instance this handle represents */
+	TSharedPtr<FChaosVDExtractedGeometryDataHandle> GetGeometryHandle() const { return ExtractedGeometryHandle; }
+
+	/** Applies to the provided color to the mesh instance this handle represents */
+	void SetInstanceColor(const FLinearColor& NewColor);
+
+	/** Returns the current color of the mesh instance this handle represents */
+	FLinearColor GetInstanceColor() const { return CurrentGeometryColor; }
+	
+	/** Applies to the provided shape collision data to the mesh instance this handle represents */
+	void UpdateMeshComponentForCollisionData(const FChaosVDShapeCollisionData& InCollisionData);
+
+	/** Returns the mesh component used to render the mesh instance this handle represents */
+	UMeshComponent* GetMeshComponent() const { return MeshComponent; }
+
+	/** Returns the instance index of the mesh instance this handle represents */
+	int32 GetMeshInstanceIndex() const { return MeshInstanceIndex; }
+
+	/** Returns the type of the component used to render the mesh instance this handle represents */
+	EChaosVDMeshComponent GetMeshComponentType() const { return MeshComponentType; }
+
+	/** Sets a Ptr to the geometry builder used to generate and manage the geometry/mesh components this handle represents */
+	void SetGeometryBuilder(const TWeakPtr<FChaosVDGeometryBuilder>& InGeometryBuilder) { GeometryBuilderInstance = InGeometryBuilder; }
+
+	/** Marks this mesh instance as selected. Used to handle Selection in Editor */
+	void SetIsSelected(bool bInIsSelected);
+
+	/** Sets the visibility of this mesh instance */
+	void SetVisibility(bool bInIsVisible);
+
+	/** Returns the current visibility state this mesh instance */
+	bool GetVisibility() const { return bIsVisible; }
+
+	/** Applies a new shape collision data to this mesh instance */
+	void SetGeometryCollisionData(const FChaosVDShapeCollisionData& InCollisionData);
+
+	/** Returns a copy of the current shape collision data to this mesh instance */
+	FChaosVDShapeCollisionData GetGeometryCollisionData() const { return CollisionData; }
+
+	// TODO: This is legacy from the old system and is used by the Static mesh Component implementation mainly where we can change the material instance on the fly
+	// We need to refactor in a way we can move it to the CVD Static Mesh component and let the Instanced Static Mesh Component have they material instances created on component creation time.
+	/** Returns an existing material instance used by this mesh instances, or creates a new one for the provided type */
+	UMaterialInstanceDynamic* GetCachedMaterialInstance(EChaosVDMaterialType Type);
+
+	/** Handles a mesh instance index update reported by the mesh component used to render this mesh instance */
+	void HandleInstanceIndexUpdated(TArrayView<const FInstancedStaticMeshDelegates::FInstanceIndexUpdateData> InIndexUpdates);
+
+	friend uint32 GetTypeHash(const FChaosVDMeshDataInstanceHandle& Handle)
+	{
+		const uint32 GeometryHandleHash = Handle.ExtractedGeometryHandle ? GetTypeHash(Handle.ExtractedGeometryHandle->GetGeometryKey()) : 0;
+		const uint32 MeshComponentHandleHash = HashCombine(HashCombine(Handle.OwningParticleID, Handle.MeshInstanceIndex), static_cast<uint32>(Handle.MeshComponentType));
+		return HashCombine(GeometryHandleHash, MeshComponentHandleHash);
+	}
+
+private:
+
+	/** Sets the mesh component used to render the mesh instance this handle represents */
+	void SetMeshComponent(UMeshComponent* NewComponent) { MeshComponent = NewComponent; }
+
+	/** Sets the mesh instance index of the mesh instance this handle represents */
+	void SetMeshInstanceIndex(int32 NewIndex) { MeshInstanceIndex = NewIndex; }
+
+	int32 OwningParticleID = INDEX_NONE;
+	int32 OwningSolverID = INDEX_NONE;
+
+	int32 MeshInstanceIndex = INDEX_NONE;
+
+	EChaosVDMeshComponent MeshComponentType = EChaosVDMeshComponent::Invalid;
+
+	TSharedPtr<FChaosVDExtractedGeometryDataHandle> ExtractedGeometryHandle;
+
+	UMeshComponent* MeshComponent;
+
+	FChaosVDShapeCollisionData CollisionData;
+
+	FLinearColor CurrentGeometryColor = FLinearColor(ForceInitToZero);
+
+	FTransform CurrentWorldTransform;
+
+	bool bIsVisible = true;
+
+	bool bIsSelected = false;
+
+	//TODO: We should avoid using a StrongObjtPtr here as it will make difficult to track any leaks
+	TMap<EChaosVDMaterialType, TStrongObjectPtr<UMaterialInstanceDynamic>> MaterialInstancesByID;
+
+	TWeakPtr<FChaosVDGeometryBuilder> GeometryBuilderInstance = nullptr;
+
+	friend UChaosVDInstancedStaticMeshComponent;
+	friend UChaosVDStaticMeshComponent;
+};
+
+enum class EChaosVDMeshAttributesFlags : uint8
+{
+	None = 0,
+	MirroredGeometry = 1 << 0,
+	TranslucentGeometry = 1 << 1
+};
+
+ENUM_CLASS_FLAGS(EChaosVDMeshAttributesFlags)
 
 UINTERFACE()
-class UChaosVDGeometryDataComponent : public UInterface
+class UChaosVDGeometryComponent : public UInterface
 {
 	GENERATED_BODY()
 };
@@ -34,17 +186,48 @@ class UChaosVDGeometryDataComponent : public UInterface
  * Interface with a default implementation for any Geometry component that
  * contains CVD data
  */
-class IChaosVDGeometryDataComponent
+class IChaosVDGeometryComponent
 {
 	GENERATED_BODY()
 
 public:
 
-	/** Returns the Geometry ID used to identify the geometry data this component represents */
-	virtual uint32 GetGeometryID() const PURE_VIRTUAL(IChaosVDGeometryDataComponent::GetGeometryID, return 0;);
-
-	/** Sets the Geometry ID used to identify the geometry data this component represents */
-	virtual void SetGeometryID(uint32 ID) PURE_VIRTUAL(IChaosVDGeometryDataComponent::SetGeometryID);
+	/** Returns the Geometry Handle used to identify the geometry data this component represents */
+	virtual uint32 GetGeometryKey() const PURE_VIRTUAL(IChaosVDGeometryDataComponent::GetGeometryKey, return 0;);
+
+	/** Returns the CVD Mesh Data Instance handle for the provided Instance index */
+	virtual TSharedPtr<FChaosVDMeshDataInstanceHandle> GetMeshDataInstanceHandle(int32 InstanceIndex) const PURE_VIRTUAL(IChaosVDGeometryDataComponent::GetMeshDataInstanceHandle, return nullptr;);
+
+	/** Returns all the CVD Mesh Data Instance handles this component is rendering */
+	virtual TArrayView<TSharedPtr<FChaosVDMeshDataInstanceHandle>> GetMeshDataInstanceHandles() PURE_VIRTUAL(IChaosVDGeometryDataComponent::GetMeshDataInstanceHandles, return TArrayView<TSharedPtr<FChaosVDMeshDataInstanceHandle>>(););
+
+	/**
+	 * Add a new instance to this mesh component
+	 * @param InstanceTransform The transform the new instance will use
+	 * @param bIsWorldSpace True if the provide transform is in world space
+	 * @param InGeometryHandle The CVD Geometry Handle with data about the geometry instantiate
+	 * @param ParticleID Particle ID owning this geometry
+	 * @param SolverID Solver ID of the particle owning this geometry
+	 * @return CVD Mesh instance handle that provides access to this component and specific instance, allowing manipulation of it
+	 */
+	virtual TSharedPtr<FChaosVDMeshDataInstanceHandle> AddMeshInstance(const FTransform InstanceTransform, bool bIsWorldSpace, const TSharedPtr<FChaosVDExtractedGeometryDataHandle>& InGeometryHandle, int32 ParticleID, int32 SolverID) PURE_VIRTUAL(IChaosVDGeometryDataComponent::AddMeshInstance, return nullptr;);
+
+	/**
+	 * Adds a new instance to this mesh component, but using an existing Mesh Data Handle instead of creating a new one
+	 * @param MeshDataHandle Handle that will provide access to this mesh instance
+	 * @param InstanceTransform The transform the new instance will use
+	 * @param bIsWorldSpace True if the provide transform is in world space
+	 * @param InGeometryHandle The CVD Geometry Handle with data about the geometry instantiate
+	 * @param ParticleID Particle ID owning this geometry
+	 * @param SolverID Solver ID of the particle owning this geometry
+	 */
+	virtual void AddMeshInstanceForHandle(TSharedPtr<FChaosVDMeshDataInstanceHandle> MeshDataHandle, const FTransform InstanceTransform, bool bIsWorldSpace, const TSharedPtr<FChaosVDExtractedGeometryDataHandle>& InGeometryHandle, int32 ParticleID, int32 SolverID) PURE_VIRTUAL(IChaosVDGeometryDataComponent::AddMeshInstanceForHandle);
+
+	/**
+	 * Removes the instance the provided handle represents
+	 * @param InHandleToRemove 
+	 */
+	virtual void RemoveMeshInstance(TSharedPtr<FChaosVDMeshDataInstanceHandle> InHandleToRemove) PURE_VIRTUAL(IChaosVDGeometryDataComponent::RemoveMeshInstance);
 
 	/** True if the mesh this component represents is ready for use */
 	virtual bool IsMeshReady() const  PURE_VIRTUAL(IChaosVDGeometryDataComponent::IsMeshReady, return false;);
@@ -55,64 +238,43 @@ public:
 	/** Triggers when the mesh this component represents is ready */
 	virtual FChaosVDMeshReadyDelegate* OnMeshReady() PURE_VIRTUAL(IChaosVDGeometryDataComponent::OnMeshReady, return nullptr;);
 
-	/** Stores a shader ptr to the root implicit object of the implicit object this component represents  */
-	virtual void SetRootImplicitObject(const Chaos::FConstImplicitObjectPtr& InImplicitObject) PURE_VIRTUAL(IChaosVDGeometryDataComponent::SetRootImplicitObject);
+	/** Triggers when the component does not have any instance to render. Used to allow it to return to the mesh component tool for future re-use*/
+	virtual FChaosVDMeshComponentEmptyDelegate* OnComponentEmpty() PURE_VIRTUAL(IChaosVDGeometryDataComponent::OnComponentEmpty, return nullptr;);
 
-	/** Stores a shader ptr to the implicit object of the implicit object this component represents  */
-	virtual void SetImplicitObject(const Chaos::FImplicitObject* InImplicitObject) PURE_VIRTUAL(IChaosVDGeometryDataComponent::SetImplicitObject);
+	/** Updates the visibility of this component based on the stored CVD data*/
+	virtual void UpdateInstanceVisibility(const TSharedPtr<FChaosVDMeshDataInstanceHandle>& InInstanceHandle, bool bIsVisible) PURE_VIRTUAL(IChaosVDGeometryDataComponent::UpdateInstanceVisibility);
 
-	/** Returns a ptr to the CVD Collision Data */
-	virtual FChaosVDShapeCollisionData* GetCollisionData() PURE_VIRTUAL(IChaosVDGeometryDataComponent::GetCollisionData, return nullptr;);
+	/** Changes the selection state of the provided instance - Used for Selection in Editor */
+	virtual void SetIsSelected(const TSharedPtr<FChaosVDMeshDataInstanceHandle>& InInstanceHandle, bool bIsSelected) PURE_VIRTUAL(IChaosVDGeometryDataComponent::SetIsSelected);
 
-	/** Updates the visibility of this component based on the stored CVD data*/
-	virtual void UpdateVisibility() PURE_VIRTUAL(IChaosVDGeometryDataComponent::UpdateVisibility);
+	/** Updates the colors of this component based on the stored CVD data */
+	virtual void UpdateInstanceColor(const TSharedPtr<FChaosVDMeshDataInstanceHandle>& InInstanceHandle, FLinearColor NewColor) PURE_VIRTUAL(IChaosVDGeometryDataComponent::UpdateInstanceColor);
 
 	/** Updates the colors of this component based on the stored CVD data */
-	virtual void UpdateColors() PURE_VIRTUAL(IChaosVDGeometryDataComponent::UpdateColors);
+	virtual void UpdateInstanceWorldTransform(const TSharedPtr<FChaosVDMeshDataInstanceHandle>& InInstanceHandle, const FTransform& InTransform) PURE_VIRTUAL(IChaosVDGeometryDataComponent::UpdateInstanceWorldTransform);
 
-	/** Finds and updates the Shape data using the provided array as source*/
-	virtual void UpdateDataFromShapeArray(const TArray<FChaosVDShapeCollisionData>& InShapeArray) PURE_VIRTUAL(IChaosVDGeometryDataComponent::UpdateDataFromShapeArray);
-};
+	/** Sets the CVD Mesh Attribute flags this component is compatible with*/
+	virtual void SetMeshComponentAttributeFlags(EChaosVDMeshAttributesFlags Flags) PURE_VIRTUAL(IChaosVDGeometryDataComponent::SetMeshComponentAttributeFlags);
+	
+	/** Returns the CVD Mesh Attribute flags this component is compatible with*/
+	virtual EChaosVDMeshAttributesFlags GetMeshComponentAttributeFlags() const PURE_VIRTUAL(IChaosVDGeometryDataComponent::GetMeshComponentAttributeFlags, return EChaosVDMeshAttributesFlags::None;);
 
-UENUM()
-enum class EChaosVDMaterialType
-{
-	SimOnlyMaterial,
-	QueryOnlyMaterial
+	/** Resets the state of this mesh component so it can be re-used later on */
+	virtual void Reset() PURE_VIRTUAL(IChaosVDGeometryDataComponent::Reset);
 };
 
-/** Base Implementation for a GeometryDataComponent */
-class FChaosVDGeometryDataComponentBase
+class FChaosVDGeometryComponentUtils
 {
-	
-protected:
-
-	void UpdateVisibility_Internal(const FChaosVDShapeCollisionData& InCollisionData, UMeshComponent* MeshComponent);
-	
-	void UpdateDataFromShapeArray_Internal(const TArray<FChaosVDShapeCollisionData>& InShapeArray, FChaosVDShapeCollisionData& CollisionDataToUpdate, UMeshComponent* MeshComponent);
-
-	void UpdateColors_Internal(UMeshComponent* MeshComponent);
-
-	void SetImplicitObject_Internal(const Chaos::FImplicitObject* InImplicitObject);
-
-	UMaterialInstanceDynamic* GetCachedMaterialInstance(EChaosVDMaterialType Type);
-
-	const Chaos::FImplicitObject* GetImplicitObject() const;
-	
-	uint32 GeometryID = 0;
-	
-	bool bIsMeshReady = false;
-
-	FLinearColor CurrentGeometryColor = FLinearColor(ForceInitToZero);
-
-	FChaosVDMeshReadyDelegate MeshReadyDelegate;
-
-	TMap<const Chaos::FImplicitObject*, uint32> CachedGeometryHashes;
-
-	Chaos::FConstImplicitObjectPtr RootImplicitObject = nullptr;
+public:
+	/** Finds and updates the Shape data using the provided array as source*/
+	static void UpdateCollisionDataFromShapeArray(const TArray<FChaosVDShapeCollisionData>& InShapeArray, const TSharedPtr<FChaosVDMeshDataInstanceHandle>& InInstanceHandle);
 
-	const Chaos::FImplicitObject* ImplicitObject = nullptr;
+	/** Calculates and updates the color used to render the mesh represented by the provided handle, based on the particle state */
+	static void UpdateMeshColor(const TSharedPtr<FChaosVDMeshDataInstanceHandle>& InInstanceHandle, const FChaosVDParticleDataWrapper& InParticleData, bool bIsServer);
 
-	// TODO: Remove the use of strong objt ptr
-	TMap<EChaosVDMaterialType, TStrongObjectPtr<UMaterialInstanceDynamic>> MaterialInstancesByID;
+	/** Calculates the correct visibility state based on the particle state, and applies it to the mesh instance the provided handle represents */
+	static void UpdateMeshVisibility(const TSharedPtr<FChaosVDMeshDataInstanceHandle>& InInstanceHandle, const FChaosVDParticleDataWrapper& InParticleData, bool bIsActive);
+private:
+	/** Returns the color that needs to be used to present the provided particle data based on its state and current selected options */
+	static FLinearColor GetGeometryParticleColor(const TSharedPtr<FChaosVDExtractedGeometryDataHandle>& InGeometryHandle, const FChaosVDParticleDataWrapper& InParticleData, bool bIsServer);
 };
diff --git a/Engine/Plugins/Experimental/ChaosVD/Source/ChaosVD/Public/ChaosVDMeshComponentPool.h b/Engine/Plugins/Experimental/ChaosVD/Source/ChaosVD/Public/ChaosVDMeshComponentPool.h
new file mode 100644
--- /dev/null
+++ b/Engine/Plugins/Experimental/ChaosVD/Source/ChaosVD/Public/ChaosVDMeshComponentPool.h
@@ -0,0 +1,83 @@
+﻿// Copyright Epic Games, Inc. All Rights Reserved.
+
+#pragma once
+
+#include "ChaosVDModule.h"
+#include "Containers/Array.h"
+
+#include "Components/DynamicMeshComponent.h"
+#include "Components/InstancedStaticMeshComponent.h"
+#include "Components/StaticMeshComponent.h"
+
+class FChaosVDMeshComponentPool : public FGCObject
+{
+public:
+	virtual void AddReferencedObjects(FReferenceCollector& Collector) override;
+	
+	virtual FString GetReferencerName() const override
+	{
+		return TEXT("FChaosVDMeshComponentPool");
+	}
+
+	template<typename TMeshComponent>
+	TMeshComponent* AcquireMeshComponent(UObject* Outer, FName Name);
+
+	void DisposeMeshComponent(UMeshComponent* MeshComponent);
+
+private:
+
+	template <typename TDesiredMeshComponent, typename TPooledComponent>
+	TDesiredMeshComponent* GetMeshComponentFromPool_Internal(TArray<TObjectPtr<TPooledComponent>>& InMeshComponentPool, UObject* Outer, FName Name);
+
+	void ResetMeshComponent(UMeshComponent* MeshComponent);
+
+	TArray<TObjectPtr<UStaticMeshComponent>> PooledStaticMeshComponent;
+	TArray<TObjectPtr<UInstancedStaticMeshComponent>> PooledInstancedStaticMeshComponent;
+	TArray<TObjectPtr<UDynamicMeshComponent>> PooledDynamicMeshComponent;
+
+	static bool bUseComponentsPool;
+	static FAutoConsoleVariableRef CVarUseComponentsPool;
+};
+
+template <typename TMeshComponent>
+TMeshComponent* FChaosVDMeshComponentPool::AcquireMeshComponent(UObject* Outer, FName Name)
+{
+	if constexpr (std::is_base_of_v<UStaticMeshComponent, TMeshComponent> && !std::is_base_of_v<UInstancedStaticMeshComponent, TMeshComponent>)
+	{
+		return GetMeshComponentFromPool_Internal<TMeshComponent>(PooledStaticMeshComponent, Outer, Name);
+	}
+	else if constexpr (std::is_base_of_v<UInstancedStaticMeshComponent, TMeshComponent>)
+	{
+		return GetMeshComponentFromPool_Internal<TMeshComponent>(PooledInstancedStaticMeshComponent, Outer, Name);
+	}
+	else if constexpr (std::is_base_of_v<UDynamicMeshComponent, TMeshComponent>)
+	{
+		return GetMeshComponentFromPool_Internal<TMeshComponent>(PooledDynamicMeshComponent, Outer, Name);
+	}
+
+	return nullptr;
+}
+
+template <typename TDesiredMeshComponent, typename TPooledComponent>
+TDesiredMeshComponent* FChaosVDMeshComponentPool::GetMeshComponentFromPool_Internal(TArray<TObjectPtr<TPooledComponent>>& InMeshComponentPool, UObject* Outer,  FName Name)
+{
+	static_assert(std::is_base_of_v<UStaticMeshComponent, TDesiredMeshComponent> || std::is_base_of_v<UInstancedStaticMeshComponent, TDesiredMeshComponent> || std::is_base_of_v<TDesiredMeshComponent, UDynamicMeshComponent>, "GetMeshComponentInternal Only supports DynamicMeshComponent, Static MeshComponent and Instanced Static Mesh Component");
+
+	// We need to ensure unique names
+	const FString NewName = Name.ToString() + FGuid::NewGuid().ToString();
+	
+	if (bUseComponentsPool &&  InMeshComponentPool.Num() > 0)
+	{
+		if (TDesiredMeshComponent* Component = Cast<TDesiredMeshComponent>(InMeshComponentPool.Pop()))
+		{
+			Component->Rename(*NewName , Outer, REN_NonTransactional | REN_DoNotDirty | REN_ForceNoResetLoaders | REN_SkipGeneratedClasses | REN_DontCreateRedirectors);
+			return Component;
+		}
+		else
+		{
+			UE_LOG(LogChaosVDEditor, Warning, TEXT("[%s] Selected Pool has [%d] objects, but none where castable to the desired type"), ANSI_TO_TCHAR(__FUNCTION__), InMeshComponentPool.Num());
+		}
+	}
+
+	return NewObject<TDesiredMeshComponent>(Outer, *NewName);
+}
diff --git a/Engine/Plugins/Experimental/ChaosVD/Source/ChaosVD/Public/ChaosVDModule.h b/Engine/Plugins/Experimental/ChaosVD/Source/ChaosVD/Public/ChaosVDModule.h
--- a/Engine/Plugins/Experimental/ChaosVD/Source/ChaosVD/Public/ChaosVDModule.h
+++ b/Engine/Plugins/Experimental/ChaosVD/Source/ChaosVD/Public/ChaosVDModule.h
@@ -25,6 +25,8 @@ public:
 
 	TSharedPtr<FChaosVDTraceManager>& GetTraceManager() { return ChaosVDTraceManager; };
 
+	void SpawnCVDTab();
+
 private:
 
 	void RegisterClassesCustomDetails() const;
@@ -45,6 +47,8 @@ private:
 
 	TMap<FGuid, TWeakPtr<SDockTab>> ActiveCVDTabs;
 
+	TArray<FName> CreatedExtraTabSpawnersIDs;
+
 	TSharedPtr<FChaosVDTraceManager> ChaosVDTraceManager;
 
 	bool bIsShuttingDown = false;
diff --git a/Engine/Plugins/Experimental/ChaosVD/Source/ChaosVD/Public/ChaosVDParticleActor.h b/Engine/Plugins/Experimental/ChaosVD/Source/ChaosVD/Public/ChaosVDParticleActor.h
--- a/Engine/Plugins/Experimental/ChaosVD/Source/ChaosVD/Public/ChaosVDParticleActor.h
+++ b/Engine/Plugins/Experimental/ChaosVD/Source/ChaosVD/Public/ChaosVDParticleActor.h
@@ -13,6 +13,7 @@
 
 #include "ChaosVDParticleActor.generated.h"
 
+class FChaosVDParticleActorCustomization;
 enum class EChaosVDParticleDataVisualizationFlags : uint32;
 class FChaosVDScene;
 struct FChaosVDParticleDebugData;
@@ -35,18 +36,19 @@ enum class EChaosVDActorGeometryUpdateFlags : int32
 };
 ENUM_CLASS_FLAGS(EChaosVDActorGeometryUpdateFlags)
 
+DECLARE_DELEGATE(FChaosVDParticleDataUpdatedDelegate)
+
 /** Actor used to represent a Chaos Particle in the Visual Debugger's world */
 UCLASS(HideCategories=(Transform))
 class AChaosVDParticleActor : public AActor, public IChaosVDParticleVisualizationDataProvider, public IChaosVDVisualizerContainerInterface,
 								public FChaosVDSceneObjectBase, public IChaosVDCollisionDataProviderInterface
 {
-
 	GENERATED_BODY()
 
 public:
 	AChaosVDParticleActor(const FObjectInitializer& ObjectInitializer);
 
-	void UpdateFromRecordedParticleData(const FChaosVDParticleDataWrapper& InRecordedData, const Chaos::FRigidTransform3& SimulationTransform);
+	void UpdateFromRecordedParticleData(const TSharedPtr<FChaosVDParticleDataWrapper>& InRecordedData, const Chaos::FRigidTransform3& SimulationTransform);
 
 	void UpdateGeometry(const Chaos::FConstImplicitObjectPtr& InImplicitObject, EChaosVDActorGeometryUpdateFlags OptionsFlags = EChaosVDActorGeometryUpdateFlags::None);
 
@@ -56,7 +58,7 @@ public:
 
 	virtual void BeginDestroy() override;
 
-	virtual const FChaosVDParticleDataWrapper* GetParticleData() override { return &ParticleDataViewer; }
+	virtual const FChaosVDParticleDataWrapper* GetParticleData() override { return ParticleDataPtr.Get(); }
 	virtual void GetVisualizationContext(FChaosVDVisualizationContext& OutVisualizationContext) override;
 
 	void CreateVisualizers();
@@ -80,19 +82,31 @@ public:
 	 */
 	bool IsActive() const { return bIsActive; }
 
+	virtual FBox GetComponentsBoundingBox(bool bNonColliding, bool bIncludeFromChildActors) const override;
+
 	//BEGIN IChaosVDCollisionDataProvider Interface
 	virtual void GetCollisionData(TArray<TSharedPtr<FChaosVDCollisionDataFinder>>& OutCollisionDataFound) override;
 	virtual bool HasCollisionData() override;
-	virtual FName GetName() override;
+	virtual FName GetProviderName() override;
 	//END IChaosVDCollisionDataProvider Interface
 
+	void SetIsServerParticle(bool bNewIsServer) { bIsServer = bNewIsServer; }
+	bool GetIsServerParticle() const { return bIsServer; }
+
+	virtual void PushSelectionToProxies() override;
+
+	FChaosVDParticleDataUpdatedDelegate& OnParticleDataUpdated() { return ParticleDataUpdatedDelegate; };
+
 protected:
 
+	void ProcessUpdatedAndRemovedHandles(TArray<TSharedPtr<FChaosVDExtractedGeometryDataHandle>>& OutExtractedGeometryDataHandles);
+
 	const TArray<TSharedPtr<FChaosVDParticlePairMidPhase>>* GetCollisionMidPhasesArray() const;
 
 	void UpdateShapeDataComponents();
 
-	void PerformTaskOnGeometryComponents(TFunction<void(IChaosVDGeometryDataComponent& InDataComponent)> TaskToPerform);
+	template<typename TTaskCallback>
+	void PerformTaskOnGeometry(const TTaskCallback& TaskToPerform);
 	
 	UPROPERTY(EditAnywhere, Category = "Viewport Visualization Flags", meta = (Bitmask, BitmaskEnum = "/Script/ChaosVD.EChaosVDParticleDataVisualizationFlags"))
 	uint8 LocalParticleDataVisualizationFlags;
@@ -100,18 +114,35 @@ protected:
 	UPROPERTY(EditAnywhere, Category = "Viewport Visualization Flags")
 	bool bShowDebugText = false;
 
-	UPROPERTY(EditAnywhere, Category = "Particle Data")
-	FChaosVDParticleDataWrapper ParticleDataViewer;
+	TSharedPtr<FChaosVDParticleDataWrapper> ParticleDataPtr;
 
 	FTransform CachedSimulationTransform;
 
 	bool bIsGeometryDataGenerationStarted = false;
 
-	TArray<TWeakObjectPtr<UMeshComponent>> MeshComponents;
-
 	FDelegateHandle GeometryUpdatedDelegate;
 
 	TMap<FStringView, TUniquePtr<FChaosVDDataVisualizerBase>> CVDVisualizers;
 
+	TArray<TSharedPtr<FChaosVDMeshDataInstanceHandle>> MeshDataHandles;
+
+	FChaosVDParticleDataUpdatedDelegate ParticleDataUpdatedDelegate;
+
 	bool bIsActive = false;
+
+	bool bIsServer = false;
+
+	friend FChaosVDParticleActorCustomization;
 };
+
+template <typename TTaskCallback>
+void AChaosVDParticleActor::PerformTaskOnGeometry(const TTaskCallback& TaskToPerform)
+{
+	for (TSharedPtr<FChaosVDMeshDataInstanceHandle>& MeshDataHandle : MeshDataHandles)
+	{
+		if (MeshDataHandle)
+		{
+			TaskToPerform(MeshDataHandle);
+		}
+	}
+}
diff --git a/Engine/Plugins/Experimental/ChaosVD/Source/ChaosVD/Public/ChaosVDParticleActorCustomization.h b/Engine/Plugins/Experimental/ChaosVD/Source/ChaosVD/Public/ChaosVDParticleActorCustomization.h
--- a/Engine/Plugins/Experimental/ChaosVD/Source/ChaosVD/Public/ChaosVDParticleActorCustomization.h
+++ b/Engine/Plugins/Experimental/ChaosVD/Source/ChaosVD/Public/ChaosVDParticleActorCustomization.h
@@ -4,12 +4,17 @@
 
 #include "Templates/SharedPointer.h"
 #include "IDetailCustomization.h"
+#include "DataWrappers/ChaosVDParticleDataWrapper.h"
+
+struct FChaosVDParticleDataWrapper;
+class AChaosVDParticleActor;
 
 /** Custom details panel for the ChaosVD Particle Actor */
 class FChaosVDParticleActorCustomization : public IDetailCustomization
 {
 public:
 	FChaosVDParticleActorCustomization();
+	virtual ~FChaosVDParticleActorCustomization() override;
 
 	inline static FName ChaosVDCategoryName = FName("Particle Data");
 	inline static FName ChaosVDVisualizationCategoryName = FName("Viewport Visualization Flags");
@@ -20,4 +25,9 @@ public:
 
 private:
 	TSet<FName> AllowedCategories;
+	TWeakObjectPtr<AChaosVDParticleActor> CurrentObservedActor;
+
+	void HandleParticleDataUpdated();
+
+	FChaosVDParticleDataWrapper CurrentParticleDataCopy;
 };
diff --git a/Engine/Plugins/Experimental/ChaosVD/Source/ChaosVD/Public/ChaosVDPlaybackViewportClient.h b/Engine/Plugins/Experimental/ChaosVD/Source/ChaosVD/Public/ChaosVDPlaybackViewportClient.h
--- a/Engine/Plugins/Experimental/ChaosVD/Source/ChaosVD/Public/ChaosVDPlaybackViewportClient.h
+++ b/Engine/Plugins/Experimental/ChaosVD/Source/ChaosVD/Public/ChaosVDPlaybackViewportClient.h
@@ -35,7 +35,7 @@ public:
 private:
 	void HandleObjectFocused(UObject* FocusedObject);
 	void HandleActorMoving(AActor* MovedActor) const;
-	void HandleFarClippingOverrideSettingsChanged(UChaosVDEditorSettings* SettingsObject);
+	void HandleViewportSettingsChanged(UChaosVDEditorSettings* SettingsObject);
 	
 	FDelegateHandle ObjectFocusedDelegateHandle;
 	UWorld* CVDWorld;
diff --git a/Engine/Plugins/Experimental/ChaosVD/Source/ChaosVD/Public/ChaosVDScene.h b/Engine/Plugins/Experimental/ChaosVD/Source/ChaosVD/Public/ChaosVDScene.h
--- a/Engine/Plugins/Experimental/ChaosVD/Source/ChaosVD/Public/ChaosVDScene.h
+++ b/Engine/Plugins/Experimental/ChaosVD/Source/ChaosVD/Public/ChaosVDScene.h
@@ -103,6 +103,8 @@ public:
 
 	AActor* GetSkySphereActor() const { return SkySphere; }
 
+	AActor* GetMeshComponentsContainerActor() const { return MeshComponentContainerActor; }
+
 	FChaosVDActorActiveStateUpdateDelegate& OnActorActiveStateChanged() { return ParticleActorUpdateDelegate; }
 
 	TSharedPtr<FChaosVDRecording> LoadedRecording;
@@ -110,12 +112,14 @@ public:
 private:
 
 	/** Creates an ChaosVDParticle actor for the Provided recorded Particle Data */
-	AChaosVDParticleActor* SpawnParticleFromRecordedData(const FChaosVDParticleDataWrapper& InParticleData, const FChaosVDSolverFrameData& InFrameData);
+	AChaosVDParticleActor* SpawnParticleFromRecordedData(const TSharedPtr<FChaosVDParticleDataWrapper>& InParticleData, const FChaosVDSolverFrameData& InFrameData);
 
 	/** Returns the ID used to track this recorded particle data */
-	int32 GetIDForRecordedParticleData(const FChaosVDParticleDataWrapper& InParticleData) const;
+	int32 GetIDForRecordedParticleData(const TSharedPtr<FChaosVDParticleDataWrapper>& InParticleData) const;
 	void CreateBaseLights(UWorld* TargetWorld) const;
 
+	AActor* CreateMeshComponentsContainer(UWorld* TargetWorld);
+
 	/** Creates the instance of the World which will be used the recorded data*/
 	UWorld* CreatePhysicsVDWorld();
 
@@ -167,6 +171,8 @@ private:
 
 	mutable AActor* SkySphere = nullptr;
 
+	AActor* MeshComponentContainerActor = nullptr;
+
 	bool bIsInitialized = false;
 
 	FChaosVDActorActiveStateUpdateDelegate ParticleActorUpdateDelegate;
diff --git a/Engine/Plugins/Experimental/ChaosVD/Source/ChaosVD/Public/Components/ChaosVDInstancedStaticMeshComponent.h b/Engine/Plugins/Experimental/ChaosVD/Source/ChaosVD/Public/Components/ChaosVDInstancedStaticMeshComponent.h
--- a/Engine/Plugins/Experimental/ChaosVD/Source/ChaosVD/Public/Components/ChaosVDInstancedStaticMeshComponent.h
+++ b/Engine/Plugins/Experimental/ChaosVD/Source/ChaosVD/Public/Components/ChaosVDInstancedStaticMeshComponent.h
@@ -8,43 +8,70 @@
 
 /** CVD version of a Instance Static Mesh Component that holds additional CVD data */
 UCLASS(HideCategories=("Transform"), MinimalAPI)
-class UChaosVDInstancedStaticMeshComponent : public UInstancedStaticMeshComponent, public IChaosVDGeometryDataComponent, public FChaosVDGeometryDataComponentBase
+class UChaosVDInstancedStaticMeshComponent : public UInstancedStaticMeshComponent, public IChaosVDGeometryComponent
 {
 	GENERATED_BODY()
-	
+
 	UChaosVDInstancedStaticMeshComponent(const FObjectInitializer& ObjectInitializer)
 	: Super(ObjectInitializer)
 	{
 		SetCanEverAffectNavigation(false);
-		bNavigationRelevant = false;
+		bHasPerInstanceHitProxies = true;
 	}
 	
 	// BEGIN IChaosVDGeometryDataComponent Interface
-	virtual FChaosVDShapeCollisionData* GetCollisionData() override { return &CollisionData; }
-	
-	virtual uint32 GetGeometryID() const override { return GeometryID; }
 	
-	virtual void  SetGeometryID(uint32 ID) override { GeometryID = ID; }
+	virtual uint32 GetGeometryKey() const override;
+
+	virtual TSharedPtr<FChaosVDMeshDataInstanceHandle> GetMeshDataInstanceHandle(int32 InstanceIndex) const override;
 
 	virtual bool IsMeshReady() const override { return bIsMeshReady; }
 	
 	virtual void SetIsMeshReady(bool bIsReady) override { bIsMeshReady = bIsReady; }
 
 	virtual FChaosVDMeshReadyDelegate* OnMeshReady() override { return &MeshReadyDelegate; }
+
+	virtual FChaosVDMeshComponentEmptyDelegate* OnComponentEmpty() override { return &ComponentEmptyDelegate; }
+
+	virtual void UpdateInstanceVisibility(const TSharedPtr<FChaosVDMeshDataInstanceHandle>& InInstanceHandle, bool bIsVisible) override;
+
+	virtual void SetIsSelected(const TSharedPtr<FChaosVDMeshDataInstanceHandle>& InInstanceHandle, bool bIsSelected) override;
+
+	virtual void UpdateInstanceColor(const TSharedPtr<FChaosVDMeshDataInstanceHandle>& InInstanceHandle, FLinearColor NewColor) override;
 	
-	virtual void SetRootImplicitObject(const Chaos::FConstImplicitObjectPtr& InImplicitObject) override;
-	
-	virtual void UpdateVisibility() override;
-	
-	virtual void UpdateDataFromShapeArray(const TArray<FChaosVDShapeCollisionData>& InShapeArray) override;
-	
-	virtual void SetImplicitObject(const Chaos::FImplicitObject* InImplicitObject) override;
+	virtual void UpdateInstanceWorldTransform(const TSharedPtr<FChaosVDMeshDataInstanceHandle>& InInstanceHandle, const FTransform& InTransform) override;
 
-	virtual void UpdateColors() override;
+	virtual void SetMeshComponentAttributeFlags(EChaosVDMeshAttributesFlags Flags) override { MeshComponentAttributeFlags = static_cast<uint8>(Flags); }
+	virtual EChaosVDMeshAttributesFlags GetMeshComponentAttributeFlags() const override {return static_cast<EChaosVDMeshAttributesFlags>(MeshComponentAttributeFlags); };
+
+	virtual TArrayView<TSharedPtr<FChaosVDMeshDataInstanceHandle>> GetMeshDataInstanceHandles() override;
+
+	virtual void Reset() override;
 	// END IChaosVDGeometryDataComponent Interface
 
+	virtual TSharedPtr<FChaosVDMeshDataInstanceHandle> AddMeshInstance(const FTransform InstanceTransform, bool bIsWorldSpace, const TSharedPtr<FChaosVDExtractedGeometryDataHandle>& InGeometryHandle, int32 ParticleID, int32 SolverID) override;
+	virtual void AddMeshInstanceForHandle(TSharedPtr<FChaosVDMeshDataInstanceHandle> MeshDataHandle, const FTransform InstanceTransform, bool bIsWorldSpace, const TSharedPtr<FChaosVDExtractedGeometryDataHandle>& InGeometryHandle, int32 ParticleID, int32 SolverID) override;
+	virtual void RemoveMeshInstance(TSharedPtr<FChaosVDMeshDataInstanceHandle> InHandleToRemove) override;
+
+	virtual bool Modify(bool bAlwaysMarkDirty) override;
+
+	virtual bool IsNavigationRelevant() const override;
+
 protected:
 	
-	UPROPERTY(VisibleAnywhere, Category="Geometry Data", meta=(FullyExpand = true))
-	FChaosVDShapeCollisionData CollisionData;
+	bool UpdateGeometryKey(uint32 NewHandleGeometryKey);
+
+	uint8 MeshComponentAttributeFlags = 0;
+
+	uint32 CurrentGeometryKey = 0;
+
+	bool bIsMeshReady = false;
+
+	UPROPERTY()
+	UMaterialInterface* CurrentMaterial = nullptr;
+
+	FChaosVDMeshReadyDelegate MeshReadyDelegate;
+	FChaosVDMeshComponentEmptyDelegate ComponentEmptyDelegate;
+
+	TArray<TSharedPtr<FChaosVDMeshDataInstanceHandle>> CurrentInstanceHandles;
 };
diff --git a/Engine/Plugins/Experimental/ChaosVD/Source/ChaosVD/Public/Components/ChaosVDStaticMeshComponent.h b/Engine/Plugins/Experimental/ChaosVD/Source/ChaosVD/Public/Components/ChaosVDStaticMeshComponent.h
--- a/Engine/Plugins/Experimental/ChaosVD/Source/ChaosVD/Public/Components/ChaosVDStaticMeshComponent.h
+++ b/Engine/Plugins/Experimental/ChaosVD/Source/ChaosVD/Public/Components/ChaosVDStaticMeshComponent.h
@@ -8,7 +8,7 @@
 
 /** CVD version of a Static Mesh Component that holds additional CVD data */
 UCLASS(HideCategories=("Transform"), MinimalAPI)
-class UChaosVDStaticMeshComponent : public UStaticMeshComponent, public IChaosVDGeometryDataComponent, public FChaosVDGeometryDataComponentBase
+class UChaosVDStaticMeshComponent : public UStaticMeshComponent, public IChaosVDGeometryComponent
 {
 	GENERATED_BODY()
 
@@ -21,11 +21,6 @@ public:
 	}
 
 	// BEGIN IChaosVDGeometryDataComponent Interface
-	virtual FChaosVDShapeCollisionData* GetCollisionData() override { return &CollisionData; }
-	
-	virtual uint32 GetGeometryID() const override { return GeometryID; }
-	
-	virtual void  SetGeometryID(uint32 ID) override { GeometryID = ID; }
 
 	virtual bool IsMeshReady() const override { return bIsMeshReady; }
 	
@@ -33,19 +28,43 @@ public:
 
 	virtual FChaosVDMeshReadyDelegate* OnMeshReady() override { return &MeshReadyDelegate; }
 	
-	virtual void SetRootImplicitObject(const Chaos::FConstImplicitObjectPtr& InImplicitObject) override;
-	
-	virtual void UpdateVisibility() override;
-	
-	virtual void UpdateDataFromShapeArray(const TArray<FChaosVDShapeCollisionData>& InShapeArray) override;
+	virtual FChaosVDMeshComponentEmptyDelegate* OnComponentEmpty() override { return &ComponentEmptyDelegate; }
 
-	virtual void SetImplicitObject(const Chaos::FImplicitObject* InImplicitObject) override;
+	virtual uint32 GetGeometryKey() const override;
+	virtual void UpdateInstanceVisibility(const TSharedPtr<FChaosVDMeshDataInstanceHandle>& InInstanceHandle, bool bIsVisible) override;
+
+	virtual void SetIsSelected(const TSharedPtr<FChaosVDMeshDataInstanceHandle>& InInstanceHandle, bool bIsSelected) override;
+
+	virtual bool ShouldRenderSelected() const override;
 	
-	virtual void UpdateColors() override;
+	virtual void UpdateInstanceColor(const TSharedPtr<FChaosVDMeshDataInstanceHandle>& InInstanceHandle, FLinearColor NewColor) override;
+	virtual void UpdateInstanceWorldTransform(const TSharedPtr<FChaosVDMeshDataInstanceHandle>& InInstanceHandle, const FTransform& InTransform) override;
+
+	virtual void SetMeshComponentAttributeFlags(EChaosVDMeshAttributesFlags Flags) override { MeshComponentAttributeFlags = static_cast<uint8>(Flags); }
+	virtual EChaosVDMeshAttributesFlags GetMeshComponentAttributeFlags() const override {return static_cast<EChaosVDMeshAttributesFlags>(MeshComponentAttributeFlags); };
+
+	virtual TSharedPtr<FChaosVDMeshDataInstanceHandle> GetMeshDataInstanceHandle(int32 InstanceIndex) const override;
+	virtual TArrayView<TSharedPtr<FChaosVDMeshDataInstanceHandle>> GetMeshDataInstanceHandles() override;
+
+	virtual void Reset() override;
+
+	virtual TSharedPtr<FChaosVDMeshDataInstanceHandle> AddMeshInstance(const FTransform InstanceTransform, bool bIsWorldSpace, const TSharedPtr<FChaosVDExtractedGeometryDataHandle>& InGeometryHandle, int32 ParticleID, int32 SolverID) override;
+	virtual void AddMeshInstanceForHandle(TSharedPtr<FChaosVDMeshDataInstanceHandle> MeshDataHandle, const FTransform InstanceTransform, bool bIsWorldSpace, const TSharedPtr<FChaosVDExtractedGeometryDataHandle>& InGeometryHandle, int32 ParticleID, int32 SolverID) override;
+	virtual void RemoveMeshInstance(TSharedPtr<FChaosVDMeshDataInstanceHandle> InHandleToRemove) override;
 	// END IChaosVDGeometryDataComponent Interface
 
 protected:
-	
-	UPROPERTY(VisibleAnywhere, Category="Geometry Data", meta=(FullyExpand = true))
-	FChaosVDShapeCollisionData CollisionData;
+
+	bool UpdateGeometryKey(uint32 NewHandleGeometryKey);
+
+	uint8 MeshComponentAttributeFlags = 0;
+	uint8 CurrentGeometryKey = 0;
+	bool bIsMeshReady = false;
+	bool bIsOwningParticleSelected = false;
+	FChaosVDMeshReadyDelegate MeshReadyDelegate;
+	FChaosVDMeshComponentEmptyDelegate ComponentEmptyDelegate;
+
+	TSharedPtr<FChaosVDMeshDataInstanceHandle> CurrentMeshDataHandle = nullptr;
+
+	TSharedPtr<FChaosVDExtractedGeometryDataHandle> CurrentGeometryHandle = nullptr;
 };
diff --git a/Engine/Plugins/Experimental/ChaosVD/Source/ChaosVD/Public/ObjectsWaitingGeometryList.h b/Engine/Plugins/Experimental/ChaosVD/Source/ChaosVD/Public/ObjectsWaitingGeometryList.h
new file mode 100644
--- /dev/null
+++ b/Engine/Plugins/Experimental/ChaosVD/Source/ChaosVD/Public/ObjectsWaitingGeometryList.h
@@ -0,0 +1,145 @@
+﻿// Copyright Epic Games, Inc. All Rights Reserved.
+
+#pragma once
+
+#include "AsyncCompilationHelpers.h"
+#include "Containers/Map.h"
+#include "Containers/Array.h"
+#include "Templates/Function.h"
+#include "Templates/SharedPointer.h"
+
+class FTextFormat;
+
+struct FChaosVDWailingObjectListCVars
+{
+	static int32 NumOfGTTaskToProcessPerTick;
+	static FAutoConsoleVariableRef CVarChaosVDGeometryToProcessPerTick;
+};
+
+/** Object that handles objects waiting for geometry to be ready to perform a desired operation
+ * The operation has to run on the Game Thread, and there is a limit on how many operations will be processed per tick
+ */
+template<typename TObjectToProcess>
+class FObjectsWaitingGeometryList
+{
+
+public:
+	 FObjectsWaitingGeometryList(const TFunction<bool(uint32, TObjectToProcess)>& InObjectProcessorCallback, const FTextFormat& InProgressNotificationNameFormat, const TFunction<bool(uint32)>& InShouldProcessObjectsOverride) :
+		  ObjectProcessorCallback(InObjectProcessorCallback)
+		, ShouldProcessObjectsForKeyOverride(InShouldProcessObjectsOverride)
+		, AsyncProgressNotification(InProgressNotificationNameFormat)
+	{
+	}
+
+	/** Adds an object to process when the provided geometry key is ready */
+	void AddObject(uint32 GeometryKey, TObjectToProcess ObjectToProcess);
+
+	/** Removed an object to process */
+	void RemoveObject(uint32 GeometryKey, TObjectToProcess& OutToProcess);
+
+	/** Evaluates the current geometry available, and process any objects on the waiting list for it */
+	bool ProcessWaitingObjects(int32& OutCurrentElementsProcessed);
+
+private:
+
+	bool ShouldProcessObjectsForKey(uint32 GeometryKey) const
+	{
+		if (ShouldProcessObjectsForKeyOverride)
+		{
+			return ShouldProcessObjectsForKeyOverride(GeometryKey);
+		}
+		return true;
+	}
+
+	TFunction<bool(uint32, TObjectToProcess)> ObjectProcessorCallback;
+
+	TFunction<bool(uint32)> ShouldProcessObjectsForKeyOverride;
+
+	TMap<uint32, TArray<TObjectToProcess>> WaitingObjectsByGeometryKey;
+
+	FAsyncCompilationNotification AsyncProgressNotification;
+
+	int32 QueuedObjectsToProcessNum = 0;
+};
+
+template <typename TObjectToProcess>
+void FObjectsWaitingGeometryList<TObjectToProcess>::AddObject(uint32 GeometryKey, TObjectToProcess ObjectToProcess)
+{
+	check(IsInGameThread());
+
+	if (TArray<TObjectToProcess>* ObjectsInQueueForKey = WaitingObjectsByGeometryKey.Find(GeometryKey))
+	{
+		ObjectsInQueueForKey->Add(ObjectToProcess);
+	}
+	else
+	{
+		WaitingObjectsByGeometryKey.Add(GeometryKey, {ObjectToProcess});
+	}
+
+	++QueuedObjectsToProcessNum;
+}
+
+template <typename TObjectToProcess>
+void FObjectsWaitingGeometryList<TObjectToProcess>::RemoveObject(uint32 GeometryKey, TObjectToProcess& OutToProcess)
+{
+	check(IsInGameThread());
+
+	if (TArray<TObjectToProcess>* QueueObjectsForKey = WaitingObjectsByGeometryKey.Find(GeometryKey))
+	{
+		TArray<TObjectToProcess>& QueueObjectsForKeyRef = *QueueObjectsForKey;
+		for(typename TArray<TObjectToProcess>::TIterator ObjectToRemoveIterator = QueueObjectsForKeyRef.CreateIterator();  ObjectToRemoveIterator; ++ObjectToRemoveIterator)
+		{
+			TObjectToProcess& ObjectToRemove = *ObjectToRemoveIterator;
+
+			if (ObjectToRemove == OutToProcess)
+			{
+				ObjectToRemoveIterator.RemoveCurrent();
+
+				--QueuedObjectsToProcessNum;
+			}
+		}		
+	}
+}
+
+template <typename TObjectToProcess>
+bool FObjectsWaitingGeometryList<TObjectToProcess>::ProcessWaitingObjects(int32& OutCurrentElementsProcessed)
+{
+	AsyncProgressNotification.Update(QueuedObjectsToProcessNum);
+
+	bool bCanContinueProcessing = true;
+	for (typename TMap<uint32, TArray<TObjectToProcess>>::TIterator QueuedObjectRemoveIterator = WaitingObjectsByGeometryKey.CreateIterator(); QueuedObjectRemoveIterator; ++QueuedObjectRemoveIterator)
+	{
+		const uint32 GeometryKey = QueuedObjectRemoveIterator.Key();
+
+		if (ShouldProcessObjectsForKey(GeometryKey))
+		{
+			for (typename TArray<TObjectToProcess>::TIterator ObjectRemoveIterator = QueuedObjectRemoveIterator.Value().CreateIterator(); ObjectRemoveIterator; ++ObjectRemoveIterator)
+			{
+				if (ObjectProcessorCallback && ObjectProcessorCallback(GeometryKey, *ObjectRemoveIterator))
+				{
+					OutCurrentElementsProcessed++;
+					--QueuedObjectsToProcessNum;
+					bCanContinueProcessing = OutCurrentElementsProcessed < FChaosVDWailingObjectListCVars::NumOfGTTaskToProcessPerTick;
+					ObjectRemoveIterator.RemoveCurrent();
+				}
+			}
+		}
+
+		if (!bCanContinueProcessing)
+		{
+			break;
+		}	
+
+		if (QueuedObjectRemoveIterator.Value().IsEmpty())
+		{
+			QueuedObjectRemoveIterator.RemoveCurrent();
+		}
+
+		if (!bCanContinueProcessing)
+		{
+			break;
+		}
+	}
+
+	return bCanContinueProcessing;
+}
diff --git a/Engine/Plugins/Experimental/ChaosVD/Source/ChaosVD/Public/Visualizers/ChaosVDDebugDrawUtils.h b/Engine/Plugins/Experimental/ChaosVD/Source/ChaosVD/Public/Visualizers/ChaosVDDebugDrawUtils.h
--- a/Engine/Plugins/Experimental/ChaosVD/Source/ChaosVD/Public/Visualizers/ChaosVDDebugDrawUtils.h
+++ b/Engine/Plugins/Experimental/ChaosVD/Source/ChaosVD/Public/Visualizers/ChaosVDDebugDrawUtils.h
@@ -15,6 +15,7 @@ public:
 	static void DrawText(FStringView StringToDraw, const FVector& Location, const FColor& Color);
 	static void DrawCircle(FPrimitiveDrawInterface* PDI, const FVector& Origin, float Radius, int32 Segments, const FColor& Color, float Thickness, const FVector& XAxis, const FVector& YAxis, FStringView DebugText, ESceneDepthPriorityGroup DepthPriority = SDPG_World);
 	static void DrawBox(FPrimitiveDrawInterface* PDI, const FVector& InExtents, const FColor& InColor, const FTransform& InTransform, FStringView DebugText, ESceneDepthPriorityGroup DepthPriority = SDPG_World);
+	static void DrawLine(FPrimitiveDrawInterface* PDI, const FVector& InStartPosition, const FVector& InEndPosition, const FColor& InColor, FStringView DebugText, ESceneDepthPriorityGroup DepthPriority = SDPG_World);
 
 	static void DrawCanvas(FViewport& InViewport, FSceneView& View, FCanvas& Canvas);
 
diff --git a/Engine/Plugins/Experimental/ChaosVD/Source/ChaosVD/Public/Visualizers/ChaosVDParticleDataVisualizer.h b/Engine/Plugins/Experimental/ChaosVD/Source/ChaosVD/Public/Visualizers/ChaosVDParticleDataVisualizer.h
--- a/Engine/Plugins/Experimental/ChaosVD/Source/ChaosVD/Public/Visualizers/ChaosVDParticleDataVisualizer.h
+++ b/Engine/Plugins/Experimental/ChaosVD/Source/ChaosVD/Public/Visualizers/ChaosVDParticleDataVisualizer.h
@@ -18,6 +18,7 @@ enum class EChaosVDParticleDataVisualizationFlags : uint32
 	AngularAcceleration = 1 << 3,
 	LinearImpulse		= 1 << 4,
 	AngularImpulse		= 1 << 5,
+	ClusterConnectivityEdge	= 1 << 6,
 };
 ENUM_CLASS_FLAGS(EChaosVDParticleDataVisualizationFlags);
 
diff --git a/Engine/Source/Runtime/Experimental/Chaos/Private/Chaos/ParticleHandle.cpp b/Engine/Source/Runtime/Experimental/Chaos/Private/Chaos/ParticleHandle.cpp
--- a/Engine/Source/Runtime/Experimental/Chaos/Private/Chaos/ParticleHandle.cpp
+++ b/Engine/Source/Runtime/Experimental/Chaos/Private/Chaos/ParticleHandle.cpp
@@ -115,6 +115,8 @@ namespace Chaos
 
 				// Needed to update the shapes array.
 				SetGeometry(GeometryParticles->GetGeometry(ParticleIdx));
+
+				CVD_TRACE_INVALIDATE_CACHED_GEOMETRY(CurrentGeometry);
 			}
 		}
 	}
@@ -170,6 +172,8 @@ namespace Chaos
 					return;
 				}
 			}
+
+			CVD_TRACE_INVALIDATE_CACHED_GEOMETRY(CurrentGeometry);
 		}
 	}
 
@@ -193,6 +197,8 @@ namespace Chaos
 
 		Union->RemoveAtSortedIndices(InIndices);
 
+		CVD_TRACE_INVALIDATE_CACHED_GEOMETRY(CurrentGeometry);
+
 		// Needed to update the shapes array.
 		// @todo(chaos): is it though? Maybe for the bounds etc?
 		SetGeometry(GeometryParticles->GetGeometry(ParticleIdx));
diff --git a/Engine/Source/Runtime/Experimental/Chaos/Private/ChaosVisualDebugger/ChaosVDDataWrapperUtils.cpp b/Engine/Source/Runtime/Experimental/Chaos/Private/ChaosVisualDebugger/ChaosVDDataWrapperUtils.cpp
--- a/Engine/Source/Runtime/Experimental/Chaos/Private/ChaosVisualDebugger/ChaosVDDataWrapperUtils.cpp
+++ b/Engine/Source/Runtime/Experimental/Chaos/Private/ChaosVisualDebugger/ChaosVDDataWrapperUtils.cpp
@@ -80,6 +80,12 @@ FChaosVDParticleDataWrapper FChaosVDDataWrapperUtils::BuildParticleDataWrapperFr
 		WrappedParticleData.ParticleDynamicsMisc.CopyFrom(*RigidParticle);
 		WrappedParticleData.ParticleMassProps.CopyFrom(*RigidParticle);
 	}
+
+	if (const Chaos::TPBDRigidClusteredParticleHandleImp<Chaos::FReal, 3, true>* ClusteredParticle = ParticleHandlePtr->CastToClustered())
+	{
+		WrappedParticleData.ParticleCluster.CopyFrom(*ClusteredParticle);
+	}
+
 	return MoveTemp(WrappedParticleData);
 }
 
diff --git a/Engine/Source/Runtime/Experimental/Chaos/Private/ChaosVisualDebugger/ChaosVisualDebuggerTrace.cpp b/Engine/Source/Runtime/Experimental/Chaos/Private/ChaosVisualDebugger/ChaosVisualDebuggerTrace.cpp
--- a/Engine/Source/Runtime/Experimental/Chaos/Private/ChaosVisualDebugger/ChaosVisualDebuggerTrace.cpp
+++ b/Engine/Source/Runtime/Experimental/Chaos/Private/ChaosVisualDebugger/ChaosVisualDebuggerTrace.cpp
@@ -54,7 +54,7 @@ struct FChaosVDGeometryTraceContext
 
 	uint32 GetGeometryHashForImplicit(const Chaos::FImplicitObject* Implicit)
 	{
-		if (!ensure(Implicit != nullptr))
+		if (Implicit == nullptr)
 		{
 			return 0;
 		}
@@ -560,6 +560,16 @@ void FChaosVisualDebuggerTrace::TraceImplicitObject(FChaosVDImplicitObjectWrappe
 	TraceBinaryData(TLSDataBuffer.BufferRef, FChaosVDImplicitObjectWrapper::WrapperTypeName);
 }
 
+void FChaosVisualDebuggerTrace::InvalidateGeometryFromCache(const Chaos::FImplicitObject* CachedGeometryToInvalidate)
+{
+	if (!IsTracing())
+	{
+		return;
+	}
+
+	GeometryTracerObject.RemoveCachedGeometryHash(CachedGeometryToInvalidate);
+}
+
 void FChaosVisualDebuggerTrace::TraceNonSolverLocation(const FVector& InLocation, FStringView DebugNameID)
 {
 	if (!IsTracing())
diff --git a/Engine/Source/Runtime/Experimental/Chaos/Public/Chaos/ImplicitObjectType.h b/Engine/Source/Runtime/Experimental/Chaos/Public/Chaos/ImplicitObjectType.h
--- a/Engine/Source/Runtime/Experimental/Chaos/Public/Chaos/ImplicitObjectType.h
+++ b/Engine/Source/Runtime/Experimental/Chaos/Public/Chaos/ImplicitObjectType.h
@@ -71,4 +71,35 @@ namespace Chaos
 
 		inline const int32 FiniteConvex = IsConvex | HasBoundingBox;
 	}
-}
\ No newline at end of file
+
+
+	inline FName GetImplicitObjectTypeName(EImplicitObjectType Val)
+	{
+		static const FName Names[] =
+		{
+				TEXT("EImplicitObjectType:Sphere"),
+				TEXT("EImplicitObjectType:Box"),
+				TEXT("EImplicitObjectType:Plane"),
+				TEXT("EImplicitObjectType:Capsule"),
+				TEXT("EImplicitObjectType:Transformed"),
+				TEXT("EImplicitObjectType:Union"),
+				TEXT("EImplicitObjectType:LevelSet"),
+				TEXT("EImplicitObjectType:Unknown"),
+				TEXT("EImplicitObjectType:Convex"),
+				TEXT("EImplicitObjectType:TaperedCylinder"),
+				TEXT("EImplicitObjectType:Cylinder"),
+				TEXT("EImplicitObjectType:TriangleMesh"),
+				TEXT("EImplicitObjectType:HeightField"),
+				TEXT("EImplicitObjectType:DEPRECATED_Scaled"),
+				TEXT("EImplicitObjectType:Triangle"),
+				TEXT("EImplicitObjectType:UnionClustered"),
+				TEXT("EImplicitObjectType:TaperedCapsule"),
+				TEXT("EImplicitObjectType:WeightedLatticeBone"),
+			};
+
+		static_assert(Chaos::ImplicitObjectType::ConcreteObjectCount == UE_ARRAY_COUNT(Names), "");
+
+		return Names[Val];
+	}
+
+}
diff --git a/Engine/Source/Runtime/Experimental/Chaos/Public/Chaos/ParticleHandle.h b/Engine/Source/Runtime/Experimental/Chaos/Public/Chaos/ParticleHandle.h
--- a/Engine/Source/Runtime/Experimental/Chaos/Public/Chaos/ParticleHandle.h
+++ b/Engine/Source/Runtime/Experimental/Chaos/Public/Chaos/ParticleHandle.h
@@ -11,6 +11,7 @@
 #include "Chaos/Collision/CollisionFilterBits.h"
 #include "ChaosCheck.h"
 #include "Chaos/ChaosDebugDrawDeclares.h"
+#include "ChaosVisualDebugger/ChaosVisualDebuggerTrace.h"
 #if CHAOS_DEBUG_DRAW
 #include "Chaos/ChaosDebugDraw.h"
 #endif
@@ -2956,6 +2957,12 @@ protected:
 					//must set geometry after because shapes are rebuilt and we want them to know about anything Func did
 					Data.SetGeometry(GeomToModify);
 				}
+				else
+				{
+					// Only invalidate the geometry if we modified the existing one.
+					// New geometry will be always evaluated the first time it is seen
+					CVD_TRACE_INVALIDATE_CACHED_GEOMETRY(GeomToModify);
+				}
 				UpdateShapesArray();
 			}
 		});
diff --git a/Engine/Source/Runtime/Experimental/Chaos/Public/ChaosVisualDebugger/ChaosVisualDebuggerTrace.h b/Engine/Source/Runtime/Experimental/Chaos/Public/ChaosVisualDebugger/ChaosVisualDebuggerTrace.h
--- a/Engine/Source/Runtime/Experimental/Chaos/Public/ChaosVisualDebugger/ChaosVisualDebuggerTrace.h
+++ b/Engine/Source/Runtime/Experimental/Chaos/Public/ChaosVisualDebugger/ChaosVisualDebuggerTrace.h
@@ -91,6 +91,10 @@
 	#define CVD_TRACE_NON_SOLVER_TRANSFORM(Transform, DebugName)
 #endif
 
+#ifndef CVD_TRACE_INVALIDATE_CACHED_GEOMETRY
+	#define CVD_TRACE_INVALIDATE_CACHED_GEOMETRY(ImplicitObjectPtr)
+#endif
+
 #else
 
 #include "ChaosVDRuntimeModule.h"
@@ -310,6 +314,11 @@ UE_TRACE_EVENT_END()
 	FChaosVisualDebuggerTrace::TraceNonSolverTransform(Transform, DebugName);
 #endif
 
+#ifndef CVD_TRACE_INVALIDATE_CACHED_GEOMETRY
+	#define CVD_TRACE_INVALIDATE_CACHED_GEOMETRY(ImplicitObjectPtr) \
+	FChaosVisualDebuggerTrace::InvalidateGeometryFromCache(ImplicitObjectPtr);
+#endif
+
 struct FChaosVDContext;
 
 namespace Chaos
@@ -411,6 +420,12 @@ public:
 	 */
 	static CHAOS_API void TraceImplicitObject(FChaosVDImplicitObjectWrapper WrappedGeometryData);
 
+	/**
+	 * Removes an implicit object from the serialized geometry IDs cache, to ensure we re-serialize it with any new changes
+	 *  @param CachedGeometryToInvalidate Ptr to the Geometry we want to invalidate from the cache
+	 */
+	static CHAOS_API void InvalidateGeometryFromCache(const Chaos::FImplicitObject* CachedGeometryToInvalidate);
+
 	/**
 	 * Traces the provided location using with the provided ID -
 	 * These are not tied to any solver step so they will be recorded as part of the current game frame data
diff --git a/Engine/Source/Runtime/Experimental/ChaosVDData/Private/ChaosVDRecording.cpp b/Engine/Source/Runtime/Experimental/ChaosVDData/Private/ChaosVDRecording.cpp
--- a/Engine/Source/Runtime/Experimental/ChaosVDData/Private/ChaosVDRecording.cpp
+++ b/Engine/Source/Runtime/Experimental/ChaosVDData/Private/ChaosVDRecording.cpp
@@ -366,9 +366,8 @@ void FChaosVDRecording::GetAvailableSolverIDsAtGameFrame_AssumesLocked(const FCh
 
 void FChaosVDRecording::CollapseSolverFramesRange_AssumesLocked(int32 SolverID, int32 StartFrame, int32 EndFrame, FChaosVDSolverFrameData& OutCollapsedFrameData)
 {
-	static TMap<int32, FChaosVDParticleDataWrapper> TempParticleDataByID;
-
-	TempParticleDataByID.Reset();
+	// Make sure we start with a clear map
+	ParticlesOnCurrentGeneratedKeyframe.Reset();
 	
 	for (int32 CurrentFrameNumber = StartFrame; CurrentFrameNumber <= EndFrame; CurrentFrameNumber++)
 	{
@@ -383,15 +382,20 @@ void FChaosVDRecording::CollapseSolverFramesRange_AssumesLocked(int32 SolverID,
 			{
 				for (const FChaosVDStepData& StepData : SolverFrameData->SolverSteps)
 				{
-					for (const FChaosVDParticleDataWrapper& ParticleData : StepData.RecordedParticlesData)
+					for (const TSharedPtr<FChaosVDParticleDataWrapper>& ParticleData : StepData.RecordedParticlesData)
 					{
-						if (FChaosVDParticleDataWrapper* FoundParticleData = TempParticleDataByID.Find(ParticleData.ParticleIndex))
+						if (!ParticleData)
+						{
+							continue;
+						}
+
+						if (TSharedPtr<FChaosVDParticleDataWrapper>* FoundParticleData = ParticlesOnCurrentGeneratedKeyframe.Find(ParticleData->ParticleIndex))
 						{
 							(*FoundParticleData) = ParticleData;
 						}
 						else
 						{
-							TempParticleDataByID.Add(ParticleData.ParticleIndex, ParticleData);
+							ParticlesOnCurrentGeneratedKeyframe.Add(ParticleData->ParticleIndex, ParticleData);
 						}
 					}
 				}
@@ -411,7 +415,9 @@ void FChaosVDRecording::CollapseSolverFramesRange_AssumesLocked(int32 SolverID,
 				FChaosVDStepData CollapsedStepData;
 				CollapsedStepData.StepName = TEXT("GeneratedStep");
 
-				TempParticleDataByID.GenerateValueArray(CollapsedStepData.RecordedParticlesData);
+				ParticlesOnCurrentGeneratedKeyframe.GenerateValueArray(CollapsedStepData.RecordedParticlesData);
+
+				ParticlesOnCurrentGeneratedKeyframe.Reset();
 				
 				OutCollapsedFrameData.SolverSteps.Add(MoveTemp(CollapsedStepData));
 			}
diff --git a/Engine/Source/Runtime/Experimental/ChaosVDData/Public/ChaosVDRecording.h b/Engine/Source/Runtime/Experimental/ChaosVDData/Public/ChaosVDRecording.h
--- a/Engine/Source/Runtime/Experimental/ChaosVDData/Public/ChaosVDRecording.h
+++ b/Engine/Source/Runtime/Experimental/ChaosVDData/Public/ChaosVDRecording.h
@@ -17,7 +17,7 @@ DECLARE_MULTICAST_DELEGATE_TwoParams(FChaosVDGeometryDataLoaded, const Chaos::FC
 struct FChaosVDStepData
 {
 	FString StepName;
-	TArray<FChaosVDParticleDataWrapper> RecordedParticlesData;
+	TArray<TSharedPtr<FChaosVDParticleDataWrapper>> RecordedParticlesData;
 	TArray<TSharedPtr<FChaosVDParticlePairMidPhase>> RecordedMidPhases;
 	TArray<FChaosVDConstraint> RecordedConstraints;
 	TMap<int32, TArray<FChaosVDConstraint>> RecordedConstraintsByParticleID;
@@ -284,6 +284,9 @@ protected:
 	/** True if this recording is being populated from a live session */
 	bool bIsLive = false;
 
+	/** Map that temporary holds generated particle data during the key frame generation process, keeping its memory allocation between generated frames*/
+	TMap<int32, TSharedPtr<FChaosVDParticleDataWrapper>> ParticlesOnCurrentGeneratedKeyframe;
+
 	friend class FChaosVDTraceProvider;
 	friend class FChaosVDTraceImplicitObjectProcessor;
 };
diff --git a/Engine/Source/Runtime/Experimental/ChaosVisualDebugger/Private/DataWrappers/ChaosVDParticleDataWrapper.cpp b/Engine/Source/Runtime/Experimental/ChaosVisualDebugger/Private/DataWrappers/ChaosVDParticleDataWrapper.cpp
--- a/Engine/Source/Runtime/Experimental/ChaosVisualDebugger/Private/DataWrappers/ChaosVDParticleDataWrapper.cpp
+++ b/Engine/Source/Runtime/Experimental/ChaosVisualDebugger/Private/DataWrappers/ChaosVDParticleDataWrapper.cpp
@@ -76,6 +76,26 @@ bool FChaosVDParticleDynamicMisc::Serialize(FArchive& Ar)
 	return true;
 }
 
+bool FChaosVDParticleCluster::Serialize(FArchive& Ar)
+{
+	Ar << ParentParticleID;	
+	Ar << NumChildren;
+	Ar << ChildToParent;
+	Ar << ClusterGroupIndex;
+	Ar << bInternalCluster;
+	Ar << CollisionImpulse;
+	Ar << ExternalStrains;
+	Ar << InternalStrains;
+	Ar << Strain;
+	Ar << ConnectivityEdges;
+	Ar << bIsAnchored;
+	Ar << bUnbreakable;
+	Ar << bIsChildToParentLocked;
+	
+	Ar << bHasValidData;
+	return true;
+}
+
 bool FChaosVDParticleDataWrapper::Serialize(FArchive& Ar)
 {
 	Ar << Type;
@@ -110,5 +130,7 @@ bool FChaosVDParticleDataWrapper::Serialize(FArchive& Ar)
 
 	Ar << CollisionDataPerShape;
 
+	Ar << ParticleCluster;
+
 	return true;
 }
diff --git a/Engine/Source/Runtime/Experimental/ChaosVisualDebugger/Public/DataWrappers/ChaosVDParticleDataWrapper.h b/Engine/Source/Runtime/Experimental/ChaosVisualDebugger/Public/DataWrappers/ChaosVDParticleDataWrapper.h
--- a/Engine/Source/Runtime/Experimental/ChaosVisualDebugger/Public/DataWrappers/ChaosVDParticleDataWrapper.h
+++ b/Engine/Source/Runtime/Experimental/ChaosVisualDebugger/Public/DataWrappers/ChaosVDParticleDataWrapper.h
@@ -394,10 +394,146 @@ struct TStructOpsTypeTraits<FChaosVDParticleDynamicMisc> : public TStructOpsType
 	};
 };
 
-/** Simplified UStruct version of FChaosVDParticleDataWrapper.
+/** Represents the data of a connectivity Edge that CVD can use to reconstruct it during playback */
+USTRUCT()
+struct FChaosVDConnectivityEdge
+{
+	GENERATED_BODY()
+
+	UPROPERTY(EditAnywhere, Category=ConnectivityEdge)
+	int32 SiblingParticleID = INDEX_NONE;
+
+	UPROPERTY(EditAnywhere, Category=ConnectivityEdge)
+	float Strain = 0.0f;
+
+	CHAOSVDRUNTIME_API bool Serialize(FArchive& Ar)
+	{
+		Ar << SiblingParticleID;
+		Ar << Strain;
+
+		return true;
+	}
+};
+
+inline FArchive& operator<<(FArchive& Ar, FChaosVDConnectivityEdge& Data)
+{
+	Data.Serialize(Ar);
+	return Ar;
+}
+
+template<>
+struct TStructOpsTypeTraits<FChaosVDConnectivityEdge> : public TStructOpsTypeTraitsBase2<FChaosVDConnectivityEdge>
+{
+	enum
+	{
+		WithSerializer = true,
+	};
+};
+
+/** Structure contained data from a Clustered Particle.
  * Used to be able to show the values in the editor and allow changes via the Property Editor.
  */
 USTRUCT()
+struct FChaosVDParticleCluster : public FChaosVDParticleDataBase
+{
+	GENERATED_BODY()
+
+	FChaosVDParticleCluster()
+	{
+	}
+
+	CHAOSVDRUNTIME_API bool Serialize(FArchive& Ar);
+
+	template <typename TOther>
+	void CopyFrom(const TOther& Other)
+	{
+		ParentParticleID = Other.ClusterIds().Id ? Other.ClusterIds().Id->UniqueIdx().Idx : INDEX_NONE;
+
+		NumChildren = Other.ClusterIds().NumChildren;
+
+		ChildToParent = Other.ChildToParent();
+		ClusterGroupIndex = Other.ClusterGroupIndex();
+		bInternalCluster = Other.InternalCluster();
+		CollisionImpulse = Other.CollisionImpulses();
+		ExternalStrains = Other.GetExternalStrain();
+		InternalStrains = Other.GetInternalStrains();
+		Strain = Other.Strain();
+
+
+		ConnectivityEdges.Reserve(Other.ConnectivityEdges().Num());
+		for (auto& Edge : Other.ConnectivityEdges())
+		{
+			int32 SiblingId = Edge.Sibling ? Edge.Sibling->UniqueIdx().Idx : INDEX_NONE;
+			ConnectivityEdges.Add( { SiblingId,  Edge.Strain });
+		}
+
+		bIsAnchored = Other.IsAnchored();
+		bUnbreakable = Other.Unbreakable();
+		bIsChildToParentLocked = Other.IsChildToParentLocked();
+		
+		bHasValidData = true;
+	}
+
+	UPROPERTY(EditAnywhere, Category= "Particle Cluster")
+	int32 ParentParticleID = INDEX_NONE;
+
+	UPROPERTY(EditAnywhere, Category= "Particle Cluster | Cluster Id")
+	int32 NumChildren = INDEX_NONE;
+	
+	UPROPERTY(EditAnywhere, Category= "Particle Cluster")
+	FTransform ChildToParent;
+
+	UPROPERTY(EditAnywhere, Category= "Particle Cluster")
+	int32 ClusterGroupIndex = INDEX_NONE;
+
+	UPROPERTY(EditAnywhere, Category= "Particle Cluster")
+	bool bInternalCluster = false;
+
+	UPROPERTY(EditAnywhere, Category= "Particle Cluster")
+	float CollisionImpulse = 0.0f;
+
+	UPROPERTY(EditAnywhere, Category= "Particle Cluster")
+	float ExternalStrains = 0.0f;
+	
+	UPROPERTY(EditAnywhere, Category= "Particle Cluster")
+	float InternalStrains = 0.0f;
+
+	UPROPERTY(EditAnywhere, Category= "Particle Cluster")
+	float Strain = 0.0f;
+
+	UPROPERTY(EditAnywhere, Category= "Particle Cluster")
+	TArray<FChaosVDConnectivityEdge> ConnectivityEdges;
+
+	UPROPERTY(EditAnywhere, Category= "Particle Cluster")
+	bool bIsAnchored = false;
+
+	UPROPERTY(EditAnywhere, Category= "Particle Cluster")
+	bool bUnbreakable = false;
+
+	UPROPERTY(EditAnywhere, Category= "Particle Cluster")
+	bool bIsChildToParentLocked = false;
+};
+
+inline FArchive& operator<<(FArchive& Ar, FChaosVDParticleCluster& Data)
+{
+	Data.Serialize(Ar);
+	return Ar;
+}
+
+template<>
+struct TStructOpsTypeTraits<FChaosVDParticleCluster> : public TStructOpsTypeTraitsBase2<FChaosVDParticleCluster>
+{
+	enum
+	{
+		WithSerializer = true,
+	};
+};
+
+
+/** Simplified UStruct version of FChaosVDParticleDataWrapper.
+ * Used to be able to show the values in the editor and allow changes via the Property Editor.
+ */
+USTRUCT(DisplayName="Particle Data")
 struct FChaosVDParticleDataWrapper : public FChaosVDParticleDataBase
 {
 	virtual ~FChaosVDParticleDataWrapper() override = default;
@@ -440,6 +576,9 @@ struct FChaosVDParticleDataWrapper : public FChaosVDParticleDataBase
 	UPROPERTY(EditAnywhere, Category= "Particle Particle Mass Props")
 	FChaosVDParticleMassProps ParticleMassProps;
 
+	UPROPERTY(EditAnywhere, Category= "Particle Cluster Data")
+	FChaosVDParticleCluster ParticleCluster;
+
 	UPROPERTY()
 	TArray<FChaosVDShapeCollisionData> CollisionDataPerShape;
 	
