[ChaosVD] Fix CIS issue C4702 unreachable code in ChaosVDMeshComponentPool and ChaosVDGeometryBuilder

diff --git a/Engine/Plugins/Experimental/ChaosVD/Source/ChaosVD/Public/ChaosVDGeometryBuilder.h b/Engine/Plugins/Experimental/ChaosVD/Source/ChaosVD/Public/ChaosVDGeometryBuilder.h
--- a/Engine/Plugins/Experimental/ChaosVD/Source/ChaosVD/Public/ChaosVDGeometryBuilder.h
+++ b/Engine/Plugins/Experimental/ChaosVD/Source/ChaosVD/Public/ChaosVDGeometryBuilder.h
@@ -501,8 +501,6 @@ ComponentType* FChaosVDGeometryBuilder::GetAvailableInstancedStaticMeshComponent
 
 		return Component;
 	}
-
-	return nullptr;
 }
 
 template <typename ComponentType>
diff --git a/Engine/Plugins/Experimental/ChaosVD/Source/ChaosVD/Public/ChaosVDMeshComponentPool.h b/Engine/Plugins/Experimental/ChaosVD/Source/ChaosVD/Public/ChaosVDMeshComponentPool.h
--- a/Engine/Plugins/Experimental/ChaosVD/Source/ChaosVD/Public/ChaosVDMeshComponentPool.h
+++ b/Engine/Plugins/Experimental/ChaosVD/Source/ChaosVD/Public/ChaosVDMeshComponentPool.h
@@ -54,8 +54,10 @@ TMeshComponent* FChaosVDMeshComponentPool::AcquireMeshComponent(UObject* Outer,
 	{
 		return GetMeshComponentFromPool_Internal<TMeshComponent>(PooledDynamicMeshComponent, Outer, Name);
 	}
-
-	return nullptr;
+	else
+	{
+		return nullptr;
+	}
 }
 
 template <typename TDesiredMeshComponent, typename TPooledComponent>
