[ChaosVD] Fix for an issue where properties with invalid data were not being hidden as indented in the Particle Data details panel.

diff --git a/Engine/Plugins/Experimental/ChaosVD/Source/ChaosVD/Private/ChaosVDParticleActorCustomization.cpp b/Engine/Plugins/Experimental/ChaosVD/Source/ChaosVD/Private/ChaosVDParticleActorCustomization.cpp
--- a/Engine/Plugins/Experimental/ChaosVD/Source/ChaosVD/Private/ChaosVDParticleActorCustomization.cpp
+++ b/Engine/Plugins/Experimental/ChaosVD/Source/ChaosVD/Private/ChaosVDParticleActorCustomization.cpp
@@ -72,6 +72,8 @@ void FChaosVDParticleActorCustomization::CustomizeDetails(IDetailLayoutBuilder&
 
 			const TSharedPtr<FStructOnScope> ParticleDataView = MakeShared<FStructOnScope>(FChaosVDParticleDataWrapper::StaticStruct(), reinterpret_cast<uint8*>(&CurrentParticleDataCopy));
 			TArray<TSharedPtr<IPropertyHandle>> Handles = CVDMainCategoryBuilder.AddAllExternalStructureProperties(ParticleDataView.ToSharedRef(), EPropertyLocation::Default, nullptr);
+
+			FChaosVDDetailsCustomizationUtils::HideInvalidParticleDataProperties(Handles);
 		}
 	}
 }
diff --git a/Engine/Plugins/Experimental/ChaosVD/Source/ChaosVD/Private/DetailsCustomizations/ChaosVDDetailsCustomizationUtils.cpp b/Engine/Plugins/Experimental/ChaosVD/Source/ChaosVD/Private/DetailsCustomizations/ChaosVDDetailsCustomizationUtils.cpp
--- a/Engine/Plugins/Experimental/ChaosVD/Source/ChaosVD/Private/DetailsCustomizations/ChaosVDDetailsCustomizationUtils.cpp
+++ b/Engine/Plugins/Experimental/ChaosVD/Source/ChaosVD/Private/DetailsCustomizations/ChaosVDDetailsCustomizationUtils.cpp
@@ -3,6 +3,8 @@
 #include "DetailsCustomizations/ChaosVDDetailsCustomizationUtils.h"
 
 #include "DetailLayoutBuilder.h"
+#include "DataWrappers/ChaosVDParticleDataWrapper.h"
+
 
 void FChaosVDDetailsCustomizationUtils::HideAllCategories(IDetailLayoutBuilder& DetailBuilder, const TSet<FName>& AllowedCategories)
 {
@@ -17,3 +19,35 @@ void FChaosVDDetailsCustomizationUtils::HideAllCategories(IDetailLayoutBuilder&
 		}
 	}
 }
+
+void FChaosVDDetailsCustomizationUtils::HideInvalidParticleDataProperties(TConstArrayView<TSharedPtr<IPropertyHandle>> InPropertyHandles)
+{
+	if (InPropertyHandles.Num() == 0)
+	{
+		return;
+	}
+
+	for (const TSharedPtr<IPropertyHandle>& Handle : InPropertyHandles)
+	{
+		if (FProperty* Property = Handle ? Handle->GetProperty() : nullptr)
+		{
+			const FStructProperty* StructProperty = CastField<FStructProperty>(Property);
+			if (StructProperty && StructProperty->Struct && StructProperty->Struct->IsChildOf(FChaosVDParticleDataBase::StaticStruct()))
+			{
+				void* Data = nullptr;
+				Handle->GetValueData(Data);
+				if (Data)
+				{
+					const FChaosVDParticleDataBase* DataViewer = static_cast<const FChaosVDParticleDataBase*>(Data);
+
+					// The Particle Data viewer struct has several fields that will have default values if there was no recorded data for them in the trace file
+					// As these do not represent any real value, we should hide them in the details panel
+					if (!DataViewer->HasValidData())
+					{
+						Handle->MarkHiddenByCustomization();
+					}
+				}
+			}
+		}
+	}
+}
diff --git a/Engine/Plugins/Experimental/ChaosVD/Source/ChaosVD/Private/DetailsCustomizations/ChaosVDParticleDataWrapperCustomization.cpp b/Engine/Plugins/Experimental/ChaosVD/Source/ChaosVD/Private/DetailsCustomizations/ChaosVDParticleDataWrapperCustomization.cpp
--- a/Engine/Plugins/Experimental/ChaosVD/Source/ChaosVD/Private/DetailsCustomizations/ChaosVDParticleDataWrapperCustomization.cpp
+++ b/Engine/Plugins/Experimental/ChaosVD/Source/ChaosVD/Private/DetailsCustomizations/ChaosVDParticleDataWrapperCustomization.cpp
@@ -5,6 +5,8 @@
 #include "DataWrappers/ChaosVDParticleDataWrapper.h"
 #include "IDetailChildrenBuilder.h"
 #include "PropertyHandle.h"
+#include "DetailsCustomizations/ChaosVDDetailsCustomizationUtils.h"
+
 
 TSharedRef<IPropertyTypeCustomization> FChaosVDParticleDataWrapperCustomization::MakeInstance()
 {
@@ -21,36 +23,13 @@ void FChaosVDParticleDataWrapperCustomization::CustomizeChildren(TSharedRef<IPro
 		return;
 	}
 
-	TSet<FName> ParticleDataViewersNames;
-	ParticleDataViewersNames.Reserve(5);
+	TArray<TSharedPtr<IPropertyHandle>> Handles;
+	Handles.Reserve(NumChildren);
 
-	ParticleDataViewersNames.Add(GET_MEMBER_NAME_CHECKED(FChaosVDParticleDataWrapper, ParticlePositionRotation));
-	ParticleDataViewersNames.Add(GET_MEMBER_NAME_CHECKED(FChaosVDParticleDataWrapper, ParticleVelocities));
-	ParticleDataViewersNames.Add(GET_MEMBER_NAME_CHECKED(FChaosVDParticleDataWrapper, ParticleDynamics));
-	ParticleDataViewersNames.Add(GET_MEMBER_NAME_CHECKED(FChaosVDParticleDataWrapper, ParticleDynamicsMisc));
-	ParticleDataViewersNames.Add(GET_MEMBER_NAME_CHECKED(FChaosVDParticleDataWrapper, ParticleMassProps));
-	ParticleDataViewersNames.Add(GET_MEMBER_NAME_CHECKED(FChaosVDParticleDataWrapper, ParticleCluster));
-	
 	for (uint32 ChildIndex = 0; ChildIndex < NumChildren; ++ChildIndex)
 	{
-		TSharedRef<IPropertyHandle> ChildHandle = StructPropertyHandle->GetChildHandle(ChildIndex).ToSharedRef();
-
-		const FName PropertyName = ChildHandle->GetProperty()->GetFName();
-		if (ParticleDataViewersNames.Contains(PropertyName))
-		{
-			void* Data = nullptr;
-			ChildHandle->GetValueData(Data);
-			if (Data)
-			{
-				FChaosVDParticleDataBase* DataViewer = static_cast<FChaosVDParticleDataBase*>(Data);
-
-				// The Particle Data viewer struct has several fields that will have default values if there was no recorded data for them in the trace file
-				// As these do not represent anything real value, we only add to the details panel the ones with recorded data
-				if (DataViewer->HasValidData())
-				{
-					StructBuilder.AddProperty(ChildHandle);
-				}
-			}
-		}
+		Handles.Add(StructPropertyHandle->GetChildHandle(ChildIndex).ToSharedRef());
 	}
+
+	FChaosVDDetailsCustomizationUtils::HideInvalidParticleDataProperties(Handles);
 }
diff --git a/Engine/Plugins/Experimental/ChaosVD/Source/ChaosVD/Public/DetailsCustomizations/ChaosVDDetailsCustomizationUtils.h b/Engine/Plugins/Experimental/ChaosVD/Source/ChaosVD/Public/DetailsCustomizations/ChaosVDDetailsCustomizationUtils.h
--- a/Engine/Plugins/Experimental/ChaosVD/Source/ChaosVD/Public/DetailsCustomizations/ChaosVDDetailsCustomizationUtils.h
+++ b/Engine/Plugins/Experimental/ChaosVD/Source/ChaosVD/Public/DetailsCustomizations/ChaosVDDetailsCustomizationUtils.h
@@ -3,8 +3,10 @@
 #pragma once
 
 #include "Containers/Set.h"
+#include "Templates/SharedPointer.h"
 #include "UObject/NameTypes.h"
 
+class IPropertyHandle;
 class IDetailLayoutBuilder;
 class FName;
 
@@ -18,4 +20,10 @@ public:
 	 * @param AllowedCategories Set of category names we do not want to hide 
 	 */
 	static void HideAllCategories(IDetailLayoutBuilder& DetailBuilder, const TSet<FName>& AllowedCategories = TSet<FName>());
+
+	/**
+	 * Marks any property of the provided handles array as hidden if they are not valid CVD properties (meaning they don't have serialized data loaded from a CVD recording)
+	 * @param InPropertyHandles Handles of properties to evaluate and hide if needed
+	 */
+	static void HideInvalidParticleDataProperties(TConstArrayView<TSharedPtr<IPropertyHandle>> InPropertyHandles);
 };
