[ChaosVD] Allow CVD to start without a valid status start subsystem.

We found cases where CVD can open before the status bar subsystem is available, triggering an assert.

There is a task to be done to allow CVD gracefully handle initialization failures like this one, and close itself up reporting a clear error. But meanwhile this change will allow CVD to start, just without the status bar.

This will not affect core functionality of the tool.

diff --git a/Engine/Plugins/Experimental/ChaosVD/Source/ChaosVD/Private/Widgets/SChaosVDMainTab.cpp b/Engine/Plugins/Experimental/ChaosVD/Source/ChaosVD/Private/Widgets/SChaosVDMainTab.cpp
--- a/Engine/Plugins/Experimental/ChaosVD/Source/ChaosVD/Private/Widgets/SChaosVDMainTab.cpp
+++ b/Engine/Plugins/Experimental/ChaosVD/Source/ChaosVD/Private/Widgets/SChaosVDMainTab.cpp
@@ -62,14 +62,31 @@ void SChaosVDMainTab::Construct(const FArguments& InArgs, TSharedPtr<FChaosVDEng
 
 	StatusBarID = FName(FChaosVDTabID::StatusBar.ToString() + InChaosVDEngine->GetInstanceGuid().ToString());
 	
-	UStatusBarSubsystem* StatusBarSubsystem = GEditor ? GEditor->GetEditorSubsystem<UStatusBarSubsystem>() : nullptr;
-	check(StatusBarSubsystem);
-	TSharedRef<SWidget> StatusBarWidget = StatusBarSubsystem->MakeStatusBarWidget(StatusBarID, TabManager->GetOwnerTab().ToSharedRef());
+	TSharedPtr<SWidget> StatusBarWidget;
+	
+	if (UStatusBarSubsystem* StatusBarSubsystem = GEditor ? GEditor->GetEditorSubsystem<UStatusBarSubsystem>() : nullptr)
+	{
+		 StatusBarWidget = StatusBarSubsystem->MakeStatusBarWidget(StatusBarID, TabManager->GetOwnerTab().ToSharedRef());
 
-	// Status bars come with the output log and content browser drawers by default, therefore we need to remove them otherwise they will be on the tool's window
-	StatusBarSubsystem->UnregisterDrawer(StatusBarID, "ContentBrowser");
-	StatusBarSubsystem->UnregisterDrawer(StatusBarID, "OutputLog");
+		// Status bars come with the output log and content browser drawers by default, therefore we need to remove them otherwise they will be on the tool's window
+		StatusBarSubsystem->UnregisterDrawer(StatusBarID, "ContentBrowser");
+		StatusBarSubsystem->UnregisterDrawer(StatusBarID, "OutputLog");
+	}
+	else
+	{
+		// TODO: Add a way to try to create the status bar later in case the status bar subsystem was not ready yet.
 	
+		StatusBarWidget = SNew(SVerticalBox)
+		+SVerticalBox::Slot()
+		.Padding(2.0f)
+		[
+			SNew(STextBlock)
+			.Text(LOCTEXT("MainTabStatusBarError", " There was an issue trying to get the status bar ready. The status bar will not be available"))
+		];
+
+		UE_LOG(LogChaosVDEditor, Error, TEXT("[%s] Failed to obtain the status bar subsystem - The status bar will not be available"), ANSI_TO_TCHAR(__FUNCTION__));		
+	}
+
 	GenerateMainWindowMenu();
 
 	ChildSlot
@@ -239,7 +256,7 @@ void SChaosVDMainTab::Construct(const FArguments& InArgs, TSharedPtr<FChaosVDEng
 		.Padding(0.0f, 2.0f, 0.0f, 0.0f)
 		.AutoHeight()
 		[
-			StatusBarWidget
+			StatusBarWidget.ToSharedRef()
 		]
 	];
 
diff --git a/Engine/Plugins/Experimental/ChaosVD/Source/ChaosVD/Private/Widgets/SChaosVDRecordingControls.cpp b/Engine/Plugins/Experimental/ChaosVD/Source/ChaosVD/Private/Widgets/SChaosVDRecordingControls.cpp
--- a/Engine/Plugins/Experimental/ChaosVD/Source/ChaosVD/Private/Widgets/SChaosVDRecordingControls.cpp
+++ b/Engine/Plugins/Experimental/ChaosVD/Source/ChaosVD/Private/Widgets/SChaosVDRecordingControls.cpp
@@ -99,21 +99,26 @@ void SChaosVDRecordingControls::HandleRecordingStop()
 		return;
 	}
 
-	UStatusBarSubsystem* StatusBarSubsystem = GEditor ? GEditor->GetEditorSubsystem<UStatusBarSubsystem>() : nullptr;
-	check(StatusBarSubsystem);
-	StatusBarSubsystem->PopStatusBarMessage(StatusBarID, RecordingMessageHandle);
+	const bool bIsLiveSession = MainTabSharedPtr->GetChaosVDEngineInstance()->GetCurrentSessionDescriptor().bIsLiveSession;
 
-	if (MainTabSharedPtr->GetChaosVDEngineInstance()->GetCurrentSessionDescriptor().bIsLiveSession)
+	if (UStatusBarSubsystem* StatusBarSubsystem = GEditor ? GEditor->GetEditorSubsystem<UStatusBarSubsystem>() : nullptr)
 	{
-		FText LiveSessionEnded = LOCTEXT("LiveSessionEndedMessage"," Live session has ended");
-		LiveSessionEndedMessageHandle = StatusBarSubsystem->PushStatusBarMessage(StatusBarID, LiveSessionEnded);
+		StatusBarSubsystem->PopStatusBarMessage(StatusBarID, RecordingMessageHandle);
 
+		if (bIsLiveSession)
+		{
+			const FText LiveSessionEnded = LOCTEXT("LiveSessionEndedMessage"," Live session has ended");
+			LiveSessionEndedMessageHandle = StatusBarSubsystem->PushStatusBarMessage(StatusBarID, LiveSessionEnded);
+		}
+		else
+		{
+			const FText RecordingPathMessage = FText::Format(LOCTEXT("RecordingSavedPathMessage"," Recoring saved at {0} "), FText::AsCultureInvariant(FChaosVDRuntimeModule::Get().GetActiveRecordingFileName()));
+			RecordingPathMessageHandle = StatusBarSubsystem->PushStatusBarMessage(StatusBarID, RecordingPathMessage);
+		}
 	}
-	else
+	
+	if (!bIsLiveSession)
 	{
-		FText RecordingPathMessage = FText::Format(LOCTEXT("RecordingSavedPathMessage"," Recoring saved at {0} "), FText::AsCultureInvariant(FChaosVDRuntimeModule::Get().GetActiveRecordingFileName()));
-		RecordingPathMessageHandle = StatusBarSubsystem->PushStatusBarMessage(StatusBarID, RecordingPathMessage);
-		
 		if (FMessageDialog::Open(EAppMsgType::YesNo, LOCTEXT("OpenLastRecordingMessage", "Do you want to load the recorded file now? ")) == EAppReturnType::Yes)
 		{
 			MainTabSharedPtr->GetChaosVDEngineInstance()->LoadRecording(FChaosVDRuntimeModule::Get().GetActiveRecordingFileName());
@@ -124,7 +129,7 @@ void SChaosVDRecordingControls::HandleRecordingStop()
 void SChaosVDRecordingControls::HandleRecordingStart()
 {
 	UStatusBarSubsystem* StatusBarSubsystem = GEditor ? GEditor->GetEditorSubsystem<UStatusBarSubsystem>() : nullptr;
-	if (!ensure(StatusBarSubsystem))
+	if (!StatusBarSubsystem)
 	{
 		return;
 	}
