[ChaosVD] Adding Scene Queries recording support

- Created CVD Macros to trace the three scene query types we have in all their modes
- Create CVD data wrappers for scene queries related data
- Added support to debug draw implicit objects in CVD directly (without creating a mesh)
- Created a new Scene Query Data inspector. This is a custom details panel which will show the data for any selected query in the viewport

diff --git a/Engine/Plugins/Experimental/ChaosVD/Source/ChaosVD/ChaosVD.Build.cs b/Engine/Plugins/Experimental/ChaosVD/Source/ChaosVD/ChaosVD.Build.cs
--- a/Engine/Plugins/Experimental/ChaosVD/Source/ChaosVD/ChaosVD.Build.cs
+++ b/Engine/Plugins/Experimental/ChaosVD/Source/ChaosVD/ChaosVD.Build.cs
@@ -45,7 +45,8 @@ public class ChaosVD : ModuleRules
 				"TypedElementFramework",
 				"StatusBar",
 				"SubobjectEditor",
-				"SubobjectDataInterface"
+				"SubobjectDataInterface",
+				"PropertyEditor"
 			}
 			);
 
diff --git a/Engine/Plugins/Experimental/ChaosVD/Source/ChaosVD/Private/Actors/ChaosVDSceneQueryDataContainer.cpp b/Engine/Plugins/Experimental/ChaosVD/Source/ChaosVD/Private/Actors/ChaosVDSceneQueryDataContainer.cpp
new file mode 100644
--- /dev/null
+++ b/Engine/Plugins/Experimental/ChaosVD/Source/ChaosVD/Private/Actors/ChaosVDSceneQueryDataContainer.cpp
@@ -0,0 +1,12 @@
+﻿// Copyright Epic Games, Inc. All Rights Reserved.
+
+#include "Actors/ChaosVDSceneQueryDataContainer.h"
+
+#include "Components/ChaosVDSceneQueryDataComponent.h"
+
+AChaosVDSceneQueryDataContainer::AChaosVDSceneQueryDataContainer()
+{
+	PrimaryActorTick.bCanEverTick = false;
+
+	SceneQueryDataComponent = CreateDefaultSubobject<UChaosVDSceneQueryDataComponent>(TEXT("ChaosVDSceneQueryDataComponent"));
+}
diff --git a/Engine/Plugins/Experimental/ChaosVD/Source/ChaosVD/Private/ChaosVDGeometryBuilder.cpp b/Engine/Plugins/Experimental/ChaosVD/Source/ChaosVD/Private/ChaosVDGeometryBuilder.cpp
--- a/Engine/Plugins/Experimental/ChaosVD/Source/ChaosVD/Private/ChaosVDGeometryBuilder.cpp
+++ b/Engine/Plugins/Experimental/ChaosVD/Source/ChaosVD/Private/ChaosVDGeometryBuilder.cpp
@@ -311,7 +311,7 @@ bool FChaosVDGeometryBuilder::ApplyMeshToComponentFromKey(TWeakObjectPtr<UMeshCo
 	return bApplyMeshRequestProcessed;
 }
 
-TSharedPtr<UE::Geometry::FMeshShapeGenerator> FChaosVDGeometryBuilder::CreateMeshGeneratorForImplicitObject(const Chaos::FImplicitObject* InImplicit)
+TSharedPtr<UE::Geometry::FMeshShapeGenerator> FChaosVDGeometryBuilder::CreateMeshGeneratorForImplicitObject(const Chaos::FImplicitObject* InImplicit, float SimpleShapesComplexityFactor)
 {
 	using namespace Chaos;
 
@@ -323,8 +323,8 @@ TSharedPtr<UE::Geometry::FMeshShapeGenerator> FChaosVDGeometryBuilder::CreateMes
 			{
 				TSharedPtr<UE::Geometry::FSphereGenerator> SphereGen = MakeShared<UE::Geometry::FSphereGenerator>();
 				SphereGen->Radius = Sphere->GetRadius();
-				SphereGen->NumTheta = 50;
-				SphereGen->NumPhi = 50;
+				SphereGen->NumTheta = 25 * SimpleShapesComplexityFactor;
+				SphereGen->NumPhi = 25 * SimpleShapesComplexityFactor;
 				SphereGen->bPolygroupPerQuad = false;
 
 				return SphereGen;	
@@ -351,8 +351,8 @@ TSharedPtr<UE::Geometry::FMeshShapeGenerator> FChaosVDGeometryBuilder::CreateMes
 				TSharedPtr<UE::Geometry::FCapsuleGenerator> CapsuleGenerator = MakeShared<UE::Geometry::FCapsuleGenerator>();
 				CapsuleGenerator->Radius = FMath::Max(FMathf::ZeroTolerance, Capsule->GetRadius());
 				CapsuleGenerator->SegmentLength = FMath::Max(FMathf::ZeroTolerance, Capsule->GetSegment().GetLength());
-				CapsuleGenerator->NumHemisphereArcSteps = 12;
-				CapsuleGenerator->NumCircleSteps = 12;
+				CapsuleGenerator->NumHemisphereArcSteps = 12 * SimpleShapesComplexityFactor;
+				CapsuleGenerator->NumCircleSteps = 12 * SimpleShapesComplexityFactor;
 
 				return CapsuleGenerator;
 			}
diff --git a/Engine/Plugins/Experimental/ChaosVD/Source/ChaosVD/Private/ChaosVDModule.cpp b/Engine/Plugins/Experimental/ChaosVD/Source/ChaosVD/Private/ChaosVDModule.cpp
--- a/Engine/Plugins/Experimental/ChaosVD/Source/ChaosVD/Private/ChaosVDModule.cpp
+++ b/Engine/Plugins/Experimental/ChaosVD/Source/ChaosVD/Private/ChaosVDModule.cpp
@@ -12,6 +12,7 @@
 #include "WorkspaceMenuStructureModule.h"
 #include "DetailsCustomizations/ChaosVDGeometryComponentCustomization.h"
 #include "DetailsCustomizations/ChaosVDParticleDataWrapperCustomization.h"
+#include "DetailsCustomizations/ChaosVDQueryDataWrappersCustomizationDetails.h"
 #include "Trace/ChaosVDTraceManager.h"
 #include "Misc/Guid.h"
 
@@ -73,7 +74,13 @@ void FChaosVDModule::RegisterClassesCustomDetails() const
 	PropertyModule.RegisterCustomClassLayout("ChaosVDParticleActor", FOnGetDetailCustomizationInstance::CreateStatic(&FChaosVDParticleActorCustomization::MakeInstance));
 	PropertyModule.RegisterCustomClassLayout("ChaosVDInstancedStaticMeshComponent", FOnGetDetailCustomizationInstance::CreateStatic(&FChaosVDGeometryComponentCustomization::MakeInstance));
 	PropertyModule.RegisterCustomClassLayout("ChaosVDStaticMeshComponent", FOnGetDetailCustomizationInstance::CreateStatic(&FChaosVDGeometryComponentCustomization::MakeInstance));
-	PropertyModule.RegisterCustomPropertyTypeLayout("ChaosVDParticleDataWrapper", FOnGetPropertyTypeCustomizationInstance::CreateStatic(&FChaosVDParticleDataWrapperCustomization::MakeInstance));
+	PropertyModule.RegisterCustomClassLayout("ChaosVDQueryVisitStep", FOnGetDetailCustomizationInstance::CreateStatic(&FChaosVDQueryVisitDataCustomization::MakeInstance));
+	PropertyModule.RegisterCustomClassLayout("ChaosVDQueryDataWrapper", FOnGetDetailCustomizationInstance::CreateStatic(&FChaosVDQueryDataWrapperCustomization::MakeInstance));
+
+	//TODO: Rename FChaosVDParticleDataWrapperCustomization to something generic as currently works with any type that wants to hide properties of type FChaosVDWrapperDataBase with invalid data.
+	// Or another option is create a new custom layout intended to be generic from the get go
+	PropertyModule.RegisterCustomPropertyTypeLayout("ChaosVDQueryDataWrapper", FOnGetPropertyTypeCustomizationInstance::CreateStatic(&FChaosVDParticleDataWrapperCustomization::MakeInstance));
+	PropertyModule.RegisterCustomPropertyTypeLayout("ChaosVDQueryVisitStep", FOnGetPropertyTypeCustomizationInstance::CreateStatic(&FChaosVDParticleDataWrapperCustomization::MakeInstance));
 }
 
 void FChaosVDModule::SpawnCVDTab()
diff --git a/Engine/Plugins/Experimental/ChaosVD/Source/ChaosVD/Private/ChaosVDParticleActorCustomization.cpp b/Engine/Plugins/Experimental/ChaosVD/Source/ChaosVD/Private/ChaosVDParticleActorCustomization.cpp
--- a/Engine/Plugins/Experimental/ChaosVD/Source/ChaosVD/Private/ChaosVDParticleActorCustomization.cpp
+++ b/Engine/Plugins/Experimental/ChaosVD/Source/ChaosVD/Private/ChaosVDParticleActorCustomization.cpp
@@ -73,7 +73,7 @@ void FChaosVDParticleActorCustomization::CustomizeDetails(IDetailLayoutBuilder&
 			const TSharedPtr<FStructOnScope> ParticleDataView = MakeShared<FStructOnScope>(FChaosVDParticleDataWrapper::StaticStruct(), reinterpret_cast<uint8*>(&CurrentParticleDataCopy));
 			TArray<TSharedPtr<IPropertyHandle>> Handles = CVDMainCategoryBuilder.AddAllExternalStructureProperties(ParticleDataView.ToSharedRef(), EPropertyLocation::Default, nullptr);
 
-			FChaosVDDetailsCustomizationUtils::HideInvalidParticleDataProperties(Handles);
+			FChaosVDDetailsCustomizationUtils::HideInvalidCVDDataWrapperProperties(Handles);
 		}
 	}
 }
diff --git a/Engine/Plugins/Experimental/ChaosVD/Source/ChaosVD/Private/ChaosVDPlaybackController.cpp b/Engine/Plugins/Experimental/ChaosVD/Source/ChaosVD/Private/ChaosVDPlaybackController.cpp
--- a/Engine/Plugins/Experimental/ChaosVD/Source/ChaosVD/Private/ChaosVDPlaybackController.cpp
+++ b/Engine/Plugins/Experimental/ChaosVD/Source/ChaosVD/Private/ChaosVDPlaybackController.cpp
@@ -239,7 +239,7 @@ void FChaosVDPlaybackController::GoToRecordedGameFrame_AssumesLocked(const int32
 						TArray<int32> AvailableSolversID;
 						LoadedRecording->GetAvailableSolverIDsAtGameFrameNumber_AssumesLocked(FrameNumber, AvailableSolversID);
 
-						SceneToControlSharedPtr->HandleEnterNewGameFrame(FrameNumber, AvailableSolversID);
+						SceneToControlSharedPtr->HandleEnterNewGameFrame(FrameNumber, AvailableSolversID, *FoundGameFrameData);
 
 						for (const int32 SolverID : AvailableSolversID)
 						{
diff --git a/Engine/Plugins/Experimental/ChaosVD/Source/ChaosVD/Private/ChaosVDPlaybackViewportClient.cpp b/Engine/Plugins/Experimental/ChaosVD/Source/ChaosVD/Private/ChaosVDPlaybackViewportClient.cpp
--- a/Engine/Plugins/Experimental/ChaosVD/Source/ChaosVD/Private/ChaosVDPlaybackViewportClient.cpp
+++ b/Engine/Plugins/Experimental/ChaosVD/Source/ChaosVD/Private/ChaosVDPlaybackViewportClient.cpp
@@ -8,6 +8,7 @@
 #include "ChaosVDPlaybackController.h"
 #include "ChaosVDScene.h"
 #include "ChaosVDSkySphereInterface.h"
+#include "Components/ChaosVDSceneQueryDataComponent.h"
 #include "ComponentVisualizer.h"
 #include "EditorModeManager.h"
 #include "Elements/Framework/TypedElementSelectionSet.h"
@@ -79,18 +80,27 @@ void FChaosVDPlaybackViewportClient::ProcessClick(FSceneView& View, HHitProxy* H
 
 	const FViewportClick Click(&View, this, Key, Event, HitX, HitY);
 
-	if (TSharedPtr<FChaosVDScene> ScenePtr = CVDScene.Pin())
+	if (const TSharedPtr<FChaosVDScene> ScenePtr = CVDScene.Pin())
 	{
 		bool bClickHandled = false;
 
-		// TODO: Iterate trough all registered visualizers
 		HComponentVisProxy* ComponentVisProxy = HitProxyCast<HComponentVisProxy>(HitProxy);
-		if (const TSharedPtr<FComponentVisualizer> Visualizer = MainTabToolkitHost->FindComponentVisualizer(UChaosVDSolverCollisionDataComponent::StaticClass()))
+		const TConstArrayView<TSharedPtr<FComponentVisualizer>> AllVisualizers = MainTabToolkitHost->GetAllComponentVisualizers();
+		for (const TSharedPtr<FComponentVisualizer>& Visualizer : AllVisualizers)
 		{
 			// Not sure if this is compliant with the normal use of the component visualizers,
 			// but passing a null hitproxy when the hit proxy was not a component
 			// It allow us to handle things like clear selection on the Collision Data Visualizer
-			bClickHandled = Visualizer->VisProxyHandleClick(this, ComponentVisProxy, Click);
+			if (Visualizer->VisProxyHandleClick(this, ComponentVisProxy, Click))
+			{
+				bClickHandled = true;
+				break;
+			}
+		}
+
+		if (bClickHandled)
+		{
+			return;
 		}
 
 		const IChaosVDGeometryComponent* AsCVDGeometryComponent = nullptr;
@@ -330,20 +340,26 @@ void FChaosVDPlaybackViewportClient::Draw(const FSceneView* View, FPrimitiveDraw
 		return;
 	}
 
-	if (TSharedPtr<FChaosVDScene> ScenePtr = CVDScene.Pin())
+	if (const TSharedPtr<FChaosVDScene> ScenePtr = CVDScene.Pin())
 	{
-	
 		for (const TPair<int32, AChaosVDSolverInfoActor*>& SolverInfoWithID : ScenePtr->GetSolverInfoActorsMap())
 		{
-			UChaosVDSolverCollisionDataComponent* CollisionDataComponent = SolverInfoWithID.Value ? SolverInfoWithID.Value->GetCollisionDataComponent() : nullptr;
-			if (CollisionDataComponent)
+			if (const UChaosVDSolverCollisionDataComponent* CollisionDataComponent = SolverInfoWithID.Value ? SolverInfoWithID.Value->GetCollisionDataComponent() : nullptr)
 			{
-				if (TSharedPtr<FComponentVisualizer> Visualizer = MainTabToolkitHost->FindComponentVisualizer(CollisionDataComponent->StaticClass()))
+				if (const TSharedPtr<FComponentVisualizer> Visualizer = MainTabToolkitHost->FindComponentVisualizer(CollisionDataComponent->StaticClass()))
 				{
 					Visualizer->DrawVisualization(CollisionDataComponent, View, PDI);
 				}
 			}
 		}
+
+		if (const UChaosVDSceneQueryDataComponent* SceneQueryDataComponent = ScenePtr->GetSceneQueryDataContainerComponent())
+		{
+			if (const TSharedPtr<FComponentVisualizer> Visualizer = MainTabToolkitHost->FindComponentVisualizer(SceneQueryDataComponent->StaticClass()))
+			{
+				Visualizer->DrawVisualization(SceneQueryDataComponent, View, PDI);
+			}
+		}
 		
 		TArray<AActor*> SelectedActors = ScenePtr->GetElementSelectionSet()->GetSelectedObjects<AActor>();
 
diff --git a/Engine/Plugins/Experimental/ChaosVD/Source/ChaosVD/Private/ChaosVDScene.cpp b/Engine/Plugins/Experimental/ChaosVD/Source/ChaosVD/Private/ChaosVDScene.cpp
--- a/Engine/Plugins/Experimental/ChaosVD/Source/ChaosVD/Private/ChaosVDScene.cpp
+++ b/Engine/Plugins/Experimental/ChaosVD/Source/ChaosVD/Private/ChaosVDScene.cpp
@@ -2,6 +2,8 @@
 
 #include "ChaosVDScene.h"
 
+#include "Actors/ChaosVDSceneQueryDataContainer.h"
+#include "Actors/ChaosVDSolverInfoActor.h"
 #include "ChaosVDEditorSettings.h"
 #include "ChaosVDGeometryBuilder.h"
 #include "ChaosVDModule.h"
@@ -9,6 +11,8 @@
 #include "Chaos/ImplicitObject.h"
 #include "ChaosVDRecording.h"
 #include "ChaosVDSkySphereInterface.h"
+#include "Components/ChaosVDSceneQueryDataComponent.h"
+#include "Components/ChaosVDSolverCollisionDataComponent.h"
 #include "DataWrappers/ChaosVDParticleDataWrapper.h"
 #include "EditorActorFolders.h"
 #include "EditorLevelUtils.h"
@@ -23,8 +27,6 @@
 #include "Selection.h"
 #include "UObject/Package.h"
 #include "WorldPersistentFolders.h"
-#include "Actors/ChaosVDSolverInfoActor.h"
-#include "Components/ChaosVDSolverCollisionDataComponent.h"
 
 #define LOCTEXT_NAMESPACE "ChaosVisualDebugger"
 
@@ -229,7 +231,7 @@ void FChaosVDScene::HandleNewGeometryData(const Chaos::FConstImplicitObjectPtr&
 	NewGeometryAvailableDelegate.Broadcast(GeometryData, GeometryID);
 }
 
-void FChaosVDScene::HandleEnterNewGameFrame(int32 FrameNumber, const TArray<int32>& AvailableSolversIds)
+void FChaosVDScene::HandleEnterNewGameFrame(int32 FrameNumber, const TArray<int32>& AvailableSolversIds, const FChaosVDGameFrameData& InNewGameFrameData)
 {
 	// Currently the particle actors from all the solvers are in the same level, and we manage them by keeping track
 	// of to which solvers they belong using maps.
@@ -280,6 +282,14 @@ void FChaosVDScene::HandleEnterNewGameFrame(int32 FrameNumber, const TArray<int3
 	{
 		CollectGarbage(GARBAGE_COLLECTION_KEEPFLAGS);
 	}
+
+	if (SceneQueriesContainer)
+	{
+		if (UChaosVDSceneQueryDataComponent* QueryDataComponent = SceneQueriesContainer->GetSceneQueryDataComponent())
+		{
+			QueryDataComponent->UpdateQueriesFromFrameData(InNewGameFrameData);
+		}
+	}
 }
 
 void FChaosVDScene::CleanUpScene()
@@ -345,6 +355,11 @@ bool FChaosVDScene::IsSolverForServer(int32 SolverID) const
 	return false;
 }
 
+UChaosVDSceneQueryDataComponent* FChaosVDScene::GetSceneQueryDataContainerComponent() const
+{
+	return SceneQueriesContainer ? SceneQueriesContainer->GetSceneQueryDataComponent() : nullptr;
+}
+
 AChaosVDParticleActor* FChaosVDScene::SpawnParticleFromRecordedData(const TSharedPtr<FChaosVDParticleDataWrapper>& InParticleData, const FChaosVDSolverFrameData& InFrameData)
 {
 	using namespace Chaos;
@@ -415,6 +430,15 @@ void FChaosVDScene::CreateBaseLights(UWorld* TargetWorld) const
 	}
 }
 
+void FChaosVDScene::CreateSceneQueriesContainer(UWorld* TargetWorld)
+{
+	const FName FolderPath("ChaosVisualDebugger/SceneQueries");
+
+	SceneQueriesContainer = TargetWorld->SpawnActor<AChaosVDSceneQueryDataContainer>();
+	SceneQueriesContainer->SetFolderPath(FolderPath);
+	SceneQueriesContainer->SetScene(AsWeak());
+}
+
 AActor* FChaosVDScene::CreateMeshComponentsContainer(UWorld* TargetWorld)
 {
 	const FName GeometryFolderPath("ChaosVisualDebugger/GeneratedMeshComponents");
@@ -446,6 +470,7 @@ UWorld* FChaosVDScene::CreatePhysicsVDWorld()
 	);
 
 	CreateBaseLights(NewWorld);
+	CreateSceneQueriesContainer(NewWorld);
 	CreateMeshComponentsContainer(NewWorld);
 
 	ActorDestroyedHandle = NewWorld->AddOnActorDestroyedHandler(FOnActorDestroyed::FDelegate::CreateRaw(this, &FChaosVDScene::HandleActorDestroyed));
diff --git a/Engine/Plugins/Experimental/ChaosVD/Source/ChaosVD/Private/ChaosVDSceneQueryDataInspectorTab.cpp b/Engine/Plugins/Experimental/ChaosVD/Source/ChaosVD/Private/ChaosVDSceneQueryDataInspectorTab.cpp
new file mode 100644
--- /dev/null
+++ b/Engine/Plugins/Experimental/ChaosVD/Source/ChaosVD/Private/ChaosVDSceneQueryDataInspectorTab.cpp
@@ -0,0 +1,40 @@
+﻿// Copyright Epic Games, Inc. All Rights Reserved.
+
+#include "ChaosVDSceneQueryDataInspectorTab.h"
+
+#include "EditorModeManager.h"
+#include "Widgets/SChaosVDMainTab.h"
+#include "Widgets/SChaosVDSceneQueryDataInspector.h"
+
+#define LOCTEXT_NAMESPACE "ChaosVisualDebugger"
+
+FChaosVDSceneQueryDataInspectorTab::~FChaosVDSceneQueryDataInspectorTab()
+{
+}
+
+TSharedRef<SDockTab> FChaosVDSceneQueryDataInspectorTab::HandleTabSpawned(const FSpawnTabArgs& Args)
+{
+	TSharedRef<SDockTab> DetailsPanelTab =
+	SNew(SDockTab)
+	.TabRole(ETabRole::MajorTab)
+	.Label(LOCTEXT("SceneQueryInspectorTab", "Scene Query Data Inspector"))
+	.ToolTipText(LOCTEXT("SceneQueryInspectorTabTip", "See the details of the any scene query selected in the viewport"));
+
+	if (const TSharedPtr<SChaosVDMainTab> MainTabPtr = OwningTabWidget.Pin())
+	{
+		DetailsPanelTab->SetContent
+		(
+			SAssignNew(SceneQueryDataInspector, SChaosVDSceneQueryDataInspector, GetChaosVDScene(), MainTabPtr->GetEditorModeManager().AsWeak())
+		);
+	}
+	else
+	{
+		DetailsPanelTab->SetContent(GenerateErrorWidget());
+	}
+
+	OnTabSpawned().Broadcast(DetailsPanelTab);
+
+	return DetailsPanelTab;
+}
+
+#undef LOCTEXT_NAMESPACE
diff --git a/Engine/Plugins/Experimental/ChaosVD/Source/ChaosVD/Private/ChaosVDTabsIDs.cpp b/Engine/Plugins/Experimental/ChaosVD/Source/ChaosVD/Private/ChaosVDTabsIDs.cpp
--- a/Engine/Plugins/Experimental/ChaosVD/Source/ChaosVD/Private/ChaosVDTabsIDs.cpp
+++ b/Engine/Plugins/Experimental/ChaosVD/Source/ChaosVD/Private/ChaosVDTabsIDs.cpp
@@ -11,4 +11,5 @@ const FName FChaosVDTabID::OutputLog = TEXT("ChaosVDOutputLog");
 const FName FChaosVDTabID::SolversTrack = TEXT("SolversTrack");
 const FName FChaosVDTabID::CVDEditorSettings = TEXT("CVDEditorSettings");
 const FName FChaosVDTabID::StatusBar = TEXT("StatusBar");
-const FName FChaosVDTabID::CollisionDataDetails = TEXT("CollisionDataDetails");
\ No newline at end of file
+const FName FChaosVDTabID::CollisionDataDetails = TEXT("CollisionDataDetails");
+const FName FChaosVDTabID::SceneQueryDataDetails = TEXT("SceneQueryDataDetails");
\ No newline at end of file
diff --git a/Engine/Plugins/Experimental/ChaosVD/Source/ChaosVD/Private/Components/ChaosVDSceneQueryDataComponent.cpp b/Engine/Plugins/Experimental/ChaosVD/Source/ChaosVD/Private/Components/ChaosVDSceneQueryDataComponent.cpp
new file mode 100644
--- /dev/null
+++ b/Engine/Plugins/Experimental/ChaosVD/Source/ChaosVD/Private/Components/ChaosVDSceneQueryDataComponent.cpp
@@ -0,0 +1,118 @@
+﻿// Copyright Epic Games, Inc. All Rights Reserved.
+
+#include "Components/ChaosVDSceneQueryDataComponent.h"
+
+#include "ChaosVDRecording.h"
+
+void FChaosVDSceneQuerySelectionHandle::SetIsSelected(bool bNewSelected)
+{
+	if (const TSharedPtr<FChaosVDQueryDataWrapper> QueryDataPtr = QueryData.Pin())
+	{
+		QueryDataPtr->bIsSelectedInEditor = bNewSelected;
+		if (QueryDataPtr && QueryDataPtr->SQVisitData.IsValidIndex(SQVisitIndex))
+		{
+			QueryDataPtr->SQVisitData[SQVisitIndex].bIsSelectedInEditor = bNewSelected;
+		}
+		else if (QueryDataPtr)
+		{
+			// If we end up with a invalid index, make sure that no visit is selected
+			SQVisitIndex = INDEX_NONE;
+			for (FChaosVDQueryVisitStep& VisitData : QueryDataPtr->SQVisitData)
+			{
+				VisitData.bIsSelectedInEditor = false;
+			}
+		}
+	}
+}
+
+bool FChaosVDSceneQuerySelectionHandle::IsSelected() const
+{
+	if (const TSharedPtr<FChaosVDQueryDataWrapper> QueryDataPtr = QueryData.Pin())
+	{
+		return QueryDataPtr->bIsSelectedInEditor;
+	}
+
+	return false;
+}
+
+UChaosVDSceneQueryDataComponent::UChaosVDSceneQueryDataComponent() : CurrentSQSelectionHandle({nullptr, INDEX_NONE})
+{
+	PrimaryComponentTick.bCanEverTick = false;
+	
+	SetCanEverAffectNavigation(false);
+	bNavigationRelevant = false;
+}
+
+void UChaosVDSceneQueryDataComponent::UpdateQueriesFromFrameData(const FChaosVDGameFrameData& InGameFrameData)
+{
+	const int32 RecordedQueriesNum = InGameFrameData.RecordedSceneQueries.Num();
+
+	RecordedQueriesByType.Empty(RecordedQueriesNum);
+	RecordedQueriesByID.Empty(RecordedQueriesNum);
+	RecordedQueries.Empty(RecordedQueriesNum);
+
+	for (const TPair<int32, TSharedPtr<FChaosVDQueryDataWrapper>>& QueryIDPair : InGameFrameData.RecordedSceneQueries)
+	{
+		if (TSharedPtr<FChaosVDQueryDataWrapper> QueryData = QueryIDPair.Value)
+		{
+			TArray<TSharedPtr<FChaosVDQueryDataWrapper>>& QueriesForType = RecordedQueriesByType.FindOrAdd(QueryData->Type);
+			QueriesForType.Emplace(QueryData);
+
+			RecordedQueriesByID.Add(QueryIDPair.Key, QueryData);
+			RecordedQueries.Add(QueryData);
+		}
+	}
+}
+
+TConstArrayView<TSharedPtr<FChaosVDQueryDataWrapper>> UChaosVDSceneQueryDataComponent::GetQueriesByType(EChaosVDSceneQueryType Type) const
+{
+	if (const TArray<TSharedPtr<FChaosVDQueryDataWrapper>>* FoundQueries = RecordedQueriesByType.Find(Type))
+	{
+		return MakeArrayView(*FoundQueries);
+	}
+
+	return TArrayView<TSharedPtr<FChaosVDQueryDataWrapper>>();
+}
+
+TConstArrayView<TSharedPtr<FChaosVDQueryDataWrapper>> UChaosVDSceneQueryDataComponent::GetAllQueries() const
+{
+	return RecordedQueries;
+}
+
+TSharedPtr<FChaosVDQueryDataWrapper> UChaosVDSceneQueryDataComponent::GetQueryByID(int32 QueryID) const
+{
+	if (const TSharedPtr<FChaosVDQueryDataWrapper>* FoundQuery = RecordedQueriesByID.Find(QueryID))
+	{
+		return *FoundQuery;
+	}
+
+	return nullptr;
+}
+
+void UChaosVDSceneQueryDataComponent::SelectQuery(int32 QueryID)
+{
+	CurrentSQSelectionHandle.SetIsSelected(false);
+	CurrentSQSelectionHandle = FChaosVDSceneQuerySelectionHandle(GetQueryByID(QueryID), 0);
+	CurrentSQSelectionHandle.SetIsSelected(true);
+	
+	SelectionChangeDelegate.Broadcast(CurrentSQSelectionHandle);
+}
+
+void UChaosVDSceneQueryDataComponent::SelectQuery(const FChaosVDSceneQuerySelectionHandle& SelectionHandle)
+{
+	CurrentSQSelectionHandle.SetIsSelected(false);
+	CurrentSQSelectionHandle = SelectionHandle;
+	CurrentSQSelectionHandle.SetIsSelected(true);
+	
+	SelectionChangeDelegate.Broadcast(CurrentSQSelectionHandle);
+}
+
+bool UChaosVDSceneQueryDataComponent::IsQuerySelected(int32 QueryID) const
+{
+	if (const TSharedPtr<FChaosVDQueryDataWrapper> QueryDataPtr = CurrentSQSelectionHandle.GetQueryData().Pin())
+	{
+		return QueryDataPtr->ID == QueryID;
+	}
+
+	return false;
+}
diff --git a/Engine/Plugins/Experimental/ChaosVD/Source/ChaosVD/Private/DetailsCustomizations/ChaosVDDetailsCustomizationUtils.cpp b/Engine/Plugins/Experimental/ChaosVD/Source/ChaosVD/Private/DetailsCustomizations/ChaosVDDetailsCustomizationUtils.cpp
--- a/Engine/Plugins/Experimental/ChaosVD/Source/ChaosVD/Private/DetailsCustomizations/ChaosVDDetailsCustomizationUtils.cpp
+++ b/Engine/Plugins/Experimental/ChaosVD/Source/ChaosVD/Private/DetailsCustomizations/ChaosVDDetailsCustomizationUtils.cpp
@@ -20,7 +20,7 @@ void FChaosVDDetailsCustomizationUtils::HideAllCategories(IDetailLayoutBuilder&
 	}
 }
 
-void FChaosVDDetailsCustomizationUtils::HideInvalidParticleDataProperties(TConstArrayView<TSharedPtr<IPropertyHandle>> InPropertyHandles)
+void FChaosVDDetailsCustomizationUtils::HideInvalidCVDDataWrapperProperties(TConstArrayView<TSharedPtr<IPropertyHandle>> InPropertyHandles)
 {
 	if (InPropertyHandles.Num() == 0)
 	{
@@ -29,25 +29,52 @@ void FChaosVDDetailsCustomizationUtils::HideInvalidParticleDataProperties(TConst
 
 	for (const TSharedPtr<IPropertyHandle>& Handle : InPropertyHandles)
 	{
-		if (FProperty* Property = Handle ? Handle->GetProperty() : nullptr)
+		bool bIsParticleDataStruct;
+		if (Handle && !HasValidCVDWrapperData(Handle, bIsParticleDataStruct))
 		{
-			const FStructProperty* StructProperty = CastField<FStructProperty>(Property);
-			if (StructProperty && StructProperty->Struct && StructProperty->Struct->IsChildOf(FChaosVDParticleDataBase::StaticStruct()))
+			// TODO: This doesn't work in all cases. It seems this just sets the IsCustom flag on, and that is why it is hidden but depends on how it is being customized
+			// We need to find a more reliable way of hiding it
+			Handle->MarkHiddenByCustomization();
+		}
+	}
+}
+
+void FChaosVDDetailsCustomizationUtils::HideInvalidCVDDataWrapperProperties(TConstArrayView<TSharedRef<IPropertyHandle>> InPropertyHandles, IDetailLayoutBuilder& DetailBuilder)
+{
+	for (const TSharedRef<IPropertyHandle>& PropertyHandle : InPropertyHandles)
+	{
+		bool bIsParticleDataStruct = false;
+		if (!HasValidCVDWrapperData(PropertyHandle, bIsParticleDataStruct))
+		{
+			if (bIsParticleDataStruct)
 			{
-				void* Data = nullptr;
-				Handle->GetValueData(Data);
-				if (Data)
-				{
-					const FChaosVDParticleDataBase* DataViewer = static_cast<const FChaosVDParticleDataBase*>(Data);
-
-					// The Particle Data viewer struct has several fields that will have default values if there was no recorded data for them in the trace file
-					// As these do not represent any real value, we should hide them in the details panel
-					if (!DataViewer->HasValidData())
-					{
-						Handle->MarkHiddenByCustomization();
-					}
-				}
+				DetailBuilder.HideProperty(PropertyHandle);
 			}
 		}
 	}
 }
+
+bool FChaosVDDetailsCustomizationUtils::HasValidCVDWrapperData(const TSharedPtr<IPropertyHandle>& InPropertyHandle, bool& bOutIsCVDBaseDataStruct)
+{
+	if (FProperty* Property = InPropertyHandle ? InPropertyHandle->GetProperty() : nullptr)
+	{
+		const FStructProperty* StructProperty = CastField<FStructProperty>(Property);
+		if (StructProperty && StructProperty->Struct && StructProperty->Struct->IsChildOf(FChaosVDWrapperDataBase::StaticStruct()))
+		{
+			bOutIsCVDBaseDataStruct = true;
+
+			void* Data = nullptr;
+			InPropertyHandle->GetValueData(Data);
+			if (Data)
+			{
+				const FChaosVDWrapperDataBase* DataViewer = static_cast<const FChaosVDWrapperDataBase*>(Data);
+
+				// The Particle Data viewer struct has several fields that will have default values if there was no recorded data for them in the trace file
+				// As these do not represent any real value, we should hide them in the details panel
+				return DataViewer->HasValidData();
+			}
+		}
+	}
+
+	return true;
+}
diff --git a/Engine/Plugins/Experimental/ChaosVD/Source/ChaosVD/Private/DetailsCustomizations/ChaosVDParticleDataWrapperCustomization.cpp b/Engine/Plugins/Experimental/ChaosVD/Source/ChaosVD/Private/DetailsCustomizations/ChaosVDParticleDataWrapperCustomization.cpp
--- a/Engine/Plugins/Experimental/ChaosVD/Source/ChaosVD/Private/DetailsCustomizations/ChaosVDParticleDataWrapperCustomization.cpp
+++ b/Engine/Plugins/Experimental/ChaosVD/Source/ChaosVD/Private/DetailsCustomizations/ChaosVDParticleDataWrapperCustomization.cpp
@@ -3,16 +3,28 @@
 #include "DetailsCustomizations/ChaosVDParticleDataWrapperCustomization.h"
 
 #include "DataWrappers/ChaosVDParticleDataWrapper.h"
+#include "DetailsCustomizations/ChaosVDDetailsCustomizationUtils.h"
 #include "IDetailChildrenBuilder.h"
+#include "PropertyCustomizationHelpers.h"
 #include "PropertyHandle.h"
-#include "DetailsCustomizations/ChaosVDDetailsCustomizationUtils.h"
 
 
+#define LOCTEXT_NAMESPACE "ChaosVisualDebugger"
+
 TSharedRef<IPropertyTypeCustomization> FChaosVDParticleDataWrapperCustomization::MakeInstance()
 {
 	return MakeShareable(new FChaosVDParticleDataWrapperCustomization());
 }
 
+void FChaosVDParticleDataWrapperCustomization::CustomizeHeader(TSharedRef<IPropertyHandle> StructPropertyHandle, FDetailWidgetRow& HeaderRow, IPropertyTypeCustomizationUtils& StructCustomizationUtils)
+{
+	HeaderRow
+	.NameContent()
+	[
+		StructPropertyHandle->CreatePropertyNameWidget()
+	];
+}
+
 void FChaosVDParticleDataWrapperCustomization::CustomizeChildren(TSharedRef<IPropertyHandle> StructPropertyHandle, IDetailChildrenBuilder& StructBuilder, IPropertyTypeCustomizationUtils& StructCustomizationUtils)
 {
 	uint32 NumChildren;
@@ -23,13 +35,20 @@ void FChaosVDParticleDataWrapperCustomization::CustomizeChildren(TSharedRef<IPro
 		return;
 	}
 
-	TArray<TSharedPtr<IPropertyHandle>> Handles;
-	Handles.Reserve(NumChildren);
-
 	for (uint32 ChildIndex = 0; ChildIndex < NumChildren; ++ChildIndex)
 	{
-		Handles.Add(StructPropertyHandle->GetChildHandle(ChildIndex).ToSharedRef());
+		TSharedPtr<IPropertyHandle> Handle = StructPropertyHandle->GetChildHandle(ChildIndex);
+
+		bool bIsParticleDataStruct = false;
+		if (FChaosVDDetailsCustomizationUtils::HasValidCVDWrapperData(Handle, bIsParticleDataStruct))
+		{
+			StructBuilder.AddProperty(Handle.ToSharedRef());
+		}
+		else if (!bIsParticleDataStruct) // If it is not a valid particle data but because it is not a data struct, just add the property anyways
+		{
+			StructBuilder.AddProperty(Handle.ToSharedRef());
+		}
 	}
-
-	FChaosVDDetailsCustomizationUtils::HideInvalidParticleDataProperties(Handles);
 }
+
+#undef LOCTEXT_NAMESPACE
diff --git a/Engine/Plugins/Experimental/ChaosVD/Source/ChaosVD/Private/DetailsCustomizations/ChaosVDQueryDataWrappersCustomizationDetails.cpp b/Engine/Plugins/Experimental/ChaosVD/Source/ChaosVD/Private/DetailsCustomizations/ChaosVDQueryDataWrappersCustomizationDetails.cpp
new file mode 100644
--- /dev/null
+++ b/Engine/Plugins/Experimental/ChaosVD/Source/ChaosVD/Private/DetailsCustomizations/ChaosVDQueryDataWrappersCustomizationDetails.cpp
@@ -0,0 +1,78 @@
+﻿// Copyright Epic Games, Inc. All Rights Reserved.
+
+#include "DetailsCustomizations/ChaosVDQueryDataWrappersCustomizationDetails.h"
+
+#include "DataWrappers/ChaosVDQueryDataWrappers.h"
+#include "DetailsCustomizations/ChaosVDDetailsCustomizationUtils.h"
+#include "DetailLayoutBuilder.h"
+#include "PropertyCustomizationHelpers.h"
+#include "PropertyHandle.h"
+
+#define LOCTEXT_NAMESPACE "ChaosVisualDebugger"
+
+TSharedRef<IPropertyTypeCustomization> FChaosVDQueryDataWrappersCustomizationDetails::MakeInstance()
+{
+	return MakeShareable(new FChaosVDQueryDataWrappersCustomizationDetails());
+
+}
+
+void FChaosVDQueryDataWrappersCustomizationDetails::CustomizeChildren(TSharedRef<IPropertyHandle> StructPropertyHandle, IDetailChildrenBuilder& StructBuilder, IPropertyTypeCustomizationUtils& StructCustomizationUtils)
+{
+	uint32 NumChildren;
+	StructPropertyHandle->GetNumChildren(NumChildren);
+
+	if (NumChildren == 0)
+	{
+		return;
+	}
+
+	TArray<TSharedPtr<IPropertyHandle>> Handles;
+	Handles.Reserve(NumChildren);
+
+	for (uint32 ChildIndex = 0; ChildIndex < NumChildren; ++ChildIndex)
+	{
+		TSharedRef<IPropertyHandle> Handle = StructPropertyHandle->GetChildHandle(ChildIndex).ToSharedRef();
+		Handles.Add(Handle);
+	}
+
+	FChaosVDDetailsCustomizationUtils::HideInvalidCVDDataWrapperProperties(Handles);
+}
+
+TSharedRef<IDetailCustomization> FChaosVDQueryVisitDataCustomization::MakeInstance()
+{
+	return MakeShareable( new FChaosVDQueryVisitDataCustomization );
+}
+
+void FChaosVDQueryVisitDataCustomization::CustomizeDetails(IDetailLayoutBuilder& DetailBuilder)
+{
+	TArray<TSharedRef<IPropertyHandle>> PotentialPropertiesToHide;
+	
+	constexpr int32 PropertiesToEvaluateNum = 2;
+	PotentialPropertiesToHide.Reset(PropertiesToEvaluateNum);
+
+	PotentialPropertiesToHide.Add(DetailBuilder.GetProperty(GET_MEMBER_NAME_CHECKED(FChaosVDQueryVisitStep, QueryFastData)));
+	PotentialPropertiesToHide.Add(DetailBuilder.GetProperty(GET_MEMBER_NAME_CHECKED(FChaosVDQueryVisitStep, HitData)));
+
+	FChaosVDDetailsCustomizationUtils::HideInvalidCVDDataWrapperProperties(PotentialPropertiesToHide, DetailBuilder);
+}
+
+TSharedRef<IDetailCustomization> FChaosVDQueryDataWrapperCustomization::MakeInstance()
+{
+	return MakeShareable( new FChaosVDQueryVisitDataCustomization );
+}
+
+void FChaosVDQueryDataWrapperCustomization::CustomizeDetails(IDetailLayoutBuilder& DetailBuilder)
+{
+	TArray<TSharedRef<IPropertyHandle>> PotentialPropertiesToHide;
+
+	constexpr int32 PropertiesToEvaluateNum = 3;
+	PotentialPropertiesToHide.Reset(PropertiesToEvaluateNum);
+
+	PotentialPropertiesToHide.Add(DetailBuilder.GetProperty(GET_MEMBER_NAME_CHECKED(FChaosVDQueryDataWrapper, CollisionQueryParams)));
+	PotentialPropertiesToHide.Add(DetailBuilder.GetProperty(GET_MEMBER_NAME_CHECKED(FChaosVDQueryDataWrapper, CollisionResponseParams)));
+	PotentialPropertiesToHide.Add(DetailBuilder.GetProperty(GET_MEMBER_NAME_CHECKED(FChaosVDQueryDataWrapper, CollisionObjectQueryParams)));
+
+	FChaosVDDetailsCustomizationUtils::HideInvalidCVDDataWrapperProperties(PotentialPropertiesToHide, DetailBuilder);
+}
+
+#undef LOCTEXT_NAMESPACE
diff --git a/Engine/Plugins/Experimental/ChaosVD/Source/ChaosVD/Private/Trace/ChaosVDTraceProvider.cpp b/Engine/Plugins/Experimental/ChaosVD/Source/ChaosVD/Private/Trace/ChaosVDTraceProvider.cpp
--- a/Engine/Plugins/Experimental/ChaosVD/Source/ChaosVD/Private/Trace/ChaosVDTraceProvider.cpp
+++ b/Engine/Plugins/Experimental/ChaosVD/Source/ChaosVD/Private/Trace/ChaosVDTraceProvider.cpp
@@ -11,6 +11,8 @@
 #include "Serialization/MemoryReader.h"
 #include "Trace/DataProcessors/ChaosVDConstraintDataProcessor.h"
 #include "Trace/DataProcessors/ChaosVDMidPhaseDataProcessor.h"
+#include "Trace/DataProcessors/ChaosVDSceneQueryDataProcessor.h"
+#include "Trace/DataProcessors/ChaosVDSceneQueryVisitDataProcessor.h"
 #include "Trace/DataProcessors/ChaosVDTraceImplicitObjectProcessor.h"
 #include "Trace/DataProcessors/ChaosVDTraceParticleDataProcessor.h"
 
@@ -202,6 +204,14 @@ void FChaosVDTraceProvider::RegisterDefaultDataProcessorsIfNeeded()
 	ConstraintDataProcessor->SetTraceProvider(AsShared());
 	RegisterDataProcessor(ConstraintDataProcessor);
 
+	TSharedPtr<FChaosVDSceneQueryDataProcessor> SceneQueryDataProcessor = MakeShared<FChaosVDSceneQueryDataProcessor>();
+	SceneQueryDataProcessor->SetTraceProvider(AsShared());
+	RegisterDataProcessor(SceneQueryDataProcessor);
+
+	TSharedPtr<FChaosVDSceneQueryVisitDataProcessor> SceneQueryVisitDataProcessor = MakeShared<FChaosVDSceneQueryVisitDataProcessor>();
+	SceneQueryVisitDataProcessor->SetTraceProvider(AsShared());
+	RegisterDataProcessor(SceneQueryVisitDataProcessor);
+
 	bDefaultDataProcessorsRegistered = true;
 }
 
diff --git a/Engine/Plugins/Experimental/ChaosVD/Source/ChaosVD/Private/Trace/DataProcessors/ChaosVDSceneQueryDataProcessor.cpp b/Engine/Plugins/Experimental/ChaosVD/Source/ChaosVD/Private/Trace/DataProcessors/ChaosVDSceneQueryDataProcessor.cpp
new file mode 100644
--- /dev/null
+++ b/Engine/Plugins/Experimental/ChaosVD/Source/ChaosVD/Private/Trace/DataProcessors/ChaosVDSceneQueryDataProcessor.cpp
@@ -0,0 +1,44 @@
+﻿// Copyright Epic Games, Inc. All Rights Reserved.
+
+#include "Trace/DataProcessors/ChaosVDSceneQueryDataProcessor.h"
+
+#include "ChaosVDRecording.h"
+#include "DataWrappers/ChaosVDQueryDataWrappers.h"
+#include "Serialization/MemoryReader.h"
+#include "Trace/ChaosVDTraceProvider.h"
+
+
+FChaosVDSceneQueryDataProcessor::FChaosVDSceneQueryDataProcessor() : IChaosVDDataProcessor(FChaosVDQueryDataWrapper::WrapperTypeName)
+{
+}
+
+bool FChaosVDSceneQueryDataProcessor::ProcessRawData(const TArray<uint8>& InData)
+{
+	const TSharedPtr<FChaosVDTraceProvider> ProviderSharedPtr = TraceProvider.Pin();
+	if (!ensure(ProviderSharedPtr.IsValid()))
+	{
+		return false;
+	}
+
+	if (const TSharedPtr<FChaosVDGameFrameData> CurrentFrameData = ProviderSharedPtr->GetCurrentGameFrame().Pin())
+	{
+		const TSharedPtr<FChaosVDQueryDataWrapper> QueryData = MakeShared<FChaosVDQueryDataWrapper>();
+		FMemoryReader MemReader(InData);
+
+		QueryData->Serialize(MemReader);
+
+		// If ParentQueryID was set, this is a sub query, so find the parent add it to the sub-queries list so we can navigate trough the query "hierarchy" later on
+		if (QueryData->ParentQueryID != INDEX_NONE)
+		{
+			if (const TSharedPtr<FChaosVDQueryDataWrapper>* ParentQueryData = CurrentFrameData->RecordedSceneQueries.Find(QueryData->ParentQueryID))
+			{
+				(*ParentQueryData)->SubQueriesIDs.Add(QueryData->ID);
+			}
+		}
+
+		CurrentFrameData->RecordedSceneQueries.Add(QueryData->ID, QueryData);
+	}
+
+	return true;
+}
+
diff --git a/Engine/Plugins/Experimental/ChaosVD/Source/ChaosVD/Private/Trace/DataProcessors/ChaosVDSceneQueryVisitDataProcessor.cpp b/Engine/Plugins/Experimental/ChaosVD/Source/ChaosVD/Private/Trace/DataProcessors/ChaosVDSceneQueryVisitDataProcessor.cpp
new file mode 100644
--- /dev/null
+++ b/Engine/Plugins/Experimental/ChaosVD/Source/ChaosVD/Private/Trace/DataProcessors/ChaosVDSceneQueryVisitDataProcessor.cpp
@@ -0,0 +1,48 @@
+﻿// Copyright Epic Games, Inc. All Rights Reserved.
+
+#include "Trace/DataProcessors/ChaosVDSceneQueryVisitDataProcessor.h"
+
+#include "ChaosVDRecording.h"
+#include "DataWrappers/ChaosVDQueryDataWrappers.h"
+#include "Serialization/MemoryReader.h"
+#include "Trace/ChaosVDTraceProvider.h"
+
+
+FChaosVDSceneQueryVisitDataProcessor::FChaosVDSceneQueryVisitDataProcessor() : IChaosVDDataProcessor(FChaosVDQueryVisitStep::WrapperTypeName)
+{
+}
+
+bool FChaosVDSceneQueryVisitDataProcessor::ProcessRawData(const TArray<uint8>& InData)
+{
+	const TSharedPtr<FChaosVDTraceProvider> ProviderSharedPtr = TraceProvider.Pin();
+	if (!ensure(ProviderSharedPtr.IsValid()))
+	{
+		return false;
+	}
+
+	if (const TSharedPtr<FChaosVDGameFrameData> CurrentFrameData = ProviderSharedPtr->GetCurrentGameFrame().Pin())
+	{
+		FChaosVDQueryVisitStep VisitStepData;
+
+		FMemoryReader MemReader(InData);
+
+		VisitStepData.Serialize(MemReader);
+		
+		if (TSharedPtr<FChaosVDQueryDataWrapper>* QueryDataPtrPtr = CurrentFrameData->RecordedSceneQueries.Find(VisitStepData.OwningQueryID))
+		{
+			TSharedPtr<FChaosVDQueryDataWrapper> QueryDataPtr = *QueryDataPtrPtr;
+			if (QueryDataPtrPtr->IsValid())
+			{
+				if (VisitStepData.HitData.HasValidData())
+				{
+					// Quick and dirty way of show the hits in the details panel. If copying this data around becomes a bottle neck we can write a customization layout for it
+					QueryDataPtr->Hits.Add(VisitStepData);
+				}
+	
+				QueryDataPtr->SQVisitData.Emplace(MoveTemp(VisitStepData));
+			}
+		}
+	}
+
+	return true;
+}
\ No newline at end of file
diff --git a/Engine/Plugins/Experimental/ChaosVD/Source/ChaosVD/Private/Visualizers/ChaosVDDebugDrawUtils.cpp b/Engine/Plugins/Experimental/ChaosVD/Source/ChaosVD/Private/Visualizers/ChaosVDDebugDrawUtils.cpp
--- a/Engine/Plugins/Experimental/ChaosVD/Source/ChaosVD/Private/Visualizers/ChaosVDDebugDrawUtils.cpp
+++ b/Engine/Plugins/Experimental/ChaosVD/Source/ChaosVD/Private/Visualizers/ChaosVDDebugDrawUtils.cpp
@@ -4,6 +4,8 @@
 
 #include "CanvasItem.h"
 #include "ChaosVDEditorSettings.h"
+#include "ChaosVDGeometryBuilder.h"
+#include "Generators/MeshShapeGenerator.h"
 #include "Engine/Engine.h"
 #include "SceneView.h"
 
@@ -169,6 +171,44 @@ void FChaosVDDebugDrawUtils::DrawLine(FPrimitiveDrawInterface* PDI, const FVecto
 	}
 }
 
+void FChaosVDDebugDrawUtils::DrawImplicitObject(FPrimitiveDrawInterface* PDI, const TSharedPtr<FChaosVDGeometryBuilder>& GeometryGenerator, const Chaos::FConstImplicitObjectPtr& ImplicitObject, const FTransform& InWorldTransform, const FColor& InColor, FStringView DebugText, ESceneDepthPriorityGroup DepthPriority)
+{
+	if (!PDI)
+	{
+		return;
+	}
+
+	if (!ImplicitObject.IsValid())
+	{
+	  return;
+	}
+
+	constexpr float SimpleShapesComplexityFactor = 0.5f;
+	if (const TSharedPtr<UE::Geometry::FMeshShapeGenerator> MeshGenerator = GeometryGenerator->CreateMeshGeneratorForImplicitObject(ImplicitObject, SimpleShapesComplexityFactor))
+	{
+		FTransform AdjustedTransform = InWorldTransform;
+		GeometryGenerator->AdjustedTransformForImplicit(ImplicitObject, AdjustedTransform);
+		MeshGenerator->Generate();
+
+		for (const UE::Geometry::FIndex3i& Triangle : MeshGenerator->Triangles)
+		{
+			FVector VertexA = AdjustedTransform.TransformPosition(MeshGenerator->Vertices[Triangle.A]);
+			FVector VertexB = AdjustedTransform.TransformPosition(MeshGenerator->Vertices[Triangle.B]);
+			FVector VertexC = AdjustedTransform.TransformPosition(MeshGenerator->Vertices[Triangle.C]);
+	
+			DrawLine(PDI, VertexA, VertexB, InColor, nullptr, DepthPriority);
+			DrawLine(PDI, VertexB, VertexC, InColor, nullptr, DepthPriority);
+			DrawLine(PDI, VertexC, VertexA, InColor, nullptr, DepthPriority);
+		}	
+	}
+
+	if (!DebugText.IsEmpty())
+	{
+		DrawText(DebugText, InWorldTransform.GetLocation(), InColor);
+	}
+}
+
+
 void FChaosVDDebugDrawUtils::DrawCanvas(FViewport& InViewport, FSceneView& View, FCanvas& Canvas)
 {
 	if (!GEngine)
diff --git a/Engine/Plugins/Experimental/ChaosVD/Source/ChaosVD/Private/Visualizers/ChaosVDSceneQueryDataComponentVisualizer.cpp b/Engine/Plugins/Experimental/ChaosVD/Source/ChaosVD/Private/Visualizers/ChaosVDSceneQueryDataComponentVisualizer.cpp
new file mode 100644
--- /dev/null
+++ b/Engine/Plugins/Experimental/ChaosVD/Source/ChaosVD/Private/Visualizers/ChaosVDSceneQueryDataComponentVisualizer.cpp
@@ -0,0 +1,308 @@
+﻿// Copyright Epic Games, Inc. All Rights Reserved.
+
+#include "Visualizers/ChaosVDSceneQueryDataComponentVisualizer.h"
+
+#include <Actors/ChaosVDSceneQueryDataContainer.h>
+#include <Components/ChaosVDSceneQueryDataComponent.h>
+
+#include "ChaosVDEditorSettings.h"
+#include "ChaosVDGeometryBuilder.h"
+#include "ChaosVDScene.h"
+#include "ChaosVDSceneQueryDataInspectorTab.h"
+#include "ChaosVDTabsIDs.h"
+#include "EditorModeManager.h"
+#include "EditorViewportClient.h"
+#include "Visualizers/ChaosVDDebugDrawUtils.h"
+#include "Visualizers/IChaosVDParticleVisualizationDataProvider.h"
+#include "Widgets/SChaosVDMainTab.h"
+
+class UChaosVDEditorSettings;
+
+IMPLEMENT_HIT_PROXY(HChaosVDSceneQueryProxy, HComponentVisProxy)
+
+void FChaosVDSceneQueryDataComponentVisualizer::DrawVisualization(const UActorComponent* Component, const FSceneView* View, FPrimitiveDrawInterface* PDI)
+{
+	const UChaosVDSceneQueryDataComponent* SceneQueryDataComponent = Cast<UChaosVDSceneQueryDataComponent>(Component);
+	if (!SceneQueryDataComponent)
+	{
+		return;
+	}
+
+	const AChaosVDSceneQueryDataContainer* SceneQueryDataContainer = Cast<AChaosVDSceneQueryDataContainer>(SceneQueryDataComponent->GetOwner());
+	if (!SceneQueryDataContainer)
+	{
+		return;
+	}
+
+	const TSharedPtr<FChaosVDScene> CVDScene = SceneQueryDataContainer->GetScene().Pin();
+	if (!CVDScene)
+	{
+		return;
+	}
+
+	const TSharedPtr<FChaosVDGeometryBuilder> GeometryGenerator = CVDScene->GetGeometryGenerator();
+	if (!GeometryGenerator)
+	{
+		return;
+	}
+
+	const TSharedPtr<FChaosVDRecording> CVDRecording = CVDScene->LoadedRecording;
+	if (!CVDRecording)
+	{
+		return;
+	}
+
+	FChaosVDSceneQueryVisualizationDataContext VisualizationContext;
+	VisualizationContext.CVDScene = SceneQueryDataContainer->GetScene();
+	VisualizationContext.SpaceTransform = FTransform::Identity;
+	VisualizationContext.GeometryGenerator = GeometryGenerator;
+
+	if (const UChaosVDEditorSettings* EditorSettings = GetDefault<UChaosVDEditorSettings>())
+	{
+		VisualizationContext.VisualizationFlags = EditorSettings->GlobalSceneQueriesVisualizationFlags;
+	}
+
+	if (EnumHasAnyFlags(EChaosVDSceneQueryVisualizationFlags::EnableDraw, static_cast<EChaosVDSceneQueryVisualizationFlags>(VisualizationContext.VisualizationFlags)))
+	{
+		// If Draw only selected Query is enabled, but no query is selected, just draw all queries
+		if (EnumHasAnyFlags(EChaosVDSceneQueryVisualizationFlags::OnlyDrawSelectedQuery, static_cast<EChaosVDSceneQueryVisualizationFlags>(VisualizationContext.VisualizationFlags)))
+		{
+			const TSharedPtr<FChaosVDQueryDataWrapper> SelectedQuery = SceneQueryDataComponent->GetSelectedQueryHandle().GetQueryData().Pin();
+			DrawSceneQuery(Component, View, PDI, CVDScene, CVDRecording, VisualizationContext, SelectedQuery);
+		}
+		else
+		{
+			for (const TSharedPtr<FChaosVDQueryDataWrapper>& Query : SceneQueryDataComponent->GetAllQueries())
+			{
+				DrawSceneQuery(Component, View, PDI, CVDScene, CVDRecording, VisualizationContext, Query);
+			}
+		}
+	}
+}
+
+bool FChaosVDSceneQueryDataComponentVisualizer::VisProxyHandleClick(FEditorViewportClient* InViewportClient, HComponentVisProxy* VisProxy, const FViewportClick& Click)
+{
+	const HChaosVDSceneQueryProxy* SceneQueryDataProxy = HitProxyCast<HChaosVDSceneQueryProxy>(VisProxy);
+	if (SceneQueryDataProxy == nullptr)
+	{
+		return false;
+	}
+	
+	if (const UChaosVDSceneQueryDataComponent* SQDataComponent = Cast<UChaosVDSceneQueryDataComponent>(SceneQueryDataProxy->Component.Get()))
+	{
+		const_cast<UChaosVDSceneQueryDataComponent*>(SQDataComponent)->SelectQuery(SceneQueryDataProxy->DataSelectionHandle);
+
+		// Bring the SQ Inspector into focus if available
+		const TSharedPtr<SChaosVDMainTab> MainTabToolkitHost = InViewportClient->GetModeTools() ? StaticCastSharedPtr<SChaosVDMainTab>(InViewportClient->GetModeTools()->GetToolkitHost()) : nullptr;
+		if (const TSharedPtr<FTabManager> TabManager = MainTabToolkitHost ? MainTabToolkitHost->GetTabManager() : nullptr)
+		{
+			TabManager->TryInvokeTab(FChaosVDTabID::SceneQueryDataDetails);
+		}
+
+		return true;
+	}
+
+	return false;
+}
+
+void FChaosVDSceneQueryDataComponentVisualizer::DrawLineTraceQuery(const UActorComponent* Component, const FChaosVDQueryDataWrapper& SceneQueryData, FChaosVDSceneQueryVisualizationDataContext& VisualizationContext, const FSceneView* View, FPrimitiveDrawInterface* PDI)
+{
+	PDI->SetHitProxy(new HChaosVDSceneQueryProxy(Component, VisualizationContext.DataSelectionHandle));
+
+	const FString DebugText = FString::Format(TEXT("Type: {0} \n Tag {1} | Owner Tag "), {TEXT("Line Trace"), SceneQueryData.CollisionQueryParams.TraceTag.ToString(), SceneQueryData.CollisionQueryParams.OwnerTag.ToString()});
+
+	FVector EndLocationToDraw;
+	if (SceneQueryData.SQVisitData.IsValidIndex(SceneQueryData.CurrentVisitIndex))
+	{
+		const FChaosVDQueryVisitStep& SQVisitData = SceneQueryData.SQVisitData[SceneQueryData.CurrentVisitIndex];
+		EndLocationToDraw = SceneQueryData.StartLocation + SQVisitData.QueryFastData.Dir * SQVisitData.QueryFastData.CurrentLength;
+	}
+	else
+	{
+		// Fallback to draw the end position
+		EndLocationToDraw = SceneQueryData.EndLocation;
+	}
+
+	FChaosVDDebugDrawUtils::DrawArrowVector(PDI, SceneQueryData.StartLocation, EndLocationToDraw, DebugText, VisualizationContext.DebugDrawColor, SDPG_Foreground);
+
+	PDI->SetHitProxy(nullptr);
+
+	if (EnumHasAnyFlags(EChaosVDSceneQueryVisualizationFlags::DrawHits, static_cast<EChaosVDSceneQueryVisualizationFlags>(VisualizationContext.VisualizationFlags)))
+	{
+		DrawHits(Component, SceneQueryData, PDI, VisualizationContext.HitColor, VisualizationContext);
+	}
+}
+
+void FChaosVDSceneQueryDataComponentVisualizer::DrawOverlapQuery(const UActorComponent* Component, const FChaosVDQueryDataWrapper& SceneQueryData, FChaosVDSceneQueryVisualizationDataContext& VisualizationContext, const FSceneView* View, FPrimitiveDrawInterface* PDI)
+{
+	PDI->SetHitProxy(new HChaosVDSceneQueryProxy(Component, VisualizationContext.DataSelectionHandle));
+
+	const Chaos::FConstImplicitObjectPtr InputShapePtr = VisualizationContext.InputGeometry;
+	if (ensure(InputShapePtr))
+	{
+		const FString DebugText = FString::Format(TEXT("Type: {0} \n Tag {1} | Owner Tag "), {TEXT("Overlap"), SceneQueryData.CollisionQueryParams.TraceTag.ToString(), SceneQueryData.CollisionQueryParams.OwnerTag.ToString()});
+		FChaosVDDebugDrawUtils::DrawImplicitObject(PDI, VisualizationContext.GeometryGenerator.Pin(), InputShapePtr, FTransform(SceneQueryData.GeometryOrientation, SceneQueryData.StartLocation), VisualizationContext.DebugDrawColor, DebugText, ESceneDepthPriorityGroup::SDPG_Foreground);
+	}
+
+	PDI->SetHitProxy(nullptr);
+
+	if (EnumHasAnyFlags(EChaosVDSceneQueryVisualizationFlags::DrawHits, static_cast<EChaosVDSceneQueryVisualizationFlags>(VisualizationContext.VisualizationFlags)))
+	{
+		DrawHits(Component, SceneQueryData, PDI, VisualizationContext.HitColor, VisualizationContext);
+	}
+}
+
+void FChaosVDSceneQueryDataComponentVisualizer::DrawSweepQuery(const UActorComponent* Component, const FChaosVDQueryDataWrapper& SceneQueryData, FChaosVDSceneQueryVisualizationDataContext& VisualizationContext, const FSceneView* View, FPrimitiveDrawInterface* PDI)
+{
+	PDI->SetHitProxy(new HChaosVDSceneQueryProxy(Component, VisualizationContext.DataSelectionHandle));
+
+	const Chaos::FConstImplicitObjectPtr InputShapePtr = VisualizationContext.InputGeometry;
+	if (ensure(InputShapePtr))
+	{
+		FChaosVDDebugDrawUtils::DrawImplicitObject(PDI, VisualizationContext.GeometryGenerator.Pin(), InputShapePtr, FTransform(SceneQueryData.GeometryOrientation, SceneQueryData.StartLocation), VisualizationContext.DebugDrawColor, nullptr, ESceneDepthPriorityGroup::SDPG_Foreground);
+
+		FVector EndLocationToDraw;
+		if (SceneQueryData.SQVisitData.IsValidIndex(SceneQueryData.CurrentVisitIndex))
+		{
+			const FChaosVDQueryVisitStep& SQVisitData = SceneQueryData.SQVisitData[SceneQueryData.CurrentVisitIndex];
+			EndLocationToDraw = SceneQueryData.StartLocation + SQVisitData.QueryFastData.Dir * SQVisitData.QueryFastData.CurrentLength;
+		}
+		else
+		{
+			// Fallback to draw the end position
+			EndLocationToDraw = SceneQueryData.EndLocation;
+		}
+
+		FChaosVDDebugDrawUtils::DrawImplicitObject(PDI, VisualizationContext.GeometryGenerator.Pin(), InputShapePtr, FTransform(SceneQueryData.GeometryOrientation, EndLocationToDraw), VisualizationContext.DebugDrawColor, nullptr, ESceneDepthPriorityGroup::SDPG_Foreground);
+	}
+
+	PDI->SetHitProxy(nullptr);
+
+	const FString DebugText = FString::Format(TEXT("Type: {0} \n Tag {1} | Owner Tag "), {TEXT("Sweep"), SceneQueryData.CollisionQueryParams.TraceTag.ToString(), SceneQueryData.CollisionQueryParams.OwnerTag.ToString()});
+	FChaosVDDebugDrawUtils::DrawArrowVector(PDI, SceneQueryData.StartLocation, SceneQueryData.EndLocation, DebugText, VisualizationContext.DebugDrawColor, SDPG_Foreground);
+
+	if (EnumHasAnyFlags(EChaosVDSceneQueryVisualizationFlags::DrawHits, static_cast<EChaosVDSceneQueryVisualizationFlags>(VisualizationContext.VisualizationFlags)))
+	{
+		DrawHits(Component, SceneQueryData, PDI, VisualizationContext.HitColor, VisualizationContext);
+	}
+}
+
+void FChaosVDSceneQueryDataComponentVisualizer::DrawHits(const UActorComponent* Component, const FChaosVDQueryDataWrapper& SceneQueryData, FPrimitiveDrawInterface* PDI, const FColor& InColor, FChaosVDSceneQueryVisualizationDataContext& VisualizationContext)
+{
+	for (int32 SQVisitIndex = 0; SQVisitIndex < SceneQueryData.SQVisitData.Num(); ++SQVisitIndex)
+	{
+		const FChaosVDQueryVisitStep& SQVisitData = SceneQueryData.SQVisitData[SQVisitIndex];
+		if (!SQVisitData.HitData.HasValidData())
+		{
+			continue;
+		}
+
+		VisualizationContext.DataSelectionHandle = FChaosVDSceneQuerySelectionHandle(VisualizationContext.DataSelectionHandle.GetQueryData(), SQVisitIndex);
+
+		PDI->SetHitProxy(new HChaosVDSceneQueryProxy(Component, VisualizationContext.DataSelectionHandle));
+
+		FString HitDebugText = FString::Format(TEXT("Distance {0} \n FaceIndex {1} \n "), {SQVisitData.HitData.Distance, SQVisitData.HitData.FaceIdx});
+
+		const FMatrix Axes = FRotationMatrix::MakeFromX(SQVisitData.HitData.WorldNormal);
+		constexpr float Thickness = 5.0f;
+		constexpr float CircleRadius = 5.0f;
+		constexpr int32 CircleSegments = 12;
+		constexpr float NormalScale = 10.5f;
+		FChaosVDDebugDrawUtils::DrawCircle(PDI, SQVisitData.HitData.WorldPosition, CircleRadius, CircleSegments, InColor, Thickness, Axes.GetUnitAxis(EAxis::Y), Axes.GetUnitAxis(EAxis::Z), HitDebugText,  ESceneDepthPriorityGroup::SDPG_Foreground);
+		FChaosVDDebugDrawUtils::DrawArrowVector(PDI, SQVisitData.HitData.WorldPosition, SQVisitData.HitData.WorldPosition + SQVisitData.HitData.FaceNormal * NormalScale, nullptr, (FLinearColor(InColor) * 0.65f).ToFColorSRGB(), ESceneDepthPriorityGroup::SDPG_Foreground);
+		FChaosVDDebugDrawUtils::DrawArrowVector(PDI, SQVisitData.HitData.WorldPosition, SQVisitData.HitData.WorldPosition + SQVisitData.HitData.WorldNormal * NormalScale, TEXT("Hit World Normal"), InColor, ESceneDepthPriorityGroup::SDPG_Foreground);
+
+		if (SQVisitData.bIsSelectedInEditor)
+		{
+			// We don't have an easy way to show something is selected with debug draw
+			// but 3D box surrounding the hit is better than nothing
+			FTransform SelectionBoxTransform;
+			SelectionBoxTransform.SetRotation(FRotationMatrix::MakeFromZ(SQVisitData.HitData.WorldNormal).ToQuat());
+			SelectionBoxTransform.SetLocation(SQVisitData.HitData.WorldPosition);
+
+			// The Selection box should be a bit bigger than the configured circle radius for the debug draw hit
+			constexpr float HitSelectionBoxSize = CircleRadius * 1.2f;
+
+			FVector SelectionBoxExtents(HitSelectionBoxSize,HitSelectionBoxSize,HitSelectionBoxSize);
+			FChaosVDDebugDrawUtils::DrawBox(PDI, SelectionBoxExtents, FColor::Yellow, SelectionBoxTransform, nullptr, ESceneDepthPriorityGroup::SDPG_Foreground);
+		}
+
+		PDI->SetHitProxy(nullptr);
+	}
+}
+
+bool FChaosVDSceneQueryDataComponentVisualizer::HasEndLocation(const FChaosVDQueryDataWrapper& SceneQueryData) const
+{
+	return SceneQueryData.Type != EChaosVDSceneQueryType::Sweep;
+}
+
+void FChaosVDSceneQueryDataComponentVisualizer::DrawSceneQuery(const UActorComponent* Component, const FSceneView* View, FPrimitiveDrawInterface* PDI, const TSharedPtr<FChaosVDScene>& CVDScene, const TSharedPtr<FChaosVDRecording>& CVDRecording, FChaosVDSceneQueryVisualizationDataContext& VisualizationContext, const TSharedPtr<FChaosVDQueryDataWrapper>& Query)
+{
+	// Reset query Specify context values
+	VisualizationContext.InputGeometry = nullptr;
+	VisualizationContext.DataSelectionHandle = FChaosVDSceneQuerySelectionHandle(nullptr, INDEX_NONE);
+
+	if (!Query)
+	{
+		return;
+	}
+			
+	const bool bIsServerQuery = CVDScene->IsSolverForServer(Query->WorldSolverID);
+	const bool bServerDrawIsEnabled = EnumHasAnyFlags(EChaosVDSceneQueryVisualizationFlags::DrawServerQueries, static_cast<EChaosVDSceneQueryVisualizationFlags>(VisualizationContext.VisualizationFlags));
+	const bool bClientDrawIsEnabled = EnumHasAnyFlags(EChaosVDSceneQueryVisualizationFlags::DrawClientQueries, static_cast<EChaosVDSceneQueryVisualizationFlags>(VisualizationContext.VisualizationFlags));
+
+	if ((bIsServerQuery && !bServerDrawIsEnabled) || (!bIsServerQuery && !bClientDrawIsEnabled))
+	{
+		return;
+	}
+
+	const bool bHideEmptyQueries = EnumHasAnyFlags(EChaosVDSceneQueryVisualizationFlags::HideEmptyQueries, static_cast<EChaosVDSceneQueryVisualizationFlags>(VisualizationContext.VisualizationFlags));
+	if (bHideEmptyQueries && Query->SQVisitData.IsEmpty())
+	{
+		return;
+	}
+
+	const bool bHideSubQueries = EnumHasAnyFlags(EChaosVDSceneQueryVisualizationFlags::HideSubQueries, static_cast<EChaosVDSceneQueryVisualizationFlags>(VisualizationContext.VisualizationFlags));
+	if (bHideSubQueries && Query->ParentQueryID != INDEX_NONE)
+	{
+		return;
+	}
+
+	if (const Chaos::FConstImplicitObjectPtr* InputShapePtrPtr = CVDRecording->GetGeometryMap().Find(Query->InputGeometryKey))
+	{
+		VisualizationContext.InputGeometry = *InputShapePtrPtr;
+	}
+
+	VisualizationContext.DataSelectionHandle = FChaosVDSceneQuerySelectionHandle(Query, INDEX_NONE);
+	VisualizationContext.GenerateColor(Query->ID, Query->bIsSelectedInEditor);
+
+	switch (Query->Type)
+	{
+	case EChaosVDSceneQueryType::RayCast:
+		{
+			if (EnumHasAnyFlags(EChaosVDSceneQueryVisualizationFlags::DrawLineTraceQueries, static_cast<EChaosVDSceneQueryVisualizationFlags>(VisualizationContext.VisualizationFlags)))
+			{
+				DrawLineTraceQuery(Component, *Query, VisualizationContext, View, PDI);
+			}
+			break;
+		}
+	case EChaosVDSceneQueryType::Overlap:
+		{
+			if (EnumHasAnyFlags(EChaosVDSceneQueryVisualizationFlags::DrawOverlapQueries, static_cast<EChaosVDSceneQueryVisualizationFlags>(VisualizationContext.VisualizationFlags)))
+			{
+				DrawOverlapQuery(Component, *Query, VisualizationContext, View, PDI);
+			}
+			break;
+		}
+	case EChaosVDSceneQueryType::Sweep:
+		{
+			if (EnumHasAnyFlags(EChaosVDSceneQueryVisualizationFlags::DrawSweepQueries, static_cast<EChaosVDSceneQueryVisualizationFlags>(VisualizationContext.VisualizationFlags)))
+			{
+				DrawSweepQuery(Component, *Query, VisualizationContext, View, PDI);
+			}
+			break;
+		}
+	default:
+		break;
+	}
+}
diff --git a/Engine/Plugins/Experimental/ChaosVD/Source/ChaosVD/Private/Widgets/SChaosVDMainTab.cpp b/Engine/Plugins/Experimental/ChaosVD/Source/ChaosVD/Private/Widgets/SChaosVDMainTab.cpp
--- a/Engine/Plugins/Experimental/ChaosVD/Source/ChaosVD/Private/Widgets/SChaosVDMainTab.cpp
+++ b/Engine/Plugins/Experimental/ChaosVD/Source/ChaosVD/Private/Widgets/SChaosVDMainTab.cpp
@@ -13,9 +13,11 @@
 #include "ChaosVDScene.h"
 #include "ChaosVDSolversTracksTab.h"
 #include "ChaosVDCollisionDataDetailsTab.h"
+#include "ChaosVDSceneQueryDataInspectorTab.h"
 #include "ChaosVDStyle.h"
 #include "ChaosVDTabsIDs.h"
 #include "ChaosVDWorldOutlinerTab.h"
+#include "Components/ChaosVDSceneQueryDataComponent.h"
 #include "DesktopPlatformModule.h"
 #include "Editor.h"
 #include "Framework/Application/SlateApplication.h"
@@ -29,6 +31,7 @@
 #include "Styling/StyleColors.h"
 #include "Styling/ToolBarStyle.h"
 #include "Trace/ChaosVDTraceManager.h"
+#include "Visualizers/ChaosVDSceneQueryDataComponentVisualizer.h"
 #include "Visualizers/ChaosVDSolverCollisionDataComponentVisualizer.h"
 #include "Widgets/SChaosBrowseTraceFileSourceModal.h"
 #include "Widgets/SChaosVDBrowseSessionsModal.h"
@@ -49,6 +52,7 @@ void SChaosVDMainTab::Construct(const FArguments& InArgs, TSharedPtr<FChaosVDEng
 	OwnerTab = InArgs._OwnerTab;
 
 	RegisterComponentVisualizer(UChaosVDSolverCollisionDataComponent::StaticClass()->GetFName(), MakeShared<FChaosVDSolverCollisionDataComponentVisualizer>());
+	RegisterComponentVisualizer(UChaosVDSceneQueryDataComponent::StaticClass()->GetFName(), MakeShared<FChaosVDSceneQueryDataComponentVisualizer>());
 
 	TabManager = FGlobalTabmanager::Get()->NewTabManager(InArgs._OwnerTab.ToSharedRef()).ToSharedPtr();
 
@@ -59,6 +63,7 @@ void SChaosVDMainTab::Construct(const FArguments& InArgs, TSharedPtr<FChaosVDEng
 	RegisterTabSpawner<FChaosVDSolversTracksTab>(FChaosVDTabID::SolversTrack);
 	RegisterTabSpawner<FChaosVDEditorVisualizationSettingsTab>(FChaosVDTabID::CVDEditorSettings);
 	RegisterTabSpawner<FChaosVDCollisionDataDetailsTab>(FChaosVDTabID::CollisionDataDetails);
+	RegisterTabSpawner<FChaosVDSceneQueryDataInspectorTab>(FChaosVDTabID::SceneQueryDataDetails);
 
 	StatusBarID = FName(FChaosVDTabID::StatusBar.ToString() + InChaosVDEngine->GetInstanceGuid().ToString());
 	
@@ -319,6 +324,7 @@ void SChaosVDMainTab::RegisterComponentVisualizer(FName ClassName, const TShared
 	if (!ComponentVisualizersMap.Contains(ClassName))
 	{
 		ComponentVisualizersMap.Add(ClassName, Visualizer);
+		ComponentVisualizers.Add(Visualizer);
 	}
 }
 
@@ -380,6 +386,7 @@ TSharedRef<FTabManager::FLayout> SChaosVDMainTab::GenerateMainLayout()
 					->SetSizeCoefficient(0.3f)
 					->AddTab(FChaosVDTabID::DetailsPanel, ETabState::OpenedTab)
 					->AddTab(FChaosVDTabID::CollisionDataDetails, ETabState::OpenedTab)
+					->AddTab(FChaosVDTabID::SceneQueryDataDetails, ETabState::OpenedTab)
 					->AddTab(FChaosVDTabID::CVDEditorSettings, ETabState::ClosedTab)
 				)
 			)
diff --git a/Engine/Plugins/Experimental/ChaosVD/Source/ChaosVD/Private/Widgets/SChaosVDSceneQueryDataInspector.cpp b/Engine/Plugins/Experimental/ChaosVD/Source/ChaosVD/Private/Widgets/SChaosVDSceneQueryDataInspector.cpp
new file mode 100644
--- /dev/null
+++ b/Engine/Plugins/Experimental/ChaosVD/Source/ChaosVD/Private/Widgets/SChaosVDSceneQueryDataInspector.cpp
@@ -0,0 +1,626 @@
+﻿// Copyright Epic Games, Inc. All Rights Reserved.
+
+#include "Widgets/SChaosVDSceneQueryDataInspector.h"
+
+#include "ChaosVDModule.h"
+#include "ChaosVDParticleActor.h"
+#include "ChaosVDScene.h"
+#include "EditorModeManager.h"
+#include "EditorViewportClient.h"
+#include "IStructureDetailsView.h"
+#include "PropertyEditorModule.h"
+#include "Actors/ChaosVDSolverInfoActor.h"
+#include "Components/ChaosVDSceneQueryDataComponent.h"
+#include "Modules/ModuleManager.h"
+#include "Toolkits/IToolkitHost.h"
+#include "Visualizers/ChaosVDSceneQueryDataComponentVisualizer.h"
+#include "Widgets/SChaosVDNameListPicker.h"
+#include "Widgets/SChaosVDTimelineWidget.h"
+#include "Widgets/SChaosVDWarningMessageBox.h"
+#include "Widgets/Input/SButton.h"
+#include "Widgets/Layout/SBox.h"
+#include "Widgets/Layout/SScrollBox.h"
+#include "Widgets/Layout/SSplitter.h"
+#include "Widgets/Layout/SUniformGridPanel.h"
+#include "Widgets/Text/STextBlock.h"
+
+#define LOCTEXT_NAMESPACE "ChaosVisualDebugger"
+
+SChaosVDSceneQueryDataInspector::SChaosVDSceneQueryDataInspector() : CurrentSceneQueryBeingInspectedHandle({nullptr, INDEX_NONE})
+{
+	
+}
+
+SChaosVDSceneQueryDataInspector::~SChaosVDSceneQueryDataInspector()
+{
+	if (const TSharedPtr<FChaosVDScene> ScenePtr = SceneWeakPtr.Pin())
+	{
+		ScenePtr->OnSceneUpdated().RemoveAll(this);
+		
+		if (UChaosVDSceneQueryDataComponent* SQDataComponent = ScenePtr->GetSceneQueryDataContainerComponent())
+		{
+			SQDataComponent->GetOnSelectionChangeDelegate().RemoveAll(this);
+		}
+	}
+}
+
+void SChaosVDSceneQueryDataInspector::Construct(const FArguments& InArgs, const TWeakPtr<FChaosVDScene>& InScenePtr, const TWeakPtr<FEditorModeTools>& InEditorModeTools)
+{
+	SceneWeakPtr = InScenePtr;
+	EditorModeToolsWeakPtr = InEditorModeTools;
+
+	if (const TSharedPtr<FChaosVDScene> ScenePtr = SceneWeakPtr.Pin())
+	{
+		ScenePtr->OnSceneUpdated().AddRaw(this, &SChaosVDSceneQueryDataInspector::HandleSceneUpdated);
+
+		if (UChaosVDSceneQueryDataComponent* SQDataComponent = ScenePtr->GetSceneQueryDataContainerComponent())
+		{
+			SQDataComponent->GetOnSelectionChangeDelegate().AddRaw(this, &SChaosVDSceneQueryDataInspector::SetQueryDataToInspect);
+		}
+	}
+
+	SceneQueryDataDetailsView = CreateDataDetailsView();
+	SceneQueryHitDataDetailsView = CreateDataDetailsView();
+
+	
+	constexpr float NoPadding = 0.0f;
+	constexpr float OuterBoxPadding = 2.0f;
+	constexpr float OuterInnerPadding = 5.0f;
+	constexpr float TagTitleBoxHorizontalPadding = 10.0f;
+	constexpr float TagTitleBoxVerticalPadding = 5.0f;
+	constexpr float InnerDetailsPanelsHorizontalPadding = 15.0f;
+	constexpr float InnerDetailsPanelsVerticalPadding = 15.0f;
+
+	ChildSlot
+	[
+		SNew(SVerticalBox)
+		+SVerticalBox::Slot()
+		.AutoHeight()
+		.Padding(OuterInnerPadding)
+		[
+			SNew(SBox)
+			.Visibility_Raw(this, &SChaosVDSceneQueryDataInspector::GetOutOfDateWarningVisibility)
+			.Padding(OuterBoxPadding, OuterBoxPadding,OuterBoxPadding,NoPadding)
+			[
+				SNew(SChaosVDWarningMessageBox)
+				.WarningText(LOCTEXT("SceneQueryDataOutOfData", "Scene change detected!. Selected scene query data is out of date..."))
+			]
+		]
+		+SVerticalBox::Slot()
+		.AutoHeight()
+		.Padding(TagTitleBoxHorizontalPadding, TagTitleBoxVerticalPadding, TagTitleBoxHorizontalPadding, NoPadding)
+		[
+			GenerateQueryTagInfoRow()
+		]
+		+SVerticalBox::Slot()
+		.Padding(TagTitleBoxHorizontalPadding, TagTitleBoxVerticalPadding, TagTitleBoxHorizontalPadding, NoPadding)
+		.AutoHeight()
+		[
+			SNew(STextBlock)
+			.Visibility_Raw(this, &SChaosVDSceneQueryDataInspector::GetNothingSelectedMessageVisibility)
+			.Justification(ETextJustify::Center)
+			.TextStyle(FAppStyle::Get(), "DetailsView.BPMessageTextStyle")
+			.Text(LOCTEXT("SceneQueryDataNoSelectedMessage", "Select a scene query or scene query hit in the viewport to see its details..."))
+			.AutoWrapText(true)
+		]
+		+SVerticalBox::Slot()
+		.Padding(TagTitleBoxHorizontalPadding, TagTitleBoxVerticalPadding, TagTitleBoxHorizontalPadding, NoPadding)
+		.AutoHeight()
+		[
+			GenerateQueryNavigationBoxWidget(TagTitleBoxHorizontalPadding, TagTitleBoxVerticalPadding)
+		]
+		+SVerticalBox::Slot()
+		.Padding(OuterInnerPadding)
+		.FillHeight(0.75f)
+		[
+			GenerateQueryDetailsPanelSection(InnerDetailsPanelsHorizontalPadding, InnerDetailsPanelsVerticalPadding)
+		]
+		+SVerticalBox::Slot()
+		.FillHeight(0.1f)
+		[
+			GenerateVisitStepControls()
+		]
+	];
+}
+
+TSharedRef<SWidget> SChaosVDSceneQueryDataInspector::GenerateQueryNavigationBoxWidget(float TagTitleBoxHorizontalPadding, float TagTitleBoxVerticalPadding)
+{
+	return SNew(SVerticalBox)
+			+SVerticalBox::Slot()
+			[
+				SNew(SUniformGridPanel)
+				.Visibility_Raw(this, &SChaosVDSceneQueryDataInspector::GetParentQuerySelectorVisibility)
+				.SlotPadding(FAppStyle::GetMargin("StandardDialog.SlotPadding"))
+				.MinDesiredSlotWidth(FAppStyle::GetFloat("StandardDialog.MinDesiredSlotWidth"))
+				.MinDesiredSlotHeight(FAppStyle::GetFloat("StandardDialog.MinDesiredSlotHeight"))
+				+SUniformGridPanel::Slot(0, 0)
+				[
+					SNew(SButton)
+					.HAlign(HAlign_Center)
+					.VAlign(VAlign_Center)
+					.ContentPadding(FAppStyle::GetMargin("StandardDialog.ContentPadding"))
+					.Text(LOCTEXT("SelectParentQuery", "Go to parent query"))
+					.OnClicked_Raw(this, &SChaosVDSceneQueryDataInspector::SelectParentQuery)
+				]
+			]
+			+SVerticalBox::Slot()
+			.Padding(FMargin(TagTitleBoxHorizontalPadding, TagTitleBoxVerticalPadding, TagTitleBoxHorizontalPadding, 0.0f))
+			[
+				SNew(SHorizontalBox)
+				.Visibility_Raw(this, &SChaosVDSceneQueryDataInspector::GetSubQuerySelectorVisibility)
+				+SHorizontalBox::Slot()
+				.HAlign(HAlign_Left)
+				.VAlign(VAlign_Center)
+				[
+					SNew(STextBlock)
+					.Text(LOCTEXT("SelectSubQueryDropDown", "Go To Subquery"))
+					.Font(FCoreStyle::GetDefaultFontStyle("Regular", 8))
+				]
+				+SHorizontalBox::Slot()
+				.HAlign(HAlign_Right)
+				.VAlign(VAlign_Center)
+				[
+					SAssignNew(SubQueryNamePickerWidget, SChaosVDNameListPicker)
+					.OnNameSleceted_Raw(this, &SChaosVDSceneQueryDataInspector::HandleSubQueryNameSelected)
+				]	
+			];
+}
+
+TSharedRef<SWidget> SChaosVDSceneQueryDataInspector::GenerateQueryDetailsPanelSection(float InnerDetailsPanelsHorizontalPadding, float InnerDetailsPanelsVerticalPadding)
+{
+	return SNew(SScrollBox)
+			.Visibility_Raw(this, &SChaosVDSceneQueryDataInspector::GetQueryDetailsSectionVisibility)
+			+SScrollBox::Slot()
+			.Padding(InnerDetailsPanelsHorizontalPadding,0.0f,InnerDetailsPanelsHorizontalPadding,InnerDetailsPanelsVerticalPadding)
+			[
+				SceneQueryDataDetailsView->GetWidget().ToSharedRef()
+			]
+			+SScrollBox::Slot()
+			.Padding(InnerDetailsPanelsHorizontalPadding,0.0f,InnerDetailsPanelsHorizontalPadding,0.0f)
+			[
+				SNew(SVerticalBox)
+				.Visibility_Raw(this, &SChaosVDSceneQueryDataInspector::GetSQVisitDetailsSectionVisibility)
+				+SVerticalBox::Slot()
+				.FillHeight(0.9f)
+				[
+					SceneQueryHitDataDetailsView->GetWidget().ToSharedRef()
+				]
+				+SVerticalBox::Slot()
+				.AutoHeight()
+				.Padding(0.0f,InnerDetailsPanelsVerticalPadding,0.0f,InnerDetailsPanelsVerticalPadding)
+				[
+					SNew(SUniformGridPanel)
+					.SlotPadding(FAppStyle::GetMargin("StandardDialog.SlotPadding"))
+					.MinDesiredSlotWidth(FAppStyle::GetFloat("StandardDialog.MinDesiredSlotWidth"))
+					.MinDesiredSlotHeight(FAppStyle::GetFloat("StandardDialog.MinDesiredSlotHeight"))
+					+SUniformGridPanel::Slot(0, 0)
+					[
+						SNew(SButton)
+						.HAlign(HAlign_Center)
+						.VAlign(VAlign_Center)
+						.ContentPadding(FAppStyle::GetMargin("StandardDialog.ContentPadding"))
+						.IsEnabled_Raw(this, &SChaosVDSceneQueryDataInspector::GetSelectParticleHitStateEnable)
+						.Text_Raw(this, &SChaosVDSceneQueryDataInspector::GetSelectParticleText)
+						.OnClicked(this, &SChaosVDSceneQueryDataInspector::SelectParticleForCurrentQueryData)
+					]
+				]
+			];
+}
+
+TSharedRef<SWidget> SChaosVDSceneQueryDataInspector::GenerateVisitStepControls()
+{
+	return SNew(SBorder)
+			.BorderImage(FAppStyle::Get().GetBrush("DetailsView.CategoryMiddle"))
+			.BorderBackgroundColor(FAppStyle::Get().GetSlateColor("Colors.Panel"))
+			[
+				SNew(SVerticalBox)
+				.Visibility_Raw(this, &SChaosVDSceneQueryDataInspector::GetQueryStepPlaybackControlsVisibility)
+				+SVerticalBox::Slot()
+				.AutoHeight()
+				.Padding(0.0f, 2.0f, 0.0f, 2.0f)
+				[
+					SNew(STextBlock)
+					.Justification(ETextJustify::Center)
+					.Text_Raw(this, &SChaosVDSceneQueryDataInspector::GetSQVisitsStepsText)
+				]
+				+SVerticalBox::Slot()
+				[
+					SAssignNew(QueryStepsTimelineWidget, SChaosVDTimelineWidget)
+					.ButtonVisibilityFlags(static_cast<uint16>(EChaosVDTimelineElementIDFlags::AllManualStepping))
+					.IsEnabled_Raw(this, &SChaosVDSceneQueryDataInspector::GetSQVisitStepsEnabled)
+					.OnFrameChanged_Raw(this, &SChaosVDSceneQueryDataInspector::HandleQueryStepSelectionUpdated)
+					.MaxFrames(0)	
+				]
+			];
+}
+
+TSharedRef<SWidget> SChaosVDSceneQueryDataInspector::GenerateQueryTagInfoRow()
+{
+	return SNew(SBorder)
+			.Padding(0.5f)
+			.BorderImage(FAppStyle::Get().GetBrush("DetailsView.CategoryMiddle"))
+			.BorderBackgroundColor(FAppStyle::Get().GetSlateColor("Colors.Panel"))
+			.Padding(0)
+			[
+				SNew(SBox)
+				.Padding(2.0f)
+				[
+					SNew(SSplitter)
+					.Style(FAppStyle::Get(), "DetailsView.Splitter")
+					.PhysicalSplitterHandleSize(1.5f)
+					.HitDetectionSplitterHandleSize(5.0f)
+					+SSplitter::Slot()
+					.Value(0.2)
+					[
+						SNew(SBox)
+						.VAlign(VAlign_Center)
+						.Padding(1.0)
+						.VAlign(VAlign_Center)
+						[
+							SNew(STextBlock)
+							.Text(LOCTEXT("SceneQueriesNameLabel", "Query Tag" ))
+							.Font(FAppStyle::GetFontStyle("PropertyWindow.NormalFont"))
+						]	
+					]
+					+SSplitter::Slot()
+					[
+						SNew(SBox)
+						.VAlign(VAlign_Center)
+						.Padding(1.0)
+						[
+							SNew(STextBlock)
+							.Text_Raw(this, &SChaosVDSceneQueryDataInspector::GetQueryBeingInspectedTag)
+							.Font(FAppStyle::GetFontStyle("PropertyWindow.NormalFont"))
+						]
+					]
+				]
+			];
+}
+
+void SChaosVDSceneQueryDataInspector::SetQueryDataToInspect(const FChaosVDSceneQuerySelectionHandle& InQueryDataSelectionHandle)
+{
+	ClearInspector();
+
+	if (const TSharedPtr<FChaosVDQueryDataWrapper> QueryDataToInspect = InQueryDataSelectionHandle.GetQueryData().Pin())
+	{
+		CurrentSceneQueryBeingInspectedHandle = InQueryDataSelectionHandle;
+
+		const TSharedPtr<FStructOnScope> QueryDataView = MakeShared<FStructOnScope>(FChaosVDQueryDataWrapper::StaticStruct(), reinterpret_cast<uint8*>(QueryDataToInspect.Get()));
+		SceneQueryDataDetailsView->SetStructureData(QueryDataView);
+
+		if (QueryDataToInspect->SQVisitData.IsValidIndex(InQueryDataSelectionHandle.GetSQVisitIndex()))
+		{
+			const TSharedPtr<FStructOnScope> SQVisitDataDataView = MakeShared<FStructOnScope>(FChaosVDQueryVisitStep::StaticStruct(), reinterpret_cast<uint8*>(&QueryDataToInspect->SQVisitData[InQueryDataSelectionHandle.GetSQVisitIndex()]));
+			SceneQueryHitDataDetailsView->SetStructureData(SQVisitDataDataView);
+			QueryDataToInspect->CurrentVisitIndex = InQueryDataSelectionHandle.GetSQVisitIndex();
+		}
+
+		if (QueryStepsTimelineWidget)
+		{
+			const int32 SQVisitsNum = QueryDataToInspect->SQVisitData.Num();
+			QueryStepsTimelineWidget->UpdateMinMaxValue(0, SQVisitsNum > 0 ? SQVisitsNum -1 : 0);
+			QueryStepsTimelineWidget->SetCurrentTimelineFrame(QueryDataToInspect->CurrentVisitIndex);
+		}
+
+		if (QueryDataToInspect->SubQueriesIDs.Num() > 0)
+		{
+			TArray<TSharedPtr<FName>> NewSubQueryNameList;
+			NewSubQueryNameList.Reserve(QueryDataToInspect->SubQueriesIDs.Num());
+			Algo::Transform(QueryDataToInspect->SubQueriesIDs, NewSubQueryNameList, [this](int32 QueryID)
+			{
+				TSharedPtr<FName> NewName = MakeShared<FName>(FString::Format(TEXT("Query ID {0}"), {QueryID}));
+				CurrentSubQueriesByName.Add(NewName, QueryID);
+				return NewName;
+			});
+
+			SubQueryNamePickerWidget->UpdateNameList(MoveTemp(NewSubQueryNameList));
+		}
+		else
+		{
+			CurrentSubQueriesByName.Empty();
+			SubQueryNamePickerWidget->UpdateNameList({});
+		}
+	}
+	else
+	{
+		ClearInspector();	
+	}
+
+	bIsUpToDate = true;
+}
+
+void SChaosVDSceneQueryDataInspector::HandleQueryStepSelectionUpdated(int32 NewStepIndex)
+{
+	if (!bListenToSelectionEvents)
+	{
+		return;
+	}
+
+	const TSharedPtr<FChaosVDScene> ScenePtr = SceneWeakPtr.Pin();
+	if (!ScenePtr)
+	{
+		ClearInspector();
+		return;
+	}
+
+	UChaosVDSceneQueryDataComponent* SQDataComponent = ScenePtr->GetSceneQueryDataContainerComponent();
+	if (!SQDataComponent)
+	{
+		ClearInspector();
+		return;
+	}
+
+	if (const TSharedPtr<FChaosVDQueryDataWrapper> QueryDataBeingInspected = CurrentSceneQueryBeingInspectedHandle.GetQueryData().Pin())
+	{
+		const FChaosVDSceneQuerySelectionHandle NewSelection(CurrentSceneQueryBeingInspectedHandle.GetQueryData(), NewStepIndex);
+		QueryDataBeingInspected->CurrentVisitIndex = NewStepIndex;
+		
+
+		FScopedSQInspectorSilencedSelectionEvents IgnoreSelectionEventsScope(*this);
+		SQDataComponent->SelectQuery(NewSelection);
+
+		SetQueryDataToInspect(NewSelection);
+	}
+	else
+	{
+		SetQueryDataToInspect(FChaosVDSceneQuerySelectionHandle());
+	}
+
+	if (const TSharedPtr<FEditorModeTools> EditorModeToolsPtr = EditorModeToolsWeakPtr.Pin())
+	{
+		// We need to request a re-draw to make sure the debug draw view and selection outline is updated
+		if (FEditorViewportClient* ViewportClient = EditorModeToolsPtr->GetFocusedViewportClient())
+		{
+			ViewportClient->bNeedsRedraw = true;
+		}
+	}
+}
+
+FText SChaosVDSceneQueryDataInspector::GetQueryBeingInspectedTag() const
+{
+	const TSharedPtr<FChaosVDQueryDataWrapper> QueryDataBeingInspected = CurrentSceneQueryBeingInspectedHandle.GetQueryData().Pin();
+	return FText::AsCultureInvariant(QueryDataBeingInspected ? QueryDataBeingInspected->CollisionQueryParams.TraceTag.ToString() : TEXT("None"));
+}
+
+FText SChaosVDSceneQueryDataInspector::GetSelectParticleText() const
+{
+	return LOCTEXT("SelectVisitedParticleHit", "Select Visited Particle Shape");
+}
+
+FText SChaosVDSceneQueryDataInspector::GetSQVisitsStepsText() const
+{
+	const TSharedPtr<FChaosVDQueryDataWrapper> QueryDataBeingInspected = CurrentSceneQueryBeingInspectedHandle.GetQueryData().Pin();
+	const int32 AvailableSQVisitDataNum = QueryDataBeingInspected ? QueryDataBeingInspected->SQVisitData.Num() : 0;
+	return FText::Format(FTextFormat(LOCTEXT("SQVisitStepsPlaybackControlsTitle", "Recorded SQ Visits Available {0}")), AvailableSQVisitDataNum);
+}
+
+FReply SChaosVDSceneQueryDataInspector::SelectParticleForCurrentQueryData() const
+{
+	const TSharedPtr<FChaosVDScene> ScenePtr = SceneWeakPtr.Pin();
+	if (!ScenePtr)
+	{
+		return FReply::Handled();
+	}
+
+	const TSharedPtr<FChaosVDQueryDataWrapper> QueryDataBeingInspected = CurrentSceneQueryBeingInspectedHandle.GetQueryData().Pin();
+
+	if (!QueryDataBeingInspected || !QueryDataBeingInspected->SQVisitData.IsValidIndex(QueryDataBeingInspected->CurrentVisitIndex))
+	{
+		return FReply::Handled();
+	}
+
+	if (AChaosVDParticleActor* ParticleActor = ScenePtr->GetParticleActor(QueryDataBeingInspected->WorldSolverID, QueryDataBeingInspected->SQVisitData[QueryDataBeingInspected->CurrentVisitIndex].ParticleIndex))
+	{
+		ScenePtr->SetSelectedObject(ParticleActor);
+
+		// When a particle actor is selected, all their mesh instances are selected, so we need to override that selection here.
+		// TODO: The API for selection is intentionally generic as we use it for more things that particle actors, but maybe we can have a new API method to provide a "context", so in this case
+		// we can tell the system how we want the selection to be visualized
+
+		const TConstArrayView<TSharedPtr<FChaosVDMeshDataInstanceHandle>> MeshHandles = ParticleActor->GetMeshInstances();
+		for (const TSharedPtr<FChaosVDMeshDataInstanceHandle>& Handle : MeshHandles)
+		{
+			if (const TSharedPtr<FChaosVDExtractedGeometryDataHandle> GeometryHandle = Handle->GetGeometryHandle())
+			{
+				if (GeometryHandle->GetImplicitObjectIndex() == QueryDataBeingInspected->SQVisitData[QueryDataBeingInspected->CurrentVisitIndex].ShapeIndex)
+				{
+					Handle->SetIsSelected(true);
+				}
+				else
+				{
+					Handle->SetIsSelected(false);
+				}
+			}
+		}
+		
+		if (const TSharedPtr<FEditorModeTools> EditorModeToolsPtr = EditorModeToolsWeakPtr.Pin())
+		{
+			// We need to request a re-draw to make sure the debug draw view is updated
+			if (FEditorViewportClient* ViewportClient = EditorModeToolsPtr->GetFocusedViewportClient())
+			{
+				ViewportClient->bNeedsRedraw = true;
+			}
+		}
+	}
+	
+	return FReply::Handled();
+}
+
+FReply SChaosVDSceneQueryDataInspector::SelectQueryToInspectByID(int32 QueryID)
+{
+	const TSharedPtr<FChaosVDScene> ScenePtr = SceneWeakPtr.Pin();
+	if (!ScenePtr)
+	{
+		return FReply::Handled();
+	}
+
+	if (UChaosVDSceneQueryDataComponent* SQDataComponent = ScenePtr->GetSceneQueryDataContainerComponent())
+	{
+		SQDataComponent->SelectQuery(QueryID);
+							
+		SetQueryDataToInspect(SQDataComponent->GetSelectedQueryHandle());
+	}
+	else
+	{
+		ClearInspector();	
+	}
+						
+	return FReply::Handled();
+}
+
+FReply SChaosVDSceneQueryDataInspector::SelectParentQuery()
+{
+	const TSharedPtr<FChaosVDScene> ScenePtr = SceneWeakPtr.Pin();
+	const UChaosVDSceneQueryDataComponent* SQDataComponent = ScenePtr ? ScenePtr->GetSceneQueryDataContainerComponent() : nullptr;
+	if (const TSharedPtr<FChaosVDQueryDataWrapper> SelectedQuery = SQDataComponent ? SQDataComponent->GetSelectedQueryHandle().GetQueryData().Pin() : nullptr)
+	{
+		SelectQueryToInspectByID(SelectedQuery->ParentQueryID);
+	}
+	else
+	{
+		ClearInspector();
+	}
+
+	return FReply::Handled();
+}
+
+TSharedPtr<IStructureDetailsView> SChaosVDSceneQueryDataInspector::CreateDataDetailsView()
+{
+	FPropertyEditorModule& PropertyEditorModule = FModuleManager::GetModuleChecked<FPropertyEditorModule>("PropertyEditor");
+
+	const FStructureDetailsViewArgs StructDetailsViewArgs;
+	FDetailsViewArgs DetailsViewArgs;
+	DetailsViewArgs.bShowOptions = false;
+	DetailsViewArgs.bAllowFavoriteSystem = false;
+	DetailsViewArgs.bAllowSearch = false;
+	DetailsViewArgs.bShowScrollBar = false;
+
+	return PropertyEditorModule.CreateStructureDetailView(DetailsViewArgs,StructDetailsViewArgs, nullptr);
+}
+
+void SChaosVDSceneQueryDataInspector::HandleSceneUpdated()
+{
+	if (GetCurrentDataBeingInspected())
+	{
+		bIsUpToDate = false;
+	}
+	else
+	{
+		ClearInspector();
+	}
+}
+
+void SChaosVDSceneQueryDataInspector::HandleSubQueryNameSelected(TSharedPtr<FName> Name)
+{
+	const TSharedPtr<FChaosVDScene> ScenePtr = SceneWeakPtr.Pin();
+	if (!ScenePtr)
+	{
+		return;
+	}
+
+	if (const int32* SelectedQueryID = CurrentSubQueriesByName.Find(Name))
+	{
+		SelectQueryToInspectByID(*SelectedQueryID);
+		return;
+	}
+
+	ClearInspector();
+
+	UE_LOG(LogChaosVDEditor, Error,TEXT("[%s] Failed to find selected subquery."), ANSI_TO_TCHAR(__FUNCTION__));
+}
+
+void SChaosVDSceneQueryDataInspector::ClearInspector()
+{
+	if (SceneQueryDataDetailsView)
+	{
+		SceneQueryDataDetailsView->SetStructureData(nullptr);
+	}
+
+	if (SceneQueryHitDataDetailsView)
+	{
+		SceneQueryHitDataDetailsView->SetStructureData(nullptr);
+	}
+
+	if (QueryStepsTimelineWidget)
+	{
+		QueryStepsTimelineWidget->UpdateMinMaxValue(0,0);
+	}
+
+	if (SubQueryNamePickerWidget)
+	{
+		SubQueryNamePickerWidget->UpdateNameList({});
+	}
+
+	CurrentSubQueriesByName.Empty();
+
+	CurrentSceneQueryBeingInspectedHandle = FChaosVDSceneQuerySelectionHandle(nullptr, INDEX_NONE);
+
+	bIsUpToDate = true;
+}
+
+EVisibility SChaosVDSceneQueryDataInspector::GetOutOfDateWarningVisibility() const
+{
+	return bIsUpToDate ? EVisibility::Collapsed : EVisibility::Visible;
+}
+
+EVisibility SChaosVDSceneQueryDataInspector::GetQueryDetailsSectionVisibility() const
+{
+	const TSharedPtr<FChaosVDQueryDataWrapper> QueryDataBeingInspected = CurrentSceneQueryBeingInspectedHandle.GetQueryData().Pin();
+	return QueryDataBeingInspected ? EVisibility::Visible : EVisibility::Collapsed;
+}
+
+EVisibility SChaosVDSceneQueryDataInspector::GetQueryStepPlaybackControlsVisibility() const
+{
+	const TSharedPtr<FChaosVDQueryDataWrapper> QueryDataBeingInspected = CurrentSceneQueryBeingInspectedHandle.GetQueryData().Pin();
+	// If the this inspector no longer reflects data represented in the viewport, we can't offer playback so we need to hide the controls
+	return bIsUpToDate && QueryDataBeingInspected ? EVisibility::Visible : EVisibility::Collapsed;
+}
+
+EVisibility SChaosVDSceneQueryDataInspector::GetSQVisitDetailsSectionVisibility() const
+{
+	const TSharedPtr<FChaosVDQueryDataWrapper> QueryDataBeingInspected = CurrentSceneQueryBeingInspectedHandle.GetQueryData().Pin();
+	return QueryDataBeingInspected && QueryDataBeingInspected->SQVisitData.Num() > 0 ? EVisibility::Visible : EVisibility::Collapsed;
+}
+
+EVisibility SChaosVDSceneQueryDataInspector::GetNothingSelectedMessageVisibility() const
+{
+	const TSharedPtr<FChaosVDQueryDataWrapper> QueryDataBeingInspected = CurrentSceneQueryBeingInspectedHandle.GetQueryData().Pin();
+	return QueryDataBeingInspected ? EVisibility::Collapsed : EVisibility::Visible;
+}
+
+EVisibility SChaosVDSceneQueryDataInspector::GetSubQuerySelectorVisibility() const
+{
+	const TSharedPtr<FChaosVDQueryDataWrapper> QueryDataBeingInspected = CurrentSceneQueryBeingInspectedHandle.GetQueryData().Pin();
+	return QueryDataBeingInspected && QueryDataBeingInspected->SubQueriesIDs.Num() > 0 ? EVisibility::Visible : EVisibility::Collapsed;
+}
+
+EVisibility SChaosVDSceneQueryDataInspector::GetParentQuerySelectorVisibility() const
+{
+	const TSharedPtr<FChaosVDQueryDataWrapper> QueryDataBeingInspected = CurrentSceneQueryBeingInspectedHandle.GetQueryData().Pin();
+	return QueryDataBeingInspected && QueryDataBeingInspected->ParentQueryID != INDEX_NONE ? EVisibility::Visible : EVisibility::Collapsed;
+}
+
+bool SChaosVDSceneQueryDataInspector::GetSelectParticleHitStateEnable() const
+{
+	const TSharedPtr<FChaosVDQueryDataWrapper> QueryDataBeingInspected = CurrentSceneQueryBeingInspectedHandle.GetQueryData().Pin();
+	if (QueryDataBeingInspected && QueryDataBeingInspected->SQVisitData.IsValidIndex(QueryDataBeingInspected->CurrentVisitIndex))
+	{
+		return QueryDataBeingInspected->SQVisitData[QueryDataBeingInspected->CurrentVisitIndex].ParticleIndex != INDEX_NONE;
+	}
+
+	return false;
+}
+
+bool SChaosVDSceneQueryDataInspector::GetSQVisitStepsEnabled() const
+{
+	const TSharedPtr<FChaosVDQueryDataWrapper> QueryDataBeingInspected = CurrentSceneQueryBeingInspectedHandle.GetQueryData().Pin();
+	return QueryDataBeingInspected && QueryDataBeingInspected->SQVisitData.Num() > 0;
+}
+
+TSharedPtr<FChaosVDQueryDataWrapper> SChaosVDSceneQueryDataInspector::GetCurrentDataBeingInspected()
+{
+	const TSharedPtr<FChaosVDQueryDataWrapper> QueryDataBeingInspected = CurrentSceneQueryBeingInspectedHandle.GetQueryData().Pin();
+	return QueryDataBeingInspected;
+}
+
+#undef LOCTEXT_NAMESPACE
diff --git a/Engine/Plugins/Experimental/ChaosVD/Source/ChaosVD/Public/Actors/ChaosVDSceneQueryDataContainer.h b/Engine/Plugins/Experimental/ChaosVD/Source/ChaosVD/Public/Actors/ChaosVDSceneQueryDataContainer.h
new file mode 100644
--- /dev/null
+++ b/Engine/Plugins/Experimental/ChaosVD/Source/ChaosVD/Public/Actors/ChaosVDSceneQueryDataContainer.h
@@ -0,0 +1,27 @@
+﻿// Copyright Epic Games, Inc. All Rights Reserved.
+
+#pragma once
+
+#include "ChaosVDSceneObjectBase.h"
+#include "GameFramework/Actor.h"
+#include "UObject/ObjectPtr.h"
+
+#include "ChaosVDSceneQueryDataContainer.generated.h"
+
+class UChaosVDSceneQueryDataComponent;
+
+/** Actor that contains Scene Queries components data and any related logic used to visualize recording scene queries */
+UCLASS()
+class CHAOSVD_API AChaosVDSceneQueryDataContainer : public AActor, public FChaosVDSceneObjectBase
+{
+	GENERATED_BODY()
+
+public:
+	AChaosVDSceneQueryDataContainer();
+
+	UChaosVDSceneQueryDataComponent* GetSceneQueryDataComponent() const { return SceneQueryDataComponent.Get(); }
+
+protected:
+	UPROPERTY()
+	TObjectPtr<UChaosVDSceneQueryDataComponent> SceneQueryDataComponent;
+};
diff --git a/Engine/Plugins/Experimental/ChaosVD/Source/ChaosVD/Public/ChaosVDEditorSettings.h b/Engine/Plugins/Experimental/ChaosVD/Source/ChaosVD/Public/ChaosVDEditorSettings.h
--- a/Engine/Plugins/Experimental/ChaosVD/Source/ChaosVD/Public/ChaosVDEditorSettings.h
+++ b/Engine/Plugins/Experimental/ChaosVD/Source/ChaosVD/Public/ChaosVDEditorSettings.h
@@ -8,6 +8,7 @@
 #include "UObject/Object.h"
 #include "UObject/SoftObjectPtr.h"
 #include "Visualizers/ChaosVDParticleDataVisualizer.h"
+#include "Visualizers/ChaosVDSceneQueryDataComponentVisualizer.h"
 #include "Visualizers/ChaosVDSolverCollisionDataComponentVisualizer.h"
 
 #include "ChaosVDEditorSettings.generated.h"
@@ -177,6 +178,9 @@ public:
 	UPROPERTY(EditAnywhere, Category = "Viewport Visualization Flags", meta = (Bitmask, BitmaskEnum = "/Script/ChaosVD.EChaosVDCollisionVisualizationFlags"))
 	uint32 GlobalCollisionDataVisualizationFlags = 0;
 
+	UPROPERTY(EditAnywhere, Category = "Viewport Visualization Flags", meta = (Bitmask, BitmaskEnum = "/Script/ChaosVD.EChaosVDSceneQueryVisualizationFlags"))
+	uint32 GlobalSceneQueriesVisualizationFlags = static_cast<uint8>(EChaosVDSceneQueryVisualizationFlags::DrawClientQueries | EChaosVDSceneQueryVisualizationFlags::DrawServerQueries |  EChaosVDSceneQueryVisualizationFlags::DrawHits | EChaosVDSceneQueryVisualizationFlags::DrawLineTraceQueries);
+
 	UPROPERTY(EditAnywhere, Category = "Viewport Visualization")
 	bool bShowDebugText = false;
 
diff --git a/Engine/Plugins/Experimental/ChaosVD/Source/ChaosVD/Public/ChaosVDExtractedGeometryDataHandle.h b/Engine/Plugins/Experimental/ChaosVD/Source/ChaosVD/Public/ChaosVDExtractedGeometryDataHandle.h
--- a/Engine/Plugins/Experimental/ChaosVD/Source/ChaosVD/Public/ChaosVDExtractedGeometryDataHandle.h
+++ b/Engine/Plugins/Experimental/ChaosVD/Source/ChaosVD/Public/ChaosVDExtractedGeometryDataHandle.h
@@ -37,11 +37,17 @@ struct FChaosVDExtractedGeometryDataHandle
 	/** Ptr to the implicit object this handle provides access to and from which the mesh was generated */
 	const Chaos::FImplicitObject* GetImplicitObject() const { return ImplicitObject; }
 
+	/** Returns the Index on which this Implicit object was in the Hierarchy of the root one */
+	int32 GetImplicitObjectIndex() const { return ObjectIndex; }
+
 	/** If the object was an union and this was a leaf, this getter returns the root union object, otherwise it will be the same as GetImplicitObject */
 	const Chaos::FImplicitObject* GetRootImplicitObject() const { return RootImplicitObject; }
 	
 	/** Sets the Implicit Object Ptr from which the data was extracted and the mesh was generated */
 	void SetImplicitObject(const Chaos::FImplicitObject* InImplicitObject) { ImplicitObject = InImplicitObject; }
+
+	/** Sets the Index on which this Implicit object was in the Hierarchy of the root one */
+	void SetImplicitObjectIndex(int32 InObjectIndex) { ObjectIndex = InObjectIndex; }
 	
 	/** Sets the Implicit Object Ptr from which the data was extracted and the mesh was generated */
 	void SetRootImplicitObject(const Chaos::FImplicitObject* InImplicitObject) { RootImplicitObject = InImplicitObject; }
@@ -57,12 +63,22 @@ struct FChaosVDExtractedGeometryDataHandle
 	}
 
 private:
+
+	/** Key representing the generated static mesh */
 	uint32 GeometryKey = 0;
 
+	/** Key to be used to match this geometry with recorded Shape Data */
 	uint32 DataComponentKey = 0;
 
+	/** Index on which this Implicit object was in the Hierarchy of the root one */
+	int32 ObjectIndex = INDEX_NONE;
+
+	/** Extracted transform of the generated implicit object in component space */
 	FTransform Transform;
 
+	/** The Implicit Object Ptr from which the data was extracted and the mesh was generated */
 	const Chaos::FImplicitObject* ImplicitObject = nullptr;
+
+	/** If the object was an union and this was a leaf, this will be the root union object, otherwise it will be the same as ImplicitObject */
 	const Chaos::FImplicitObject* RootImplicitObject = nullptr;
 };
diff --git a/Engine/Plugins/Experimental/ChaosVD/Source/ChaosVD/Public/ChaosVDGeometryBuilder.h b/Engine/Plugins/Experimental/ChaosVD/Source/ChaosVD/Public/ChaosVDGeometryBuilder.h
--- a/Engine/Plugins/Experimental/ChaosVD/Source/ChaosVD/Public/ChaosVDGeometryBuilder.h
+++ b/Engine/Plugins/Experimental/ChaosVD/Source/ChaosVD/Public/ChaosVDGeometryBuilder.h
@@ -221,14 +221,23 @@ private:
 	 */
 	bool ApplyMeshToComponentFromKey(TWeakObjectPtr<UMeshComponent> MeshComponent, const uint32 GeometryKey);
 
-	/** Creates a mesh generator for the provided Implicit object which will be used to create a Static Mesh or Dynamic Mesh */
-	TSharedPtr<UE::Geometry::FMeshShapeGenerator> CreateMeshGeneratorForImplicitObject(const Chaos::FImplicitObject* InImplicit);
+public:
+	/** Creates a mesh generator for the provided Implicit object which will be used to create a Static Mesh or Dynamic Mesh
+	 * @param InImplicit ImplicitObject used a data source for the mesh generator
+	 * @param SimpleShapesComplexityFactor Factor used to reduce or increase the complexity (number of triangles generated) of simple shapes (Sphere/Capsule)
+	 */
+	TSharedPtr<UE::Geometry::FMeshShapeGenerator> CreateMeshGeneratorForImplicitObject(const Chaos::FImplicitObject* InImplicit, float SimpleShapesComplexityFactor = 1.0f);
 
-	const Chaos::FImplicitObject* UnpackImplicitObject(const Chaos::FImplicitObject* InImplicitObject, Chaos::FRigidTransform3& InOutTransform) const;
+private:
 
+	const Chaos::FImplicitObject* UnpackImplicitObject(const Chaos::FImplicitObject* InImplicitObject, Chaos::FRigidTransform3& InOutTransform) const;
+	
+public:
 	/** Re-adjust the provided transform if needed, so it can be visualized properly with its generated mesh */
 	void AdjustedTransformForImplicit(const Chaos::FImplicitObject* InImplicit, FTransform& OutAdjustedTransform);
 
+private:
+
 	/** Extracts data from an implicit object in a format CVD can use, and starts the Mesh generation process if needed
 	 * @return Returns a handle to the generated data that can be used to access the generated mesh when ready
 	 */
@@ -421,6 +430,7 @@ void FChaosVDGeometryBuilder::CreateMeshesFromImplicit_Internal(const Chaos::FIm
 	if (const TSharedPtr<FChaosVDExtractedGeometryDataHandle> MeshDataHandle = ExtractGeometryDataForImplicit<MeshType>(InLeafImplicitObject, InTransform, MeshIndex))
 	{	
 		MeshDataHandle->SetImplicitObject(InLeafImplicitObject);
+		MeshDataHandle->SetImplicitObjectIndex(MeshIndex);
 		MeshDataHandle->SetRootImplicitObject(InRootImplicitObject);
 
 		OutMeshDataHandles.Add(MeshDataHandle);
diff --git a/Engine/Plugins/Experimental/ChaosVD/Source/ChaosVD/Public/ChaosVDParticleActor.h b/Engine/Plugins/Experimental/ChaosVD/Source/ChaosVD/Public/ChaosVDParticleActor.h
--- a/Engine/Plugins/Experimental/ChaosVD/Source/ChaosVD/Public/ChaosVDParticleActor.h
+++ b/Engine/Plugins/Experimental/ChaosVD/Source/ChaosVD/Public/ChaosVDParticleActor.h
@@ -97,6 +97,8 @@ public:
 
 	FChaosVDParticleDataUpdatedDelegate& OnParticleDataUpdated() { return ParticleDataUpdatedDelegate; };
 
+	TConstArrayView<TSharedPtr<FChaosVDMeshDataInstanceHandle>> GetMeshInstances() const { return MeshDataHandles; }
+
 protected:
 
 	void ProcessUpdatedAndRemovedHandles(TArray<TSharedPtr<FChaosVDExtractedGeometryDataHandle>>& OutExtractedGeometryDataHandles);
diff --git a/Engine/Plugins/Experimental/ChaosVD/Source/ChaosVD/Public/ChaosVDScene.h b/Engine/Plugins/Experimental/ChaosVD/Source/ChaosVD/Public/ChaosVDScene.h
--- a/Engine/Plugins/Experimental/ChaosVD/Source/ChaosVD/Public/ChaosVDScene.h
+++ b/Engine/Plugins/Experimental/ChaosVD/Source/ChaosVD/Public/ChaosVDScene.h
@@ -9,12 +9,14 @@
 #include "UObject/ObjectMacros.h"
 #include "UObject/ObjectPtr.h"
 
+class AChaosVDSceneQueryDataContainer;
 class AChaosVDSolverInfoActor;
 class AChaosVDSceneCollisionContainer;
 class UChaosVDEditorSettings;
 class FChaosVDGeometryBuilder;
 class AChaosVDParticleActor;
 class FReferenceCollector;
+class UChaosVDSceneQueryDataComponent;
 class UObject;
 class UMaterial;
 class USelection;
@@ -56,7 +58,7 @@ public:
 	// No need to deprecate the old version since it is not a public API nor inline 
 	void HandleNewGeometryData(const Chaos::FConstImplicitObjectPtr& Geometry, const uint32 GeometryID) const;
 	
-	void HandleEnterNewGameFrame(int32 FrameNumber, const TArray<int32>& AvailableSolversIds);
+	void HandleEnterNewGameFrame(int32 FrameNumber, const TArray<int32>& AvailableSolversIds, const FChaosVDGameFrameData& InNewGameFrameData);
 
 	/** Deletes all actors of the Scene and underlying UWorld */
 	void CleanUpScene();
@@ -105,6 +107,8 @@ public:
 
 	AActor* GetMeshComponentsContainerActor() const { return MeshComponentContainerActor; }
 
+	UChaosVDSceneQueryDataComponent* GetSceneQueryDataContainerComponent() const;
+
 	FChaosVDActorActiveStateUpdateDelegate& OnActorActiveStateChanged() { return ParticleActorUpdateDelegate; }
 
 	TSharedPtr<FChaosVDRecording> LoadedRecording;
@@ -118,6 +122,8 @@ private:
 	int32 GetIDForRecordedParticleData(const TSharedPtr<FChaosVDParticleDataWrapper>& InParticleData) const;
 	void CreateBaseLights(UWorld* TargetWorld) const;
 
+	void CreateSceneQueriesContainer(UWorld* TargetWorld);
+
 	AActor* CreateMeshComponentsContainer(UWorld* TargetWorld);
 
 	/** Creates the instance of the World which will be used the recorded data*/
@@ -172,6 +178,7 @@ private:
 	mutable AActor* SkySphere = nullptr;
 
 	AActor* MeshComponentContainerActor = nullptr;
+	AChaosVDSceneQueryDataContainer* SceneQueriesContainer = nullptr;
 
 	bool bIsInitialized = false;
 
diff --git a/Engine/Plugins/Experimental/ChaosVD/Source/ChaosVD/Public/ChaosVDSceneQueryDataInspectorTab.h b/Engine/Plugins/Experimental/ChaosVD/Source/ChaosVD/Public/ChaosVDSceneQueryDataInspectorTab.h
new file mode 100644
--- /dev/null
+++ b/Engine/Plugins/Experimental/ChaosVD/Source/ChaosVD/Public/ChaosVDSceneQueryDataInspectorTab.h
@@ -0,0 +1,26 @@
+﻿// Copyright Epic Games, Inc. All Rights Reserved.
+
+#pragma once
+
+#include "ChaosVDTabSpawnerBase.h"
+#include "Templates/SharedPointer.h"
+
+class SChaosVDSceneQueryDataInspector;
+
+/** Spawns and handles and instance for the visual debugger Scene Query Data Inspector tab */
+class FChaosVDSceneQueryDataInspectorTab : public FChaosVDTabSpawnerBase
+{
+public:
+	FChaosVDSceneQueryDataInspectorTab(const FName& InTabID, const TSharedPtr<FTabManager>& InTabManager, const TWeakPtr<SChaosVDMainTab>& InOwningTabWidget)
+		: FChaosVDTabSpawnerBase(InTabID, InTabManager, InOwningTabWidget)
+	{
+	}
+
+	virtual ~FChaosVDSceneQueryDataInspectorTab() override;
+	virtual TSharedRef<SDockTab> HandleTabSpawned(const FSpawnTabArgs& Args) override;
+
+	TWeakPtr<SChaosVDSceneQueryDataInspector> GetSceneQueryDataInspectorInstance() const { return SceneQueryDataInspector; }
+
+protected:
+	TSharedPtr<SChaosVDSceneQueryDataInspector> SceneQueryDataInspector;
+};
diff --git a/Engine/Plugins/Experimental/ChaosVD/Source/ChaosVD/Public/ChaosVDTabsIDs.h b/Engine/Plugins/Experimental/ChaosVD/Source/ChaosVD/Public/ChaosVDTabsIDs.h
--- a/Engine/Plugins/Experimental/ChaosVD/Source/ChaosVD/Public/ChaosVDTabsIDs.h
+++ b/Engine/Plugins/Experimental/ChaosVD/Source/ChaosVD/Public/ChaosVDTabsIDs.h
@@ -17,4 +17,5 @@ public:
 	static const FName CVDEditorSettings;
 	static const FName StatusBar;
 	static const FName CollisionDataDetails;
+	static const FName SceneQueryDataDetails;
 };
diff --git a/Engine/Plugins/Experimental/ChaosVD/Source/ChaosVD/Public/Components/ChaosVDSceneQueryDataComponent.h b/Engine/Plugins/Experimental/ChaosVD/Source/ChaosVD/Public/Components/ChaosVDSceneQueryDataComponent.h
new file mode 100644
--- /dev/null
+++ b/Engine/Plugins/Experimental/ChaosVD/Source/ChaosVD/Public/Components/ChaosVDSceneQueryDataComponent.h
@@ -0,0 +1,76 @@
+﻿// Copyright Epic Games, Inc. All Rights Reserved.
+
+#pragma once
+
+#include "Components/ActorComponent.h"
+#include "DataWrappers/ChaosVDQueryDataWrappers.h"
+#include "ChaosVDSceneQueryDataComponent.generated.h"
+
+struct FChaosVDGameFrameData;
+
+/** Struct used to pass data about a specific query to other objects */
+struct FChaosVDSceneQuerySelectionHandle
+{
+	FChaosVDSceneQuerySelectionHandle()
+	{		
+	}
+
+	FChaosVDSceneQuerySelectionHandle(const TWeakPtr<FChaosVDQueryDataWrapper>& InQueryData, int32 SQVisitIndex)
+		: QueryData(InQueryData)
+		, SQVisitIndex(SQVisitIndex)
+	{
+		if (const TSharedPtr<FChaosVDQueryDataWrapper> QueryDataPtr = QueryData.Pin())
+		{
+			QueryID = QueryDataPtr->ID;
+		}
+	}
+
+	void SetIsSelected(bool bNewSelected);
+	bool IsSelected() const;
+
+	TWeakPtr<FChaosVDQueryDataWrapper> GetQueryData() const { return QueryData; } 
+	int32 GetQueryID() const { return QueryID; } 
+	int32 GetSQVisitIndex() const { return SQVisitIndex; } 
+
+private:
+	TWeakPtr<FChaosVDQueryDataWrapper> QueryData = nullptr;
+	int32 SQVisitIndex = INDEX_NONE;
+	int32 QueryID = INDEX_NONE;
+};
+
+DECLARE_MULTICAST_DELEGATE_OneParam(FChaosVDSQSelectionChangedDelegate, const FChaosVDSceneQuerySelectionHandle& SelectionHandle)
+
+/** Actor Component that contains all the scene queries recorded at the current loaded frame */
+UCLASS()
+class CHAOSVD_API UChaosVDSceneQueryDataComponent : public UActorComponent
+{
+	GENERATED_BODY()
+
+public:
+	UChaosVDSceneQueryDataComponent();
+
+	void UpdateQueriesFromFrameData(const FChaosVDGameFrameData& InGameFrameData);
+
+	TConstArrayView<TSharedPtr<FChaosVDQueryDataWrapper>> GetQueriesByType(EChaosVDSceneQueryType Type) const;
+	TConstArrayView<TSharedPtr<FChaosVDQueryDataWrapper>> GetAllQueries() const;
+	TSharedPtr<FChaosVDQueryDataWrapper> GetQueryByID(int32 QueryID) const;
+
+	void SelectQuery(int32 QueryID);
+	void SelectQuery(const FChaosVDSceneQuerySelectionHandle& SelectionHandle);
+
+	bool IsQuerySelected(int32 QueryID) const;
+
+	FChaosVDSceneQuerySelectionHandle GetSelectedQueryHandle() const { return CurrentSQSelectionHandle; }
+
+	FChaosVDSQSelectionChangedDelegate& GetOnSelectionChangeDelegate() { return SelectionChangeDelegate; }
+
+protected:
+	
+	TMap<EChaosVDSceneQueryType, TArray<TSharedPtr<FChaosVDQueryDataWrapper>>> RecordedQueriesByType;
+	TMap<int32, TSharedPtr<FChaosVDQueryDataWrapper>> RecordedQueriesByID;
+	TArray<TSharedPtr<FChaosVDQueryDataWrapper>> RecordedQueries;
+	
+	FChaosVDSQSelectionChangedDelegate SelectionChangeDelegate;
+
+	FChaosVDSceneQuerySelectionHandle CurrentSQSelectionHandle;
+};
diff --git a/Engine/Plugins/Experimental/ChaosVD/Source/ChaosVD/Public/DetailsCustomizations/ChaosVDDetailsCustomizationUtils.h b/Engine/Plugins/Experimental/ChaosVD/Source/ChaosVD/Public/DetailsCustomizations/ChaosVDDetailsCustomizationUtils.h
--- a/Engine/Plugins/Experimental/ChaosVD/Source/ChaosVD/Public/DetailsCustomizations/ChaosVDDetailsCustomizationUtils.h
+++ b/Engine/Plugins/Experimental/ChaosVD/Source/ChaosVD/Public/DetailsCustomizations/ChaosVDDetailsCustomizationUtils.h
@@ -25,5 +25,19 @@ public:
 	 * Marks any property of the provided handles array as hidden if they are not valid CVD properties (meaning they don't have serialized data loaded from a CVD recording)
 	 * @param InPropertyHandles Handles of properties to evaluate and hide if needed
 	 */
-	static void HideInvalidParticleDataProperties(TConstArrayView<TSharedPtr<IPropertyHandle>> InPropertyHandles);
+	static void HideInvalidCVDDataWrapperProperties(TConstArrayView<TSharedPtr<IPropertyHandle>> InPropertyHandles);
+
+	/**
+	 * Marks any property of the provided handles array as hidden if they are not valid CVD properties (meaning they don't have serialized data loaded from a CVD recording), using the provided details builder
+	 * @param InPropertyHandles Handles of properties to evaluate and hide if needed
+	 * @param DetailBuilder Details builder that will hide the property
+	 */
+	static void HideInvalidCVDDataWrapperProperties(TConstArrayView<TSharedRef<IPropertyHandle>> InPropertyHandles, IDetailLayoutBuilder& DetailBuilder);
+
+	/**
+	 * Marks any property of the provided handles array as hidden if they are not valid CVD properties (meaning they don't have serialized data loaded from a CVD recording), using the provided details builder
+	 * @param InPropertyHandle Property Handle to evaluate if it is valid
+	 * @param bOutIsCVDBaseDataStruct Set to true if the property handle provided is from a CVD Wrapper Data Base. All other types will be deemed valid
+	 */
+	static bool HasValidCVDWrapperData(const TSharedPtr<IPropertyHandle>& InPropertyHandle, bool& bOutIsCVDBaseDataStruct);
 };
diff --git a/Engine/Plugins/Experimental/ChaosVD/Source/ChaosVD/Public/DetailsCustomizations/ChaosVDParticleDataWrapperCustomization.h b/Engine/Plugins/Experimental/ChaosVD/Source/ChaosVD/Public/DetailsCustomizations/ChaosVDParticleDataWrapperCustomization.h
--- a/Engine/Plugins/Experimental/ChaosVD/Source/ChaosVD/Public/DetailsCustomizations/ChaosVDParticleDataWrapperCustomization.h
+++ b/Engine/Plugins/Experimental/ChaosVD/Source/ChaosVD/Public/DetailsCustomizations/ChaosVDParticleDataWrapperCustomization.h
@@ -16,6 +16,6 @@ public:
 
 	static TSharedRef<IPropertyTypeCustomization> MakeInstance();
 
-	virtual void CustomizeHeader(TSharedRef<IPropertyHandle> StructPropertyHandle, FDetailWidgetRow& HeaderRow, IPropertyTypeCustomizationUtils& StructCustomizationUtils) override {}
+	virtual void CustomizeHeader(TSharedRef<IPropertyHandle> StructPropertyHandle, FDetailWidgetRow& HeaderRow, IPropertyTypeCustomizationUtils& StructCustomizationUtils) override;
 	virtual void CustomizeChildren(TSharedRef<IPropertyHandle> StructPropertyHandle, IDetailChildrenBuilder& StructBuilder, IPropertyTypeCustomizationUtils& StructCustomizationUtils) override;
 };
diff --git a/Engine/Plugins/Experimental/ChaosVD/Source/ChaosVD/Public/DetailsCustomizations/ChaosVDQueryDataWrappersCustomizationDetails.h b/Engine/Plugins/Experimental/ChaosVD/Source/ChaosVD/Public/DetailsCustomizations/ChaosVDQueryDataWrappersCustomizationDetails.h
new file mode 100644
--- /dev/null
+++ b/Engine/Plugins/Experimental/ChaosVD/Source/ChaosVD/Public/DetailsCustomizations/ChaosVDQueryDataWrappersCustomizationDetails.h
@@ -0,0 +1,38 @@
+﻿// Copyright Epic Games, Inc. All Rights Reserved.
+
+#pragma once
+#include "IDetailCustomization.h"
+#include "IPropertyTypeCustomization.h"
+
+/** Custom property layout for the ChaosVD SQ Data wrapper struct */
+class FChaosVDQueryDataWrappersCustomizationDetails : public IPropertyTypeCustomization
+{
+	static TSharedRef<IPropertyTypeCustomization> MakeInstance();
+
+	virtual void CustomizeHeader(TSharedRef<IPropertyHandle> StructPropertyHandle, FDetailWidgetRow& HeaderRow, IPropertyTypeCustomizationUtils& StructCustomizationUtils) override {}
+	virtual void CustomizeChildren(TSharedRef<IPropertyHandle> StructPropertyHandle, IDetailChildrenBuilder& StructBuilder, IPropertyTypeCustomizationUtils& StructCustomizationUtils) override;
+};
+
+/** Custom details panel for the ChaosVD SQ Visit Data struct */
+class FChaosVDQueryVisitDataCustomization : public IDetailCustomization
+{
+public:
+	FChaosVDQueryVisitDataCustomization(){};
+	virtual ~FChaosVDQueryVisitDataCustomization() override {};
+
+	static TSharedRef<IDetailCustomization> MakeInstance();
+
+	virtual void CustomizeDetails(IDetailLayoutBuilder& DetailBuilder) override;
+};
+
+/** Custom details panel for the ChaosVD SQ Data Wrapper struct */
+class FChaosVDQueryDataWrapperCustomization : public IDetailCustomization
+{
+public:
+	FChaosVDQueryDataWrapperCustomization(){};
+	virtual ~FChaosVDQueryDataWrapperCustomization() override{};
+
+	static TSharedRef<IDetailCustomization> MakeInstance();
+
+	virtual void CustomizeDetails(IDetailLayoutBuilder& DetailBuilder) override;
+};
diff --git a/Engine/Plugins/Experimental/ChaosVD/Source/ChaosVD/Public/Trace/DataProcessors/ChaosVDSceneQueryDataProcessor.h b/Engine/Plugins/Experimental/ChaosVD/Source/ChaosVD/Public/Trace/DataProcessors/ChaosVDSceneQueryDataProcessor.h
new file mode 100644
--- /dev/null
+++ b/Engine/Plugins/Experimental/ChaosVD/Source/ChaosVD/Public/Trace/DataProcessors/ChaosVDSceneQueryDataProcessor.h
@@ -0,0 +1,17 @@
+﻿// Copyright Epic Games, Inc. All Rights Reserved.
+
+#pragma once
+
+#include "Trace/DataProcessors/IChaosVDDataProcessor.h"
+
+/**
+ * Data processor implementation that is able to deserialize traced scene queries data
+ */
+class FChaosVDSceneQueryDataProcessor final : public IChaosVDDataProcessor
+{
+public:
+	explicit FChaosVDSceneQueryDataProcessor();
+
+	virtual bool ProcessRawData(const TArray<uint8>& InData) override;
+};
+
diff --git a/Engine/Plugins/Experimental/ChaosVD/Source/ChaosVD/Public/Trace/DataProcessors/ChaosVDSceneQueryVisitDataProcessor.h b/Engine/Plugins/Experimental/ChaosVD/Source/ChaosVD/Public/Trace/DataProcessors/ChaosVDSceneQueryVisitDataProcessor.h
new file mode 100644
--- /dev/null
+++ b/Engine/Plugins/Experimental/ChaosVD/Source/ChaosVD/Public/Trace/DataProcessors/ChaosVDSceneQueryVisitDataProcessor.h
@@ -0,0 +1,16 @@
+﻿// Copyright Epic Games, Inc. All Rights Reserved.
+
+#pragma once
+#include "IChaosVDDataProcessor.h"
+
+/**
+ * Data processor implementation that is able to deserialize traced SQ Visits
+ */
+class FChaosVDSceneQueryVisitDataProcessor final : public IChaosVDDataProcessor
+{
+public:
+	explicit FChaosVDSceneQueryVisitDataProcessor();
+
+	virtual bool ProcessRawData(const TArray<uint8>& InData) override;
+};
+
diff --git a/Engine/Plugins/Experimental/ChaosVD/Source/ChaosVD/Public/Visualizers/ChaosVDDebugDrawUtils.h b/Engine/Plugins/Experimental/ChaosVD/Source/ChaosVD/Public/Visualizers/ChaosVDDebugDrawUtils.h
--- a/Engine/Plugins/Experimental/ChaosVD/Source/ChaosVD/Public/Visualizers/ChaosVDDebugDrawUtils.h
+++ b/Engine/Plugins/Experimental/ChaosVD/Source/ChaosVD/Public/Visualizers/ChaosVDDebugDrawUtils.h
@@ -1,9 +1,13 @@
 ﻿// Copyright Epic Games, Inc. All Rights Reserved.
 
 #pragma once
+
+#include "Chaos/ImplicitFwd.h"
 #include "Containers/Queue.h"
 #include "SceneManagement.h"
 
+class FChaosVDGeometryBuilder;
+
 /** Utility methods that allows Debug draw into the Chaos VD Editor */
 class FChaosVDDebugDrawUtils
 {
@@ -16,6 +20,7 @@ public:
 	static void DrawCircle(FPrimitiveDrawInterface* PDI, const FVector& Origin, float Radius, int32 Segments, const FColor& Color, float Thickness, const FVector& XAxis, const FVector& YAxis, FStringView DebugText, ESceneDepthPriorityGroup DepthPriority = SDPG_World);
 	static void DrawBox(FPrimitiveDrawInterface* PDI, const FVector& InExtents, const FColor& InColor, const FTransform& InTransform, FStringView DebugText, ESceneDepthPriorityGroup DepthPriority = SDPG_World);
 	static void DrawLine(FPrimitiveDrawInterface* PDI, const FVector& InStartPosition, const FVector& InEndPosition, const FColor& InColor, FStringView DebugText, ESceneDepthPriorityGroup DepthPriority = SDPG_World);
+	static void DrawImplicitObject(FPrimitiveDrawInterface* PDI, const TSharedPtr<FChaosVDGeometryBuilder>& GeometryGenerator, const Chaos::FConstImplicitObjectPtr& ImplicitObject, const FTransform& InWorldTransform, const FColor& InColor, FStringView DebugText, ESceneDepthPriorityGroup DepthPriority = SDPG_World);
 
 	static void DrawCanvas(FViewport& InViewport, FSceneView& View, FCanvas& Canvas);
 
diff --git a/Engine/Plugins/Experimental/ChaosVD/Source/ChaosVD/Public/Visualizers/ChaosVDSceneQueryDataComponentVisualizer.h b/Engine/Plugins/Experimental/ChaosVD/Source/ChaosVD/Public/Visualizers/ChaosVDSceneQueryDataComponentVisualizer.h
new file mode 100644
--- /dev/null
+++ b/Engine/Plugins/Experimental/ChaosVD/Source/ChaosVD/Public/Visualizers/ChaosVDSceneQueryDataComponentVisualizer.h
@@ -0,0 +1,99 @@
+﻿// Copyright Epic Games, Inc. All Rights Reserved.
+
+#pragma once
+
+#include "ComponentVisualizer.h"
+#include "IChaosVDParticleVisualizationDataProvider.h"
+#include "Chaos/ImplicitFwd.h"
+#include "Components/ChaosVDSceneQueryDataComponent.h"
+
+struct FChaosVDRecording;
+class FChaosVDGeometryBuilder;
+struct FChaosVDVisualizationContext;
+struct FChaosVDQueryDataWrapper;
+
+/** Visualization context structure specific for Scene Queries visualizations */
+struct FChaosVDSceneQueryVisualizationDataContext : public FChaosVDVisualizationContext
+{
+	FChaosVDSceneQuerySelectionHandle DataSelectionHandle = FChaosVDSceneQuerySelectionHandle(nullptr, INDEX_NONE);
+	bool bIsServerVisualizationEnabled = false;
+	bool bIsClientVisualizationEnabled = false;
+
+	/** Generates a random color based on the selection state and query ID, which will be used to debug draw the scene query */
+	void GenerateColor(int32 QueryID, bool bIsSelected)
+	{
+		RandomSeededColor = bIsSelected ? FLinearColor::White : FLinearColor::MakeRandomSeededColor(QueryID);
+		DebugDrawColor = RandomSeededColor.ToFColorSRGB();
+		DebugDrawDarkerColor = (RandomSeededColor * 0.85f).ToFColorSRGB();
+		HitColor = (RandomSeededColor * 1.2f).ToFColorSRGB();
+	}
+
+	FLinearColor RandomSeededColor = FLinearColor(EForceInit::ForceInit);
+	FColor DebugDrawColor = FColor(EForceInit::ForceInit);
+	FColor DebugDrawDarkerColor = FColor(EForceInit::ForceInit);
+	FColor HitColor = FColor(EForceInit::ForceInit);
+
+	Chaos::FConstImplicitObjectPtr InputGeometry = nullptr;
+	TWeakPtr<FChaosVDGeometryBuilder> GeometryGenerator = nullptr;
+};
+
+/** Custom Hit Proxy for debug drawn scene queries */
+struct HChaosVDSceneQueryProxy : public HComponentVisProxy
+{
+	DECLARE_HIT_PROXY()
+	
+	HChaosVDSceneQueryProxy(const UActorComponent* Component, const FChaosVDSceneQuerySelectionHandle& InContactFinderData) : HComponentVisProxy(Component, HPP_UI), DataSelectionHandle(InContactFinderData)
+	{	
+	}
+
+	virtual EMouseCursor::Type GetMouseCursor() override
+	{
+		return EMouseCursor::Crosshairs;
+	}
+
+	FChaosVDSceneQuerySelectionHandle DataSelectionHandle;
+};
+
+/** Set of visualization flags options for Scene Queries */
+UENUM(meta = (Bitflags, UseEnumValuesAsMaskValuesInEditor = "true"))
+enum class EChaosVDSceneQueryVisualizationFlags: uint32
+{
+	None					= 0 UMETA(Hidden),
+	EnableDraw				= 1 << 0,
+	DrawLineTraceQueries	= 1 << 1,
+	DrawSweepQueries		= 1 << 2,
+	DrawOverlapQueries		= 1 << 3,
+	DrawHits				= 1 << 4,
+	DrawServerQueries		= 1 << 5,
+	DrawClientQueries		= 1 << 6,
+	OnlyDrawSelectedQuery	= 1 << 7,
+	HideEmptyQueries		= 1 << 8,
+	HideSubQueries			= 1 << 9,
+};
+ENUM_CLASS_FLAGS(EChaosVDSceneQueryVisualizationFlags);
+
+/**
+ * Component visualizer in charge of generating debug draw visualizations for scene queries in a ChaosVDSceneQueryDataComponent
+ */
+class FChaosVDSceneQueryDataComponentVisualizer final : public FComponentVisualizer
+{
+
+public:
+	FChaosVDSceneQueryDataComponentVisualizer()
+	{
+	}
+
+	virtual void DrawVisualization(const UActorComponent* Component, const FSceneView* View, FPrimitiveDrawInterface* PDI) override;
+	virtual bool VisProxyHandleClick(FEditorViewportClient* InViewportClient, HComponentVisProxy* VisProxy, const FViewportClick& Click) override;
+
+private:
+
+	void DrawLineTraceQuery(const UActorComponent* Component, const FChaosVDQueryDataWrapper& SceneQueryData, FChaosVDSceneQueryVisualizationDataContext& VisualizationContext, const FSceneView* View, FPrimitiveDrawInterface* PDI);
+	void DrawOverlapQuery(const UActorComponent* Component, const FChaosVDQueryDataWrapper& SceneQueryData, FChaosVDSceneQueryVisualizationDataContext& VisualizationContext, const FSceneView* View, FPrimitiveDrawInterface* PDI);
+	void DrawSweepQuery(const UActorComponent* Component, const FChaosVDQueryDataWrapper& SceneQueryData, FChaosVDSceneQueryVisualizationDataContext& VisualizationContext, const FSceneView* View, FPrimitiveDrawInterface* PDI);
+	void DrawHits(const UActorComponent* Component, const FChaosVDQueryDataWrapper& SceneQueryData, FPrimitiveDrawInterface* PDI, const FColor& InColor, FChaosVDSceneQueryVisualizationDataContext& VisualizationContext);
+
+	void DrawSceneQuery(const UActorComponent* Component, const FSceneView* View, FPrimitiveDrawInterface* PDI, const TSharedPtr<FChaosVDScene>& CVDScene, const TSharedPtr<FChaosVDRecording>& CVDRecording, FChaosVDSceneQueryVisualizationDataContext& VisualizationContext, const TSharedPtr<FChaosVDQueryDataWrapper>& Query);
+
+	bool HasEndLocation(const FChaosVDQueryDataWrapper& SceneQueryData) const;
+};
diff --git a/Engine/Plugins/Experimental/ChaosVD/Source/ChaosVD/Public/Widgets/SChaosVDMainTab.h b/Engine/Plugins/Experimental/ChaosVD/Source/ChaosVD/Public/Widgets/SChaosVDMainTab.h
--- a/Engine/Plugins/Experimental/ChaosVD/Source/ChaosVD/Public/Widgets/SChaosVDMainTab.h
+++ b/Engine/Plugins/Experimental/ChaosVD/Source/ChaosVD/Public/Widgets/SChaosVDMainTab.h
@@ -52,6 +52,8 @@ public:
 	TSharedPtr<FComponentVisualizer> FindComponentVisualizer(UClass* ClassPtr);
 	TSharedPtr<FComponentVisualizer> FindComponentVisualizer(FName ClassName);
 
+	TConstArrayView<TSharedPtr<FComponentVisualizer>> GetAllComponentVisualizers() { return ComponentVisualizers; }
+
 private:
 
 	template<typename TabType>
@@ -83,6 +85,7 @@ private:
 	TMap<FName, TSharedPtr<FChaosVDTabSpawnerBase>> TabSpawnersByIDMap;
 
 	TMap<FName, TSharedPtr<FComponentVisualizer>> ComponentVisualizersMap;
+	TArray<TSharedPtr<FComponentVisualizer>> ComponentVisualizers;
 
 	TMap<FName, TWeakPtr<SDockTab>> ActiveTabsByID;
 
diff --git a/Engine/Plugins/Experimental/ChaosVD/Source/ChaosVD/Public/Widgets/SChaosVDSceneQueryDataInspector.h b/Engine/Plugins/Experimental/ChaosVD/Source/ChaosVD/Public/Widgets/SChaosVDSceneQueryDataInspector.h
new file mode 100644
--- /dev/null
+++ b/Engine/Plugins/Experimental/ChaosVD/Source/ChaosVD/Public/Widgets/SChaosVDSceneQueryDataInspector.h
@@ -0,0 +1,113 @@
+﻿// Copyright Epic Games, Inc. All Rights Reserved.
+
+#pragma once
+
+#include "Components/ChaosVDSceneQueryDataComponent.h"
+#include "Widgets/SCompoundWidget.h"
+
+struct FChaosVDSceneQuerySelectionHandle;
+class SChaosVDNameListPicker;
+class FEditorModeTools;
+class IToolkitHost;
+class SChaosVDTimelineWidget;
+struct FChaosVDQueryDataWrapper;
+class IStructureDetailsView;
+class FChaosVDScene;
+
+/**
+ * Widget for the Chaos Visual Debugger Scene Queries data inspector
+ */
+class SChaosVDSceneQueryDataInspector : public SCompoundWidget
+{
+public:
+	SChaosVDSceneQueryDataInspector();
+
+	SLATE_BEGIN_ARGS(SChaosVDSceneQueryDataInspector)
+		{
+		}
+	SLATE_END_ARGS()
+
+	virtual ~SChaosVDSceneQueryDataInspector() override;
+
+	/** Constructs this widget with InArgs */
+	void Construct(const FArguments& InArgs, const TWeakPtr<FChaosVDScene>& InScenePtr, const TWeakPtr<FEditorModeTools>& InEditorModeTools);
+
+	/** Sets a new query data to be inspected */
+	void SetQueryDataToInspect(const FChaosVDSceneQuerySelectionHandle& InQueryDataSelectionHandle);
+
+protected:
+
+	TSharedRef<SWidget> GenerateQueryTagInfoRow();
+	TSharedRef<SWidget> GenerateQueryNavigationBoxWidget(float TagTitleBoxHorizontalPadding, float TagTitleBoxVerticalPadding);
+	TSharedRef<SWidget> GenerateQueryDetailsPanelSection(float InnerDetailsPanelsHorizontalPadding, float InnerDetailsPanelsVerticalPadding);
+	TSharedRef<SWidget> GenerateVisitStepControls();
+
+	void HandleQueryStepSelectionUpdated(int32 NewStepIndex);
+	
+	FText GetQueryBeingInspectedTag() const;
+	FText GetSelectParticleText() const;
+	FText GetSQVisitsStepsText() const;
+	
+	FReply SelectParticleForCurrentQueryData() const;
+	FReply SelectQueryToInspectByID(int32 QueryID);
+	FReply SelectParentQuery();
+
+	static TSharedPtr<IStructureDetailsView> CreateDataDetailsView();
+	
+	void HandleSceneUpdated();
+	void HandleSubQueryNameSelected(TSharedPtr<FName> Name);
+
+	void ClearInspector();
+	
+	EVisibility GetOutOfDateWarningVisibility() const;
+	EVisibility GetQueryDetailsSectionVisibility() const;
+	EVisibility GetQueryStepPlaybackControlsVisibility() const;
+	EVisibility GetSQVisitDetailsSectionVisibility() const;
+	EVisibility GetNothingSelectedMessageVisibility() const;
+	EVisibility GetSubQuerySelectorVisibility() const;
+	EVisibility GetParentQuerySelectorVisibility() const;
+
+	bool GetSelectParticleHitStateEnable() const;
+	bool GetSQVisitStepsEnabled() const;
+
+	TSharedPtr<FChaosVDQueryDataWrapper> GetCurrentDataBeingInspected();
+
+	TSharedPtr<SChaosVDTimelineWidget> QueryStepsTimelineWidget;
+	
+	TWeakPtr<FChaosVDScene> SceneWeakPtr;
+
+	TSharedPtr<IStructureDetailsView> SceneQueryDataDetailsView;
+	
+	TSharedPtr<IStructureDetailsView> SceneQueryHitDataDetailsView;
+
+	TSharedPtr<SChaosVDNameListPicker> SubQueryNamePickerWidget;
+
+	TWeakPtr<FEditorModeTools> EditorModeToolsWeakPtr;
+
+	TMap<TSharedPtr<FName>, int32> CurrentSubQueriesByName;
+	
+	FChaosVDSceneQuerySelectionHandle CurrentSceneQueryBeingInspectedHandle;
+
+	bool bIsUpToDate = true;
+
+	bool bListenToSelectionEvents = true;
+
+	friend struct FScopedSQInspectorSilencedSelectionEvents;
+};
+
+/** Structure that makes a SQ Inspector ignore selection events withing a scope */
+struct FScopedSQInspectorSilencedSelectionEvents
+{
+	FScopedSQInspectorSilencedSelectionEvents(SChaosVDSceneQueryDataInspector& InInspectorIgnoringEvents) : InspectorIgnoringSelectionEvents(InInspectorIgnoringEvents)
+	{
+		InspectorIgnoringSelectionEvents.bListenToSelectionEvents = false;
+	}
+	
+	~FScopedSQInspectorSilencedSelectionEvents()
+	{
+		InspectorIgnoringSelectionEvents.bListenToSelectionEvents = true;
+	}
+	
+private:
+	SChaosVDSceneQueryDataInspector& InspectorIgnoringSelectionEvents;
+};
diff --git a/Engine/Source/Runtime/AnimGraphRuntime/Private/BoneControllers/AnimNode_RigidBody.cpp b/Engine/Source/Runtime/AnimGraphRuntime/Private/BoneControllers/AnimNode_RigidBody.cpp
--- a/Engine/Source/Runtime/AnimGraphRuntime/Private/BoneControllers/AnimNode_RigidBody.cpp
+++ b/Engine/Source/Runtime/AnimGraphRuntime/Private/BoneControllers/AnimNode_RigidBody.cpp
@@ -1067,6 +1067,7 @@ void FAnimNode_RigidBody::InitPhysics(const UAnimInstance* InAnimInstance)
 
 #if WITH_CHAOS_VISUAL_DEBUGGER
 		PhysicsSimulation->GetChaosVDContextData().Id = FChaosVDRuntimeModule::Get().GenerateUniqueID();
+		PhysicsSimulation->GetChaosVDContextData().Id = static_cast<int32>(EChaosVDContextType::Solver);
 #endif
 
 		const int32 NumBodies = UsePhysicsAsset->SkeletalBodySetups.Num();
diff --git a/Engine/Source/Runtime/Engine/Private/Collision/ChaosVDSQTraceHelper.cpp b/Engine/Source/Runtime/Engine/Private/Collision/ChaosVDSQTraceHelper.cpp
new file mode 100644
--- /dev/null
+++ b/Engine/Source/Runtime/Engine/Private/Collision/ChaosVDSQTraceHelper.cpp
@@ -0,0 +1,100 @@
+﻿// Copyright Epic Games, Inc. All Rights Reserved.
+
+#include "ChaosVDSQTraceHelper.h"
+
+#if WITH_CHAOS_VISUAL_DEBUGGER
+
+#include "Components/PrimitiveComponent.h"
+#include "Engine/World.h"
+
+
+#include "ChaosVisualDebugger/ChaosVDContextProvider.h"
+#include "ChaosVisualDebugger/ChaosVisualDebuggerTrace.h"
+#include "DataWrappers/ChaosVDQueryDataWrappers.h"
+#include "PBDRigidsSolver.h"
+#include "Physics/PhysicsGeometry.h"
+#include "Physics/Experimental/PhysScene_Chaos.h"
+#include "PhysicsEngine/PhysicsLogUtil.h"
+
+namespace Chaos::VisualDebugger::TraceHelpers
+{
+
+/** This is not great but all ignored actors/components in the collision query params struct are stores as the UObject internal index,
+ * so we need to do this reverse look up to get the actual object to be able to extract the name or access the particles, so then we can record it */
+template<typename TObject>
+const TObject* GetObjetPointerFromID(int32 ID)
+{
+	constexpr  bool bEvenIfGarbage = false;
+	const FUObjectItem* ObjectItem = GUObjectArray.IndexToValidObject(ID, bEvenIfGarbage);
+	if (const TObject* CastedObject = Cast<TObject>(ObjectItem ? static_cast<const UObject*>(ObjectItem->Object) : nullptr))
+	{
+		return CastedObject;
+	}
+
+	return nullptr;
+}
+
+int32 GetWorldSolverID(const UWorld* World)
+	{
+		if (Chaos::FPhysicsSolver* Solver = World->GetPhysicsScene() ? World->GetPhysicsScene()->GetSolver() : nullptr)
+		{
+			return  Solver->GetChaosVDContextData().Id;
+		}
+
+		return INDEX_NONE;
+	}
+
+	FChaosVDContext CreateSceneQueryContextHelper()
+	{
+		if (!FChaosVisualDebuggerTrace::IsTracing())
+		{
+			return FChaosVDContext();
+		}
+
+		FChaosVDContext CurrentCVDContext;
+		CVD_GET_CURRENT_CONTEXT(CurrentCVDContext)
+		
+		FChaosVDContext NewSceneQueryContext;
+		NewSceneQueryContext.Id = FChaosVDRuntimeModule::Get().GenerateUniqueID();
+
+		if (CurrentCVDContext.Type ==  static_cast<int32>(EChaosVDContextType::Query) || CurrentCVDContext.Type == static_cast<int32>(EChaosVDContextType::SubTraceQuery))
+		{
+			NewSceneQueryContext.OwnerID = CurrentCVDContext.Id;
+			NewSceneQueryContext.Type = static_cast<int32>(EChaosVDContextType::SubTraceQuery);
+		}
+		else
+		{
+			NewSceneQueryContext.Type = static_cast<int32>(EChaosVDContextType::Query);
+		}
+
+		return NewSceneQueryContext;
+	}
+
+	void TraceCVDSceneQueryStartHelper(const UWorld* World, const FPhysicsGeometry* Geom, const FTransform& StartGeomPose, const FVector& EndLocation, ECollisionChannel TraceChannel, const FCollisionQueryParams& Params, const FCollisionResponseParams& ResponseParams, const FCollisionObjectQueryParams& ObjectParams, EChaosVDSceneQueryType Type, EChaosVDSceneQueryMode Mode, bool bIsRetry)
+	{
+		if (!FChaosVisualDebuggerTrace::IsTracing())
+		{
+			return;
+		}
+
+		//TODO: See if we can move most of the code here to the Trace class of CVD using templates, like we do already for other similar cases
+
+		FChaosVDCollisionQueryParams CVDCollisionParams;
+		CVDCollisionParams.CopyFrom(Params);
+
+		CVDCollisionParams.IgnoredActorsNames.Reserve(CVDCollisionParams.IgnoredActorsIDs.Num());
+		CVDCollisionParams.IgnoredComponentsNames.Reserve(CVDCollisionParams.IgnoredComponentsIDs.Num());
+		Algo::Transform(CVDCollisionParams.IgnoredActorsIDs, CVDCollisionParams.IgnoredActorsNames, [](uint32 ID){ return FName(PhysicsLogUtil::MakeActorNameString(GetObjetPointerFromID<AActor>(ID))); } );
+		Algo::Transform(CVDCollisionParams.IgnoredComponentsIDs, CVDCollisionParams.IgnoredComponentsNames, [](uint32 ID){ return FName(PhysicsLogUtil::MakeActorNameString(GetObjetPointerFromID<UActorComponent>(ID))); } );
+
+		FChaosVDCollisionResponseParams CVDCollisionResponseParams;
+		CVDCollisionResponseParams.CopyFrom(ResponseParams);
+		FChaosVDCollisionObjectQueryParams CVDCollisionObjectParams;
+		CVDCollisionObjectParams.CopyFrom(ObjectParams);
+
+		CVD_TRACE_SCENE_QUERY_START(Geom, StartGeomPose.GetRotation(), StartGeomPose.GetLocation(), EndLocation, TraceChannel, MoveTemp(CVDCollisionParams), MoveTemp(CVDCollisionResponseParams), MoveTemp(CVDCollisionObjectParams), Type, Mode, TraceHelpers::GetWorldSolverID(World), bIsRetry);
+	}
+}
+
+#endif
+	
diff --git a/Engine/Source/Runtime/Engine/Private/Collision/ChaosVDSQTraceHelper.h b/Engine/Source/Runtime/Engine/Private/Collision/ChaosVDSQTraceHelper.h
new file mode 100644
--- /dev/null
+++ b/Engine/Source/Runtime/Engine/Private/Collision/ChaosVDSQTraceHelper.h
@@ -0,0 +1,43 @@
+﻿// Copyright Epic Games, Inc. All Rights Reserved.
+
+#pragma once
+#if WITH_CHAOS_VISUAL_DEBUGGER
+
+#include "ChaosVisualDebugger/ChaosVDContextProvider.h"
+#include "HAL/Platform.h"
+#include "Math/Transform.h"
+#include "PhysicsInterfaceDeclaresCore.h"
+
+class UWorld;
+
+enum ECollisionChannel: int;
+enum class EChaosVDSceneQueryType;
+enum class EChaosVDSceneQueryMode;
+
+struct FCollisionQueryParams;
+struct FCollisionResponseParams;
+struct FCollisionObjectQueryParams;
+
+/** Initializes a CVD Wrapper for a Scene Query entry with the provider data, and traces it immediately but it keeps scope the CVD context active until it goes out scope (so other operations like SQ Visits can be recorded under the same context) */
+#ifndef CVD_TRACE_SCOPED_SCENE_QUERY_HELPER
+	#define CVD_TRACE_SCOPED_SCENE_QUERY_HELPER(World, Geom, StartGeomPose, EndLocation, TraceChannel, Params,  ResponseParams, ObjectParams, Type, Mode, bIsRetry) \
+	CVD_SCOPE_CONTEXT(Chaos::VisualDebugger::TraceHelpers::CreateSceneQueryContextHelper());\
+	Chaos::VisualDebugger::TraceHelpers::TraceCVDSceneQueryStartHelper(World, Geom, StartGeomPose, EndLocation, TraceChannel, Params,  ResponseParams, ObjectParams, Type, Mode, bIsRetry);
+#endif
+
+namespace Chaos::VisualDebugger::TraceHelpers
+{
+	int32 GetWorldSolverID(const UWorld* World);
+
+	/** Creates a CVD context with all the data required to trace a scene query */
+	FChaosVDContext CreateSceneQueryContextHelper();
+
+	/** Process types not accessible in the Chaos module, and call the CVD trace to perform the trace of the scene query */
+	void TraceCVDSceneQueryStartHelper(const UWorld* World, const FPhysicsGeometry* Geom, const FTransform& StartGeomPose, const FVector& EndLocation, ECollisionChannel TraceChannel, const FCollisionQueryParams& Params, const FCollisionResponseParams& ResponseParams, const FCollisionObjectQueryParams& ObjectParams, EChaosVDSceneQueryType Type, EChaosVDSceneQueryMode Mode, bool bIsRetry);
+}
+
+#else // WITH_CHAOS_VISUAL_DEBUGGER
+	#ifndef CVD_TRACE_SCOPED_SCENE_QUERY_HELPER
+		#define CVD_TRACE_SCOPED_SCENE_QUERY_HELPER(World, Geom, StartGeomPose, EndLocation, TraceChannel, Params,  ResponseParams, ObjectParams, Type, Mode, bIsRetry)
+	#endif
+#endif // WITH_CHAOS_VISUAL_DEBUGGER
\ No newline at end of file
diff --git a/Engine/Source/Runtime/Engine/Private/Collision/SceneQuery.cpp b/Engine/Source/Runtime/Engine/Private/Collision/SceneQuery.cpp
--- a/Engine/Source/Runtime/Engine/Private/Collision/SceneQuery.cpp
+++ b/Engine/Source/Runtime/Engine/Private/Collision/SceneQuery.cpp
@@ -12,6 +12,12 @@
 
 #include "Collision/CollisionDebugDrawing.h"
 
+#if WITH_CHAOS_VISUAL_DEBUGGER
+#include "DataWrappers/ChaosVDQueryDataWrappers.h"
+#endif
+
+#include "ChaosVDSQTraceHelper.h"
+
 float DebugLineLifetime = 2.f;
 
 #include "PhysicsEngine/CollisionAnalyzerCapture.h"
@@ -659,11 +665,16 @@ bool TSceneCastCommonImp(const UWorld* World, typename Traits::TOutHits& OutHits
 	bool bReturnResult = false;
 	FCollisionQueryParams RetryParams;
 
-	bReturnResult = TSceneCastCommonImpWithRetryRequest<Traits, TGeomInputs, TAccelContainer>(World, OutHits, GeomInputs, Start, End, TraceChannel, Params, ResponseParams, ObjectParams, AccelContainer, bRequestRetry, RetryParams);	
-	
+	{
+		constexpr bool bIsRetryQuery = false;
+		CVD_TRACE_SCOPED_SCENE_QUERY_HELPER(World, GeomInputs.GetGeometry(), FTransform(GeomInputs.GetGeometryOrientation() ? *GeomInputs.GetGeometryOrientation() : FQuat::Identity, Start), End, TraceChannel, Params, ResponseParams, ObjectParams, Traits::IsSweep() ? EChaosVDSceneQueryType::Sweep : EChaosVDSceneQueryType::RayCast, static_cast<EChaosVDSceneQueryMode>(Traits::SingleMultiOrTest), bIsRetryQuery);
+		bReturnResult = TSceneCastCommonImpWithRetryRequest<Traits, TGeomInputs, TAccelContainer>(World, OutHits, GeomInputs, Start, End, TraceChannel, Params, ResponseParams, ObjectParams, AccelContainer, bRequestRetry, RetryParams);
+	}
+
 	int InfiniteLoopProtection = 10;
 	while (bRequestRetry && InfiniteLoopProtection > 0)
 	{
+		CVD_TRACE_SCOPED_SCENE_QUERY_HELPER(World, GeomInputs.GetGeometry(), FTransform(GeomInputs.GetGeometryOrientation() ? *GeomInputs.GetGeometryOrientation() : FQuat::Identity, Start), End, TraceChannel, Params, ResponseParams, ObjectParams, Traits::IsSweep() ? EChaosVDSceneQueryType::Sweep : EChaosVDSceneQueryType::RayCast, static_cast<EChaosVDSceneQueryMode>(Traits::SingleMultiOrTest), bRequestRetry);
 		bReturnResult = TSceneCastCommonImpWithRetryRequest<Traits, TGeomInputs, TAccelContainer>(World, OutHits, GeomInputs, Start, End, TraceChannel, RetryParams, ResponseParams, ObjectParams, AccelContainer, bRequestRetry, RetryParams);
 		InfiniteLoopProtection--;
 	}
@@ -921,6 +932,9 @@ bool GeomOverlapMultiImp(const UWorld* World, const FPhysicsGeometry& Geom, cons
 	const ECollisionShapeType GeomType = GetType(Geom);
 	if (GeomType == ECollisionShapeType::Sphere || GeomType == ECollisionShapeType::Capsule || GeomType == ECollisionShapeType::Box || GeomType == ECollisionShapeType::Convex)
 	{
+		constexpr bool bIsRetryQuery = false;
+		CVD_TRACE_SCOPED_SCENE_QUERY_HELPER(World, &Geom, GeomPose, FVector::ZeroVector, TraceChannel, Params, ResponseParams, ObjectParams, EChaosVDSceneQueryType::Overlap, (InfoType == EQueryInfo::GatherAll) ? EChaosVDSceneQueryMode::Multi : EChaosVDSceneQueryMode::Test, bIsRetryQuery);
+
 		// Create filter data used to filter collisions
 		FCollisionFilterData Filter = CreateQueryFilterData(TraceChannel, Params.bTraceComplex, ResponseParams.CollisionResponse, Params, ObjectParams, InfoType != EQueryInfo::IsAnything);
 		FCollisionQueryFilterCallback QueryCallback(Params, false);
diff --git a/Engine/Source/Runtime/Experimental/Chaos/Private/ChaosVisualDebugger/ChaosVDContextProvider.cpp b/Engine/Source/Runtime/Experimental/Chaos/Private/ChaosVisualDebugger/ChaosVDContextProvider.cpp
--- a/Engine/Source/Runtime/Experimental/Chaos/Private/ChaosVisualDebugger/ChaosVDContextProvider.cpp
+++ b/Engine/Source/Runtime/Experimental/Chaos/Private/ChaosVisualDebugger/ChaosVDContextProvider.cpp
@@ -2,6 +2,8 @@
 
 #include "ChaosVisualDebugger/ChaosVDContextProvider.h"
 
+#include "ChaosVisualDebugger/ChaosVisualDebuggerTrace.h"
+
 #if WITH_CHAOS_VISUAL_DEBUGGER
 bool FChaosVDThreadContext::GetCurrentContext(FChaosVDContext& OutContext)
 {
@@ -14,11 +16,13 @@ bool FChaosVDThreadContext::GetCurrentContext(FChaosVDContext& OutContext)
 	return false;
 }
 
-const FChaosVDContext* FChaosVDThreadContext::GetCurrentContext()
+const FChaosVDContext* FChaosVDThreadContext::GetCurrentContext(EChaosVDContextType Type)
 {
 	if (LocalContextStack.Num() > 0)
 	{
-		return &LocalContextStack.Top();
+		const FChaosVDContext* Context = &LocalContextStack.Top();
+
+		return Context && Context->Type == static_cast<int32>(Type) ? Context : nullptr;
 	}
 	else
 	{
@@ -26,13 +30,26 @@ const FChaosVDContext* FChaosVDThreadContext::GetCurrentContext()
 	}
 }
 
+const FChaosVDContext* FChaosVDThreadContext::GetCurrentContext()
+{
+	return LocalContextStack.Num() > 0 ? &LocalContextStack.Top() : nullptr;
+}
+
 void FChaosVDThreadContext::PushContext(const FChaosVDContext& InContext)
 {
+	if (!FChaosVisualDebuggerTrace::IsTracing())
+	{
+		return;
+	}
+
 	LocalContextStack.Add(InContext);
 }
 
 void FChaosVDThreadContext::PopContext()
 {
-	LocalContextStack.Pop();
+	if (LocalContextStack.Num() > 0)
+	{
+		LocalContextStack.Pop();
+	}
 }
 #endif //WITH_CHAOS_VISUAL_DEBUGGER
diff --git a/Engine/Source/Runtime/Experimental/Chaos/Private/ChaosVisualDebugger/ChaosVDDataWrapperUtils.cpp b/Engine/Source/Runtime/Experimental/Chaos/Private/ChaosVisualDebugger/ChaosVDDataWrapperUtils.cpp
--- a/Engine/Source/Runtime/Experimental/Chaos/Private/ChaosVisualDebugger/ChaosVDDataWrapperUtils.cpp
+++ b/Engine/Source/Runtime/Experimental/Chaos/Private/ChaosVisualDebugger/ChaosVDDataWrapperUtils.cpp
@@ -188,4 +188,24 @@ FChaosVDParticlePairMidPhase FChaosVDDataWrapperUtils::BuildMidPhaseDataWrapperF
 
 	return MoveTemp(WrappedMidPhaseData);
 }
+
+void FChaosVDDataWrapperUtils::CopyShapeDataToWrapper(const Chaos::FShapeInstancePtr& ShapeDataPtr, FChaosVDShapeCollisionData& OutCopyTo)
+{
+	const Chaos::FCollisionData& CollisionData = ShapeDataPtr->GetCollisionData();
+
+	OutCopyTo.bQueryCollision = CollisionData.bQueryCollision;
+	OutCopyTo.bIsProbe = CollisionData.bIsProbe;
+	OutCopyTo.bSimCollision = CollisionData.bSimCollision;
+	OutCopyTo.CollisionTraceType = static_cast<EChaosVDCollisionTraceFlag>(CollisionData.CollisionTraceType);
+
+	OutCopyTo.SimData.Word0 = CollisionData.SimData.Word0;
+	OutCopyTo.SimData.Word1 = CollisionData.SimData.Word1;
+	OutCopyTo.SimData.Word2 = CollisionData.SimData.Word2;
+	OutCopyTo.SimData.Word3 = CollisionData.SimData.Word3;
+	
+	OutCopyTo.QueryData.Word0 = CollisionData.QueryData.Word0;
+	OutCopyTo.QueryData.Word1 = CollisionData.QueryData.Word1;
+	OutCopyTo.QueryData.Word2 = CollisionData.QueryData.Word2;
+	OutCopyTo.QueryData.Word3 = CollisionData.QueryData.Word3;
+}
 #endif //WITH_CHAOS_VISUAL_DEBUGGER
diff --git a/Engine/Source/Runtime/Experimental/Chaos/Private/ChaosVisualDebugger/ChaosVisualDebuggerTrace.cpp b/Engine/Source/Runtime/Experimental/Chaos/Private/ChaosVisualDebugger/ChaosVisualDebuggerTrace.cpp
--- a/Engine/Source/Runtime/Experimental/Chaos/Private/ChaosVisualDebugger/ChaosVisualDebuggerTrace.cpp
+++ b/Engine/Source/Runtime/Experimental/Chaos/Private/ChaosVisualDebugger/ChaosVisualDebuggerTrace.cpp
@@ -1,4 +1,5 @@
 // Copyright Epic Games, Inc. All Rights Reserved.
+
 #include "ChaosVisualDebugger/ChaosVisualDebuggerTrace.h"
 
 #if CHAOS_VISUAL_DEBUGGER_ENABLED
@@ -13,6 +14,7 @@
 #include "DataWrappers/ChaosVDCollisionDataWrappers.h"
 #include "DataWrappers/ChaosVDImplicitObjectDataWrapper.h"
 #include "DataWrappers/ChaosVDParticleDataWrapper.h"
+#include "DataWrappers/ChaosVDQueryDataWrappers.h"
 #include "HAL/CriticalSection.h"
 #include "Misc/ScopeRWLock.h"
 #include "Serialization/MemoryWriter.h"
@@ -110,7 +112,7 @@ void FChaosVisualDebuggerTrace::TraceParticle(const Chaos::FGeometryParticleHand
 		return;
 	}
 
-	const FChaosVDContext* CVDContextData = FChaosVDThreadContext::Get().GetCurrentContext();
+	const FChaosVDContext* CVDContextData = FChaosVDThreadContext::Get().GetCurrentContext(EChaosVDContextType::Solver);
 	if (!ensure(CVDContextData))
 	{
 		return;
@@ -146,13 +148,8 @@ void FChaosVisualDebuggerTrace::TraceParticle(Chaos::FGeometryParticleHandle* Pa
 		
 		for (const Chaos::FShapeInstancePtr& ShapeData : ShapesInstancesArray)
 		{
-			const Chaos::FCollisionData& CollisionData = ShapeData->GetCollisionData();
 			FChaosVDShapeCollisionData CVDCollisionData;
-			CVDCollisionData.bQueryCollision = CollisionData.bQueryCollision;
-			CVDCollisionData.bIsProbe = CollisionData.bIsProbe;
-			CVDCollisionData.bSimCollision = CollisionData.bSimCollision;
-			CVDCollisionData.CollisionTraceType = static_cast<EChaosVDCollisionTraceFlag>(CollisionData.CollisionTraceType);
-
+			FChaosVDDataWrapperUtils::CopyShapeDataToWrapper(ShapeData, CVDCollisionData);
 			ParticleDataWrapper.CollisionDataPerShape.Add(MoveTemp(CVDCollisionData));
 		}
 		
@@ -175,7 +172,7 @@ void FChaosVisualDebuggerTrace::TraceParticles(const Chaos::TGeometryParticleHan
 		return;
 	}
 
-	const FChaosVDContext* CVDContextData = FChaosVDThreadContext::Get().GetCurrentContext();
+	const FChaosVDContext* CVDContextData = FChaosVDThreadContext::Get().GetCurrentContext(EChaosVDContextType::Solver);
 	if (!ensure(CVDContextData))
 	{
 		return;
@@ -203,7 +200,7 @@ void FChaosVisualDebuggerTrace::TraceParticleDestroyed(const Chaos::FGeometryPar
 
 	GeometryTracerObject.RemoveCachedGeometryHash(ParticleHandle->GetGeometry());
 	
-	const FChaosVDContext* CVDContextData = FChaosVDThreadContext::Get().GetCurrentContext();
+	const FChaosVDContext* CVDContextData = FChaosVDThreadContext::Get().GetCurrentContext(EChaosVDContextType::Solver);
 	if (!ensure(CVDContextData))
 	{
 		return;
@@ -222,7 +219,7 @@ void FChaosVisualDebuggerTrace::TraceParticlesSoA(const Chaos::FPBDRigidsSOAs& P
 		return;
 	}
 
-	const FChaosVDContext* CVDContextData = FChaosVDThreadContext::Get().GetCurrentContext();
+	const FChaosVDContext* CVDContextData = FChaosVDThreadContext::Get().GetCurrentContext(EChaosVDContextType::Solver);
 	if (!ensure(CVDContextData))
 	{
 		return;
@@ -268,7 +265,7 @@ void FChaosVisualDebuggerTrace::TraceMidPhase(const Chaos::FParticlePairMidPhase
 		return;
 	}
 
-	const FChaosVDContext* CVDContextData = FChaosVDThreadContext::Get().GetCurrentContext();
+	const FChaosVDContext* CVDContextData = FChaosVDThreadContext::Get().GetCurrentContext(EChaosVDContextType::Solver);
 	if (!ensure(CVDContextData))
 	{
 		return;
@@ -298,7 +295,7 @@ void FChaosVisualDebuggerTrace::TraceMidPhasesFromCollisionConstraints(Chaos::FP
 		return;
 	}
 
-	const FChaosVDContext* CVDContextData = FChaosVDThreadContext::Get().GetCurrentContext();
+	const FChaosVDContext* CVDContextData = FChaosVDThreadContext::Get().GetCurrentContext(EChaosVDContextType::Solver);
 	if (!ensure(CVDContextData))
 	{
 		return;
@@ -320,7 +317,7 @@ void FChaosVisualDebuggerTrace::TraceCollisionConstraint(const Chaos::FPBDCollis
 		return;
 	}
 
-	const FChaosVDContext* CVDContextData = FChaosVDThreadContext::Get().GetCurrentContext();
+	const FChaosVDContext* CVDContextData = FChaosVDThreadContext::Get().GetCurrentContext(EChaosVDContextType::Solver);
 	if (!ensure(CVDContextData))
 	{
 		return;
@@ -344,7 +341,7 @@ void FChaosVisualDebuggerTrace::TraceCollisionConstraintView(TArrayView<Chaos::F
 		return;
 	}
 
-	const FChaosVDContext* CVDContextData = FChaosVDThreadContext::Get().GetCurrentContext();
+	const FChaosVDContext* CVDContextData = FChaosVDThreadContext::Get().GetCurrentContext(EChaosVDContextType::Solver);
 	if (!ensure(CVDContextData))
 	{
 		return;
@@ -369,6 +366,11 @@ void FChaosVisualDebuggerTrace::TraceSolverFrameStart(const FChaosVDContext& Con
 		return;
 	}
 
+	if (!ensure(ContextData.Type == static_cast<int32>(EChaosVDContextType::Solver)))
+	{
+		return;
+	}
+
 	FChaosVDThreadContext::Get().PushContext(ContextData);
 
 	// Check if we need to do a full capture for this solver, and setup accordingly
@@ -416,7 +418,7 @@ void FChaosVisualDebuggerTrace::TraceSolverStepStart(FStringView StepName)
 		return;
 	}
 
-	const FChaosVDContext* CVDContextData = FChaosVDThreadContext::Get().GetCurrentContext();
+	const FChaosVDContext* CVDContextData = FChaosVDThreadContext::Get().GetCurrentContext(EChaosVDContextType::Solver);
 	if (!ensure(CVDContextData))
 	{
 		return;
@@ -435,7 +437,7 @@ void FChaosVisualDebuggerTrace::TraceSolverStepEnd()
 		return;
 	}
 
-	const FChaosVDContext* CVDContextData = FChaosVDThreadContext::Get().GetCurrentContext();
+	const FChaosVDContext* CVDContextData = FChaosVDThreadContext::Get().GetCurrentContext(EChaosVDContextType::Solver);
 	if (!ensure(CVDContextData))
 	{
 		return;
@@ -453,7 +455,7 @@ void FChaosVisualDebuggerTrace::TraceSolverSimulationSpace(const Chaos::FRigidTr
 		return;
 	}
 
-	const FChaosVDContext* CVDContextData = FChaosVDThreadContext::Get().GetCurrentContext();
+	const FChaosVDContext* CVDContextData = FChaosVDThreadContext::Get().GetCurrentContext(EChaosVDContextType::Solver);
 	if (!ensure(CVDContextData))
 	{
 		return;
@@ -598,6 +600,88 @@ void FChaosVisualDebuggerTrace::TraceNonSolverTransform(const FTransform& InTran
 		<< ChaosVDNonSolverTransform.DebugName(DebugNameID.GetData(), DebugNameID.Len());
 }
 
+void FChaosVisualDebuggerTrace::TraceSceneQueryStart(const Chaos::FImplicitObject* InputGeometry, const FQuat& GeometryOrientation,  const FVector& Start, const FVector& End, ECollisionChannel TraceChannel, FChaosVDCollisionQueryParams&& Params, FChaosVDCollisionResponseParams&& ResponseParams, FChaosVDCollisionObjectQueryParams&& ObjectParams, EChaosVDSceneQueryType QueryType, EChaosVDSceneQueryMode QueryMode, int32 SolverID, bool bIsRetry)
+{
+	if (!IsTracing())
+	{
+		return;
+	}
+
+	const FChaosVDContext* CVDContextData = FChaosVDThreadContext::Get().GetCurrentContext();
+	const bool bIsQueryContext = CVDContextData && (CVDContextData->Type == static_cast<int32>(EChaosVDContextType::Query) ||  CVDContextData->Type == static_cast<int32>(EChaosVDContextType::SubTraceQuery));
+
+	if (!ensure(bIsQueryContext))
+	{
+		return;
+	}
+
+	FChaosVDQueryDataWrapper WrappedQueryData;
+
+	if (InputGeometry)
+	{
+		const uint32 GeometryHash = GeometryTracerObject.GetGeometryHashForImplicit(InputGeometry);
+		TraceImplicitObject({ GeometryHash, const_cast<Chaos::FImplicitObject*>(InputGeometry) });
+	
+		WrappedQueryData.InputGeometryKey = GeometryHash;
+	}
+	
+	WrappedQueryData.ID = CVDContextData->Id;
+	WrappedQueryData.ParentQueryID = CVDContextData->OwnerID;
+	WrappedQueryData.WorldSolverID = SolverID;
+	WrappedQueryData.bIsRetryQuery = bIsRetry;
+	
+	WrappedQueryData.GeometryOrientation = GeometryOrientation;
+
+	WrappedQueryData.CollisionChannel = TraceChannel;
+	WrappedQueryData.StartLocation = Start;
+	WrappedQueryData.EndLocation = End;
+
+	WrappedQueryData.CollisionQueryParams = MoveTemp(Params);
+	WrappedQueryData.CollisionResponseParams = MoveTemp(ResponseParams);
+	WrappedQueryData.CollisionObjectQueryParams = MoveTemp(ObjectParams);
+
+	WrappedQueryData.Mode = QueryMode;
+	WrappedQueryData.Type = QueryType;
+
+	FChaosVDScopedTLSBufferAccessor TLSDataBuffer;
+
+	FMemoryWriter MemWriterAr(TLSDataBuffer.BufferRef);
+	MemWriterAr.SetShouldSkipUpdateCustomVersion(true);
+	MemWriterAr.SetUseUnversionedPropertySerialization(true);
+
+	WrappedQueryData.Serialize(MemWriterAr);
+
+	TraceBinaryData(TLSDataBuffer.BufferRef, FChaosVDQueryDataWrapper::WrapperTypeName);
+}
+
+void FChaosVisualDebuggerTrace::TraceSceneQueryVisit(FChaosVDQueryVisitStep&& InQueryVisitData)
+{
+	if (!IsTracing())
+	{
+		return;
+	}
+
+	const FChaosVDContext* CVDContextData = FChaosVDThreadContext::Get().GetCurrentContext();
+	const bool bIsQueryContext = CVDContextData && (CVDContextData->Type == static_cast<int32>(EChaosVDContextType::Query) ||  CVDContextData->Type == static_cast<int32>(EChaosVDContextType::SubTraceQuery));
+
+	if (!ensure(bIsQueryContext))
+	{
+		return;
+	}
+
+	InQueryVisitData.OwningQueryID = CVDContextData->Id;
+
+	FChaosVDScopedTLSBufferAccessor TLSDataBuffer;
+
+	FMemoryWriter MemWriterAr(TLSDataBuffer.BufferRef);
+	MemWriterAr.SetShouldSkipUpdateCustomVersion(true);
+	MemWriterAr.SetUseUnversionedPropertySerialization(true);
+
+	InQueryVisitData.Serialize(MemWriterAr);
+
+	TraceBinaryData(TLSDataBuffer.BufferRef, FChaosVDQueryVisitStep::WrapperTypeName);
+}
+
 bool FChaosVisualDebuggerTrace::IsTracing()
 {
 	return bIsTracing;
diff --git a/Engine/Source/Runtime/Experimental/Chaos/Public/ChaosVisualDebugger/ChaosVDContextProvider.h b/Engine/Source/Runtime/Experimental/Chaos/Public/ChaosVisualDebugger/ChaosVDContextProvider.h
--- a/Engine/Source/Runtime/Experimental/Chaos/Public/ChaosVisualDebugger/ChaosVDContextProvider.h
+++ b/Engine/Source/Runtime/Experimental/Chaos/Public/ChaosVisualDebugger/ChaosVDContextProvider.h
@@ -10,11 +10,20 @@
 
 #include "ChaosVDRuntimeModule.h"
 
+enum class EChaosVDContextType : int32
+{
+	Any,
+	Solver,
+	Query,
+	SubTraceQuery,
+};
+
 /** Chaos Visual Debugger data used to context for logging or debugging purposes */
 struct FChaosVDContext
 {
 	int32 OwnerID = INDEX_NONE;
 	int32 Id = INDEX_NONE;
+	int32 Type = INDEX_NONE;
 };
 
 /** Singleton class that manages the thread local storage used to store CVD Context data */
@@ -35,6 +44,13 @@ public:
 	 */
 	bool GetCurrentContext(FChaosVDContext& OutContext);
 	
+	/** Gets the current CVD context data -
+	 * Don't use of a function that will recursively push new context data as it might invalidate the pointer
+	 * @param Type type of the context we want to get
+	 * @return Ptr to the Current CVD context data if it matches the provided type
+	 */
+	const FChaosVDContext* GetCurrentContext(EChaosVDContextType Type);
+	
 	/** Gets the current CVD context data -
 	 * Don't use of a function that will recursively push new context data as it might invalidate the pointer
 	 * @return Ptr to the Current CVD context data
diff --git a/Engine/Source/Runtime/Experimental/Chaos/Public/ChaosVisualDebugger/ChaosVDDataWrapperUtils.h b/Engine/Source/Runtime/Experimental/Chaos/Public/ChaosVisualDebugger/ChaosVDDataWrapperUtils.h
--- a/Engine/Source/Runtime/Experimental/Chaos/Public/ChaosVisualDebugger/ChaosVDDataWrapperUtils.h
+++ b/Engine/Source/Runtime/Experimental/Chaos/Public/ChaosVisualDebugger/ChaosVDDataWrapperUtils.h
@@ -1,6 +1,7 @@
 ﻿// Copyright Epic Games, Inc. All Rights Reserved.
 
 #pragma once
+#include "Chaos/ShapeInstanceFwd.h"
 
 #if WITH_CHAOS_VISUAL_DEBUGGER
 #include "Chaos/Collision/ContactPoint.h"
@@ -52,6 +53,8 @@ private:
 	/** Converts a Chaos::FVec3 to a FVector. It is worth notice that FVector is double precision and FVec3 is single */
 	static FVector ConvertToFVector(const Chaos::FVec3f& VectorRef) { return FVector(VectorRef); }
 
+	static void CopyShapeDataToWrapper(const Chaos::FShapeInstancePtr& ShapeDataPtr, FChaosVDShapeCollisionData& OutCopyTo);
+
 	friend FChaosVisualDebuggerTrace;
 };
 #endif //WITH_CHAOS_VISUAL_DEBUGGER
diff --git a/Engine/Source/Runtime/Experimental/Chaos/Public/ChaosVisualDebugger/ChaosVisualDebuggerTrace.h b/Engine/Source/Runtime/Experimental/Chaos/Public/ChaosVisualDebugger/ChaosVisualDebuggerTrace.h
--- a/Engine/Source/Runtime/Experimental/Chaos/Public/ChaosVisualDebugger/ChaosVisualDebuggerTrace.h
+++ b/Engine/Source/Runtime/Experimental/Chaos/Public/ChaosVisualDebugger/ChaosVisualDebuggerTrace.h
@@ -95,6 +95,18 @@
 	#define CVD_TRACE_INVALIDATE_CACHED_GEOMETRY(ImplicitObjectPtr)
 #endif
 
+#ifndef CVD_TRACE_SCOPED_SCENE_QUERY
+	#define CVD_TRACE_SCENE_QUERY_START(InputGeometry, GeometryOrientation, Start, End, TraceChannel, Params, ResponseParams, ObjectParams, QueryType, QueryMode, SolverID)
+#endif
+
+#ifndef CVD_TRACE_SCOPED_SCENE_QUERY_VISIT
+	#define CVD_TRACE_SCOPED_SCENE_QUERY_VISIT(InQueryVisitData)
+#endif
+
+#ifndef CVD_TRACE_SCENE_QUERY_VISIT
+	#define CVD_TRACE_SCENE_QUERY_VISIT(InQueryVisitData)
+#endif
+
 #else
 
 #include "ChaosVDRuntimeModule.h"
@@ -319,7 +331,32 @@ UE_TRACE_EVENT_END()
 	FChaosVisualDebuggerTrace::InvalidateGeometryFromCache(ImplicitObjectPtr);
 #endif
 
+#ifndef CVD_TRACE_SCOPED_SCENE_QUERY
+	#define CVD_TRACE_SCENE_QUERY_START(InputGeometry, GeometryOrientation, Start, End, TraceChannel, Params, ResponseParams, ObjectParams, QueryType, QueryMode, SolverID, bIsRetry) \
+	FChaosVisualDebuggerTrace::TraceSceneQueryStart(InputGeometry, GeometryOrientation, Start, End, TraceChannel, Params, ResponseParams, ObjectParams, QueryType, QueryMode, SolverID, bIsRetry);
+#endif
+
+#ifndef CVD_TRACE_SCOPED_SCENE_QUERY_VISIT
+	#define CVD_TRACE_SCOPED_SCENE_QUERY_VISIT(InQueryVisitData) \
+	FChaosVDScopeSceneQueryVisit CVDQueryVisit(InQueryVisitData);
+#endif
+
+#ifndef CVD_TRACE_SCENE_QUERY_VISIT
+	#define CVD_TRACE_SCENE_QUERY_VISIT(InQueryVisitData) \
+	FChaosVisualDebuggerTrace::TraceSceneQueryVisit(InQueryVisitData);
+#endif
+
 struct FChaosVDContext;
+struct FChaosVDQueryVisitStep;
+struct FChaosVDCollisionResponseParams;
+struct FChaosVDCollisionObjectQueryParams;
+struct FChaosVDCollisionQueryParams;
+enum class EChaosVDSceneQueryMode;
+enum class EChaosVDSceneQueryType;
+struct FCollisionObjectQueryParams;
+struct FCollisionResponseParams;
+struct FCollisionQueryParams;
+enum ECollisionChannel : int;
 
 namespace Chaos
 {
@@ -442,6 +479,9 @@ public:
 	 */
 	static CHAOS_API void TraceNonSolverTransform(const FTransform& InTransform, FStringView DebugNameID);
 
+	static CHAOS_API void TraceSceneQueryStart(const Chaos::FImplicitObject* InputGeometry, const FQuat& GeometryOrientation, const FVector& Start, const FVector& End, ECollisionChannel TraceChannel, FChaosVDCollisionQueryParams&& Params, FChaosVDCollisionResponseParams&& ResponseParams, FChaosVDCollisionObjectQueryParams&& ObjectParams, EChaosVDSceneQueryType QueryType, EChaosVDSceneQueryMode QueryMode, int32 SolverID, bool bIsRetry);
+	static CHAOS_API void TraceSceneQueryVisit(FChaosVDQueryVisitStep&& InQueryVisitData);
+
 	/** Returns true if the provided solver ID needs a Full Capture */
 	static bool ShouldPerformFullCapture(int32 SolverID);
 
@@ -505,13 +545,13 @@ void FChaosVisualDebuggerTrace::TraceParticlesView(const Chaos::TParticleView<Pa
 		return;
 	}
 
-	const FChaosVDContext* CVDContextData = FChaosVDThreadContext::Get().GetCurrentContext();
+	const FChaosVDContext* CVDContextData = FChaosVDThreadContext::Get().GetCurrentContext(EChaosVDContextType::Solver);
 	if (!ensure(CVDContextData))
 	{
 		return;
 	}
 
-	FChaosVDContext CopyContext = *FChaosVDThreadContext::Get().GetCurrentContext();
+	FChaosVDContext CopyContext = *CVDContextData;
 	
 	ParticlesView.ParallelFor([CopyContext](auto& Particle, int32 Index)
 	{
@@ -564,4 +604,19 @@ struct FChaosVDScopeSolverFrame
 
 	T& SolverRef;
 };
+
+struct FChaosVDScopeSceneQueryVisit
+{
+	FChaosVDScopeSceneQueryVisit(FChaosVDQueryVisitStep& InVisitData) : VisitData(InVisitData)
+	{
+		
+	}
+
+	~FChaosVDScopeSceneQueryVisit()
+	{
+		CVD_TRACE_SCENE_QUERY_VISIT(MoveTemp(VisitData));
+	}
+
+	FChaosVDQueryVisitStep& VisitData;
+};
 #endif // CHAOS_VISUAL_DEBUGGER_ENABLED
diff --git a/Engine/Source/Runtime/Experimental/ChaosVDData/Public/ChaosVDRecording.h b/Engine/Source/Runtime/Experimental/ChaosVDData/Public/ChaosVDRecording.h
--- a/Engine/Source/Runtime/Experimental/ChaosVDData/Public/ChaosVDRecording.h
+++ b/Engine/Source/Runtime/Experimental/ChaosVDData/Public/ChaosVDRecording.h
@@ -10,6 +10,7 @@
 #include "Containers/UnrealString.h"
 #include "Chaos/ImplicitFwd.h"
 #include "Chaos/ImplicitObject.h"
+#include "DataWrappers/ChaosVDQueryDataWrappers.h"
 
 DECLARE_MULTICAST_DELEGATE(FChaosVDRecordingUpdated)
 DECLARE_MULTICAST_DELEGATE_TwoParams(FChaosVDGeometryDataLoaded, const Chaos::FConstImplicitObjectPtr&, const uint32 GeometryID)
@@ -87,6 +88,7 @@ struct FChaosVDGameFrameData
 
 	TMap<FName, FChaosVDTrackedLocation> RecordedNonSolverLocationsByID;
 	TMap<FName, FChaosVDTrackedTransform> RecordedNonSolverTransformsByID;
+	TMap<int32, TSharedPtr<FChaosVDQueryDataWrapper>> RecordedSceneQueries;
 };
 
 /**
diff --git a/Engine/Source/Runtime/Experimental/ChaosVisualDebugger/Private/DataWrappers/ChaosVDCollisionDataWrappers.cpp b/Engine/Source/Runtime/Experimental/ChaosVisualDebugger/Private/DataWrappers/ChaosVDCollisionDataWrappers.cpp
--- a/Engine/Source/Runtime/Experimental/ChaosVisualDebugger/Private/DataWrappers/ChaosVDCollisionDataWrappers.cpp
+++ b/Engine/Source/Runtime/Experimental/ChaosVisualDebugger/Private/DataWrappers/ChaosVDCollisionDataWrappers.cpp
@@ -117,6 +117,16 @@ bool FChaosVDParticlePairMidPhase::Serialize(FArchive& Ar)
 	return true;
 }
 
+bool FChaosVDCollisionFilterData::Serialize(FArchive& Ar)
+{
+	Ar << Word0;
+	Ar << Word1;
+	Ar << Word2;
+	Ar << Word3;
+
+	return !Ar.IsError();
+}
+
 bool FChaosVDShapeCollisionData::Serialize(FArchive& Ar)
 {
 	Ar << CollisionTraceType;
diff --git a/Engine/Source/Runtime/Experimental/ChaosVisualDebugger/Private/DataWrappers/ChaosVDQueryDataWrappers.cpp b/Engine/Source/Runtime/Experimental/ChaosVisualDebugger/Private/DataWrappers/ChaosVDQueryDataWrappers.cpp
new file mode 100644
--- /dev/null
+++ b/Engine/Source/Runtime/Experimental/ChaosVisualDebugger/Private/DataWrappers/ChaosVDQueryDataWrappers.cpp
@@ -0,0 +1,114 @@
+﻿// Copyright Epic Games, Inc. All Rights Reserved.
+
+#include "DataWrappers/ChaosVDQueryDataWrappers.h"
+
+bool FChaosVDCollisionResponseParams::Serialize(FArchive& Ar)
+{
+	Ar << FlagsPerChannel;
+	Ar << bHasValidData;
+
+	return !Ar.IsError();
+}
+
+bool FChaosVDCollisionObjectQueryParams::Serialize(FArchive& Ar)
+{
+	Ar << ObjectTypesToQuery;
+	Ar << IgnoreMask;
+
+	Ar << bHasValidData;
+
+	return !Ar.IsError();
+}
+
+bool FChaosVDCollisionQueryParams::Serialize(FArchive& Ar)
+{
+	Ar << TraceTag;
+	Ar << OwnerTag;
+	Ar << bTraceComplex;
+	Ar << bFindInitialOverlaps;
+	Ar << bReturnFaceIndex;
+	Ar << bReturnPhysicalMaterial;
+	Ar << bIgnoreBlocks;
+	Ar << bIgnoreTouches;
+	Ar << bSkipNarrowPhase;
+	Ar << bTraceIntoSubComponents;
+	Ar << bReplaceHitWithSubComponents;
+	Ar << IgnoreMask;
+	Ar << IgnoredComponentsIDs;
+	Ar << IgnoredActorsIDs;
+	Ar << IgnoredComponentsNames;
+	Ar << IgnoredActorsNames;
+
+	Ar << bHasValidData;
+	
+	return !Ar.IsError();
+}
+
+bool FChaosVDQueryFastData::Serialize(FArchive& Ar)
+{
+	Ar << Dir;
+	Ar << InvDir;
+	Ar << CurrentLength;
+	Ar << InvDir;
+
+	//TODO : Pack these flags
+	Ar << bParallel0;
+	Ar << bParallel1;
+	Ar << bParallel2;
+
+	Ar << bHasValidData;
+
+	return !Ar.IsError();
+}
+
+bool FChaosVDQueryHitData::Serialize(FArchive& Ar)
+{
+	Ar << Distance;
+	Ar << FaceIdx;
+	Ar << Flags;
+	Ar << WorldPosition;
+	Ar << WorldNormal;
+	Ar << FaceNormal;
+	Ar << bHasValidData;
+
+	return !Ar.IsError();
+}
+
+bool FChaosVDQueryVisitStep::Serialize(FArchive& Ar)
+{
+	Ar << OwningQueryID;
+	Ar << Type;
+	Ar << ShapeIndex;
+	Ar << ParticleIndex;
+	Ar << ParticleTransform;
+	Ar << HitType;
+	Ar << HitData;
+
+	Ar << QueryFastData;
+
+	Ar << bHasValidData;
+
+	return !Ar.IsError();
+}
+
+bool FChaosVDQueryDataWrapper::Serialize(FArchive& Ar)
+{
+	Ar << ID;
+	Ar << ParentQueryID;
+	Ar << WorldSolverID;
+	Ar << bIsRetryQuery;
+	Ar << InputGeometryKey;
+	Ar << GeometryOrientation;
+	Ar << Type;
+	Ar << Mode;
+	Ar << StartLocation;
+	Ar << EndLocation;
+	Ar << CollisionChannel;
+	Ar << CollisionQueryParams;
+	Ar << CollisionResponseParams;
+	Ar << CollisionObjectQueryParams;
+
+	// Hits and steps are intentionally not serialized as they are recorded as separated events, and reconstructed during trace analysis
+
+	return !Ar.IsError();
+}
diff --git a/Engine/Source/Runtime/Experimental/ChaosVisualDebugger/Public/DataWrappers/ChaosVDCollisionDataWrappers.h b/Engine/Source/Runtime/Experimental/ChaosVisualDebugger/Public/DataWrappers/ChaosVDCollisionDataWrappers.h
--- a/Engine/Source/Runtime/Experimental/ChaosVisualDebugger/Public/DataWrappers/ChaosVDCollisionDataWrappers.h
+++ b/Engine/Source/Runtime/Experimental/ChaosVisualDebugger/Public/DataWrappers/ChaosVDCollisionDataWrappers.h
@@ -395,6 +395,38 @@ enum class EChaosVDCollisionTraceFlag
 	MAX,
 };
 
+USTRUCT()
+struct CHAOSVDRUNTIME_API FChaosVDCollisionFilterData
+{
+	GENERATED_BODY()
+
+	UPROPERTY(VisibleAnywhere, Category=CollisionData)
+	uint32 Word0;
+	UPROPERTY(VisibleAnywhere, Category=CollisionData)
+	uint32 Word1;
+	UPROPERTY(VisibleAnywhere, Category=CollisionData)
+	uint32 Word2;
+	UPROPERTY(VisibleAnywhere, Category=CollisionData)
+	uint32 Word3;
+
+	bool Serialize(FArchive& Ar);
+};
+
+template<>
+struct TStructOpsTypeTraits<FChaosVDCollisionFilterData> : public TStructOpsTypeTraitsBase2<FChaosVDCollisionFilterData>
+{
+	enum
+	{
+		WithSerializer = true,
+	};
+};
+
+inline FArchive& operator<<(FArchive& Ar, FChaosVDCollisionFilterData& Data)
+{
+	Data.Serialize(Ar);
+	return Ar;
+}
+
 USTRUCT()
 struct CHAOSVDRUNTIME_API FChaosVDShapeCollisionData
 {
@@ -410,6 +442,12 @@ struct CHAOSVDRUNTIME_API FChaosVDShapeCollisionData
 	UPROPERTY(VisibleAnywhere, Category=CollisionData)
 	uint8 bIsProbe : 1 = false;
 
+	UPROPERTY(VisibleAnywhere, Category=FilterData)
+	FChaosVDCollisionFilterData QueryData;
+
+	UPROPERTY(VisibleAnywhere, Category=SimData)
+	FChaosVDCollisionFilterData SimData;
+
 	bool bIsComplex = false;
 
 	bool bIsValid = false;
diff --git a/Engine/Source/Runtime/Experimental/ChaosVisualDebugger/Public/DataWrappers/ChaosVDParticleDataWrapper.h b/Engine/Source/Runtime/Experimental/ChaosVisualDebugger/Public/DataWrappers/ChaosVDParticleDataWrapper.h
--- a/Engine/Source/Runtime/Experimental/ChaosVisualDebugger/Public/DataWrappers/ChaosVDParticleDataWrapper.h
+++ b/Engine/Source/Runtime/Experimental/ChaosVisualDebugger/Public/DataWrappers/ChaosVDParticleDataWrapper.h
@@ -42,20 +42,22 @@ enum class EChaosVDObjectStateType: int8
 
 /** Base struct that declares the interface to be used for any ParticleData Viewer */
 USTRUCT()
-struct FChaosVDParticleDataBase
+struct FChaosVDWrapperDataBase
 {
 	GENERATED_BODY()
-	virtual ~FChaosVDParticleDataBase() = default;
+	virtual ~FChaosVDWrapperDataBase() = default;
 
 	virtual bool HasValidData() const { return bHasValidData; }
 
+	void MarkAsValid() { bHasValidData = true; }
+
 protected:
 	UPROPERTY()
 	bool bHasValidData = false;
 };
 
 USTRUCT()
-struct FChaosVDFRigidParticleControlFlags : public FChaosVDParticleDataBase 
+struct FChaosVDFRigidParticleControlFlags : public FChaosVDWrapperDataBase 
 {
 	GENERATED_BODY()
 
@@ -113,7 +115,7 @@ struct TStructOpsTypeTraits<FChaosVDFRigidParticleControlFlags> : public TStruct
  * Used to be able to show the values in the editor and allow changes via the Property Editor.
  */
 USTRUCT()
-struct FChaosVDParticlePositionRotation : public FChaosVDParticleDataBase
+struct FChaosVDParticlePositionRotation : public FChaosVDWrapperDataBase
 {
 	GENERATED_BODY()
 
@@ -158,7 +160,7 @@ struct TStructOpsTypeTraits<FChaosVDParticlePositionRotation> : public TStructOp
  * Used to be able to show the values in the editor and allow changes via the Property Editor.
  */
 USTRUCT()
-struct FChaosVDParticleVelocities : public FChaosVDParticleDataBase
+struct FChaosVDParticleVelocities : public FChaosVDWrapperDataBase
 {
 	GENERATED_BODY()
 
@@ -202,7 +204,7 @@ struct TStructOpsTypeTraits<FChaosVDParticleVelocities> : public TStructOpsTypeT
  * Used to be able to show the values in the editor and allow changes via the Property Editor.
  */
 USTRUCT()
-struct FChaosVDParticleDynamics : public FChaosVDParticleDataBase
+struct FChaosVDParticleDynamics : public FChaosVDWrapperDataBase
 {
 	GENERATED_BODY()
 	
@@ -256,7 +258,7 @@ struct TStructOpsTypeTraits<FChaosVDParticleDynamics> : public TStructOpsTypeTra
  * Used to be able to show the values in the editor and allow changes via the Property Editor.
  */
 USTRUCT()
-struct FChaosVDParticleMassProps : public FChaosVDParticleDataBase
+struct FChaosVDParticleMassProps : public FChaosVDWrapperDataBase
 {
 	GENERATED_BODY()
 
@@ -317,7 +319,7 @@ struct TStructOpsTypeTraits<FChaosVDParticleMassProps> : public TStructOpsTypeTr
  * Used to be able to show the values in the editor and allow changes via the Property Editor.
  */
 USTRUCT()
-struct FChaosVDParticleDynamicMisc : public FChaosVDParticleDataBase
+struct FChaosVDParticleDynamicMisc : public FChaosVDWrapperDataBase
 {
 	GENERATED_BODY()
 
@@ -438,7 +440,7 @@ struct TStructOpsTypeTraits<FChaosVDConnectivityEdge> : public TStructOpsTypeTra
  * Used to be able to show the values in the editor and allow changes via the Property Editor.
  */
 USTRUCT()
-struct FChaosVDParticleCluster : public FChaosVDParticleDataBase
+struct FChaosVDParticleCluster : public FChaosVDWrapperDataBase
 {
 	GENERATED_BODY()
 
@@ -538,7 +540,7 @@ struct TStructOpsTypeTraits<FChaosVDParticleCluster> : public TStructOpsTypeTrai
  * Used to be able to show the values in the editor and allow changes via the Property Editor.
  */
 USTRUCT(DisplayName="Particle Data")
-struct FChaosVDParticleDataWrapper : public FChaosVDParticleDataBase
+struct FChaosVDParticleDataWrapper : public FChaosVDWrapperDataBase
 {
 	virtual ~FChaosVDParticleDataWrapper() override = default;
 
diff --git a/Engine/Source/Runtime/Experimental/ChaosVisualDebugger/Public/DataWrappers/ChaosVDQueryDataWrappers.h b/Engine/Source/Runtime/Experimental/ChaosVisualDebugger/Public/DataWrappers/ChaosVDQueryDataWrappers.h
new file mode 100644
--- /dev/null
+++ b/Engine/Source/Runtime/Experimental/ChaosVisualDebugger/Public/DataWrappers/ChaosVDQueryDataWrappers.h
@@ -0,0 +1,428 @@
+﻿// Copyright Epic Games, Inc. All Rights Reserved.
+
+#pragma once
+#include "ChaosVDParticleDataWrapper.h"
+#include "Math/BoxSphereBounds.h"
+
+#include "ChaosVDQueryDataWrappers.generated.h"
+
+UENUM()
+enum class EChaosVDSceneQueryType
+{
+	Invalid,
+	Sweep,
+	Overlap,
+	RayCast
+};
+
+UENUM()
+enum class EChaosVDSceneQueryMode
+{
+	Invalid,
+	Single,
+	Multi,
+	Test
+};
+
+USTRUCT()
+struct FChaosVDCollisionResponseParams : public FChaosVDWrapperDataBase
+{
+	GENERATED_BODY()
+
+	CHAOSVDRUNTIME_API bool Serialize(FArchive& Ar);
+
+	template <typename TOther>
+	void CopyFrom(const TOther& Other)
+	{
+		FlagsPerChannel = TArray(Other.CollisionResponse.EnumArray, UE_ARRAY_COUNT(Other.CollisionResponse.EnumArray));
+		bHasValidData = true;
+	}
+	
+	UPROPERTY(VisibleAnywhere, Category=QueryData)
+	TArray<uint8> FlagsPerChannel;
+};
+
+inline FArchive& operator<<(FArchive& Ar, FChaosVDCollisionResponseParams& Data)
+{
+	Data.Serialize(Ar);
+	return Ar;
+}
+
+template<>
+struct TStructOpsTypeTraits<FChaosVDCollisionResponseParams> : public TStructOpsTypeTraitsBase2<FChaosVDCollisionResponseParams>
+{
+	enum
+	{
+		WithSerializer = true,
+	};
+};
+
+USTRUCT()
+struct FChaosVDCollisionObjectQueryParams : public FChaosVDWrapperDataBase
+{
+	GENERATED_BODY()
+
+	CHAOSVDRUNTIME_API bool Serialize(FArchive& Ar);
+
+	template <typename TOther>
+	void CopyFrom(const TOther& Other)
+	{
+		ObjectTypesToQuery = Other.ObjectTypesToQuery;
+		IgnoreMask = Other.IgnoreMask;
+		bHasValidData = true;
+	}
+	
+	UPROPERTY(VisibleAnywhere, Category=QueryData)
+	uint8 ObjectTypesToQuery = 0;
+
+	UPROPERTY(VisibleAnywhere, Category=QueryData)
+	uint8 IgnoreMask = 0;
+};
+
+inline FArchive& operator<<(FArchive& Ar, FChaosVDCollisionObjectQueryParams& Data)
+{
+	Data.Serialize(Ar);
+	return Ar;
+}
+
+template<>
+struct TStructOpsTypeTraits<FChaosVDCollisionObjectQueryParams> : public TStructOpsTypeTraitsBase2<FChaosVDCollisionObjectQueryParams>
+{
+	enum
+	{
+		WithSerializer = true,
+	};
+};
+
+USTRUCT()
+struct FChaosVDCollisionQueryParams : public FChaosVDWrapperDataBase
+{
+	GENERATED_BODY()
+
+	CHAOSVDRUNTIME_API bool Serialize(FArchive& Ar);
+
+	template <typename TOther>
+	void CopyFrom(const TOther& Other)
+	{
+		TraceTag = Other.TraceTag;
+		OwnerTag = Other.OwnerTag;
+		bTraceComplex = Other.bTraceComplex;
+		bFindInitialOverlaps = Other.bFindInitialOverlaps;
+		bReturnFaceIndex = Other.bReturnFaceIndex;
+		bReturnPhysicalMaterial = Other.bReturnPhysicalMaterial;
+		bIgnoreBlocks = Other.bIgnoreBlocks;
+		bIgnoreTouches = Other.bIgnoreTouches;
+		bSkipNarrowPhase = Other.bSkipNarrowPhase;
+		bTraceIntoSubComponents = Other.bTraceIntoSubComponents;
+		bReplaceHitWithSubComponents = Other.bReplaceHitWithSubComponents;
+		IgnoreMask = Other.IgnoreMask;
+
+		IgnoredComponentsIDs = Other.GetIgnoredComponents();
+		IgnoredActorsIDs = Other.GetIgnoredActors();
+
+		bHasValidData = true;
+	}
+	
+	/** Tag used to provide extra information or filtering for debugging of the trace (e.g. Collision Analyzer) */
+	UPROPERTY(VisibleAnywhere, Category=QueryData)
+	FName TraceTag;
+
+	/** Tag used to indicate an owner for this trace */
+	UPROPERTY(VisibleAnywhere, Category=QueryData)
+	FName OwnerTag;
+
+	/** Whether we should trace against complex collision */
+	UPROPERTY(VisibleAnywhere, Category=QueryData)
+	bool bTraceComplex = false;
+
+	/** Whether we want to find out initial overlap or not. If true, it will return if this was initial overlap. */
+	UPROPERTY(VisibleAnywhere, Category=QueryData)
+	bool bFindInitialOverlaps = false;
+
+	/** Whether we want to return the triangle face index for complex static mesh traces */
+	UPROPERTY(VisibleAnywhere, Category=QueryData)
+	bool bReturnFaceIndex = false;
+
+	/** Whether we want to include the physical material in the results. */
+	UPROPERTY(VisibleAnywhere, Category=QueryData)
+	bool bReturnPhysicalMaterial = false;
+
+	/** Whether to ignore blocking results. */
+	UPROPERTY(VisibleAnywhere, Category=QueryData)
+	bool bIgnoreBlocks = false;
+
+	/** Whether to ignore touch/overlap results. */
+	UPROPERTY(VisibleAnywhere, Category=QueryData)
+	bool bIgnoreTouches = false;
+
+	/** Whether to skip narrow phase checks (only for overlaps). */
+	UPROPERTY(VisibleAnywhere, Category=QueryData)
+	bool bSkipNarrowPhase = false;
+
+	/** Whether to ignore traces to the cluster union and trace against its children instead. */
+	UPROPERTY(VisibleAnywhere, Category=QueryData)
+	bool bTraceIntoSubComponents = false;
+
+	/** If bTraceIntoSubComponents is true, whether to replace the hit of the cluster union with its children instead. */
+	UPROPERTY(VisibleAnywhere, Category=QueryData)
+	bool bReplaceHitWithSubComponents = false;
+
+	/** Extra filtering done on the query. See declaration for filtering logic */
+	UPROPERTY(VisibleAnywhere, Category=QueryData)
+	uint8 IgnoreMask = 0;
+
+	TArray<uint32> IgnoredComponentsIDs;
+
+	TArray<uint32> IgnoredActorsIDs;
+
+	UPROPERTY(VisibleAnywhere, Category=QueryData)
+	TArray<FName> IgnoredActorsNames;
+	
+	UPROPERTY(VisibleAnywhere, Category=QueryData)
+	TArray<FName> IgnoredComponentsNames;
+};
+
+inline FArchive& operator<<(FArchive& Ar, FChaosVDCollisionQueryParams& Data)
+{
+	Data.Serialize(Ar);
+	return Ar;
+}
+
+template<>
+struct TStructOpsTypeTraits<FChaosVDCollisionQueryParams> : public TStructOpsTypeTraitsBase2<FChaosVDCollisionQueryParams>
+{
+	enum
+	{
+		WithSerializer = true,
+	};
+};
+
+USTRUCT()
+struct FChaosVDQueryFastData : public FChaosVDWrapperDataBase
+{
+	GENERATED_BODY()
+
+	CHAOSVDRUNTIME_API bool Serialize(FArchive& Ar);
+
+	UPROPERTY(VisibleAnywhere, Category=QueryFastData)
+	FVector Dir = FVector::ZeroVector;
+	
+	UPROPERTY(VisibleAnywhere, Category=QueryFastData)
+	FVector InvDir = FVector::ZeroVector;
+	
+	UPROPERTY(VisibleAnywhere, Category=QueryFastData)
+	float CurrentLength = 0.0f;
+
+	UPROPERTY(VisibleAnywhere, Category=QueryFastData)
+	float InvCurrentLength = 0.0f;
+
+	UPROPERTY(VisibleAnywhere, Category=QueryFastData)
+	bool bParallel0 = false;
+
+	UPROPERTY(VisibleAnywhere, Category=QueryFastData)
+	bool bParallel1 = false;
+
+	UPROPERTY(VisibleAnywhere, Category=QueryFastData)
+	bool bParallel2 = false;
+};
+
+inline FArchive& operator<<(FArchive& Ar, FChaosVDQueryFastData& Data)
+{
+	Data.Serialize(Ar);
+	return Ar;
+}
+
+template<>
+struct TStructOpsTypeTraits<FChaosVDQueryFastData> : public TStructOpsTypeTraitsBase2<FChaosVDQueryFastData>
+{
+	enum
+	{
+		WithSerializer = true,
+	};
+};
+
+USTRUCT()
+struct FChaosVDQueryHitData : public FChaosVDWrapperDataBase
+{
+	GENERATED_BODY()
+
+	CHAOSVDRUNTIME_API bool Serialize(FArchive& Ar);
+
+	UPROPERTY(VisibleAnywhere, Category=QueryData)
+	float Distance = 0.0f;
+
+	UPROPERTY(VisibleAnywhere, Category=QueryData)
+	int32 FaceIdx = INDEX_NONE;
+
+	UPROPERTY(VisibleAnywhere, Category=QueryData)
+	uint16 Flags = 0;
+
+	UPROPERTY(VisibleAnywhere, Category=QueryData)
+	FVector WorldPosition = FVector::ZeroVector;
+
+	UPROPERTY(VisibleAnywhere, Category=QueryData)
+	FVector WorldNormal = FVector::ZeroVector;
+
+	UPROPERTY(VisibleAnywhere, Category=QueryData)
+	FVector FaceNormal = FVector::ZeroVector;
+};
+
+inline FArchive& operator<<(FArchive& Ar, FChaosVDQueryHitData& Data)
+{
+	Data.Serialize(Ar);
+	return Ar;
+}
+
+template<>
+struct TStructOpsTypeTraits<FChaosVDQueryHitData> : public TStructOpsTypeTraitsBase2<FChaosVDQueryHitData>
+{
+	enum
+	{
+		WithSerializer = true,
+	};
+};
+
+UENUM()
+enum class EChaosVDCollisionQueryHitType
+{
+	None = 0,
+	Touch = 1,
+	Block = 2
+};
+
+UENUM()
+enum class EChaosVDSceneQueryVisitorType
+{
+	Invalid,
+	BroadPhase,
+	NarrowPhase
+};
+
+USTRUCT()
+struct FChaosVDQueryVisitStep : public FChaosVDWrapperDataBase
+{
+	GENERATED_BODY()
+	
+	inline static FStringView WrapperTypeName = TEXT("FChaosVDQueryVisitStep");
+
+	CHAOSVDRUNTIME_API bool Serialize(FArchive& Ar);
+
+	int32 OwningQueryID = INDEX_NONE;
+
+	UPROPERTY(VisibleAnywhere, Category="SQ Visit Data")
+	EChaosVDSceneQueryVisitorType Type = EChaosVDSceneQueryVisitorType::Invalid;
+
+	UPROPERTY(VisibleAnywhere, Category="SQ Visit Data")
+	uint32 ShapeIndex = 0;
+	
+	UPROPERTY(VisibleAnywhere, Category="SQ Visit Data")
+	int32 ParticleIndex = INDEX_NONE;
+
+	UPROPERTY(VisibleAnywhere, Category="SQ Visit Data")
+	FTransform ParticleTransform;
+
+	UPROPERTY(VisibleAnywhere, Category="SQ Visit Data")
+	FChaosVDQueryFastData QueryFastData;
+	
+	UPROPERTY(VisibleAnywhere, Category="SQ Visit Data")
+	EChaosVDCollisionQueryHitType HitType = EChaosVDCollisionQueryHitType::None;
+
+	UPROPERTY(VisibleAnywhere, Category="SQ Visit Data")
+	FChaosVDQueryHitData HitData;
+
+	// Editor only properties
+
+	bool bIsSelectedInEditor = false;
+};
+
+inline FArchive& operator<<(FArchive& Ar, FChaosVDQueryVisitStep& Data)
+{
+	Data.Serialize(Ar);
+	return Ar;
+}
+
+template<>
+struct TStructOpsTypeTraits<FChaosVDQueryVisitStep> : public TStructOpsTypeTraitsBase2<FChaosVDQueryVisitStep>
+{
+	enum
+	{
+		WithSerializer = true,
+	};
+};
+
+USTRUCT()
+struct FChaosVDQueryDataWrapper
+{
+	GENERATED_BODY()
+
+	inline static FStringView WrapperTypeName = TEXT("FChaosVDQueryDataWrapper");
+
+	CHAOSVDRUNTIME_API bool Serialize(FArchive& Ar);
+
+	UPROPERTY(VisibleAnywhere, Category="CVD QueryData")
+	int32 ID = INDEX_NONE;
+
+	UPROPERTY(VisibleAnywhere, Category="CVD QueryData")
+	int32 ParentQueryID = INDEX_NONE;
+
+	UPROPERTY(VisibleAnywhere, Category="CVD QueryData")
+	int32 WorldSolverID = INDEX_NONE;
+
+	UPROPERTY(VisibleAnywhere, Category=RecordedQueryData)
+	bool bIsRetryQuery = false;
+
+	uint32 InputGeometryKey = 0;
+
+	UPROPERTY(VisibleAnywhere, Category=RecordedQueryData)
+	FQuat GeometryOrientation;
+
+	UPROPERTY(VisibleAnywhere, Category=RecordedQueryData)
+	EChaosVDSceneQueryType Type = EChaosVDSceneQueryType::Invalid;
+
+	UPROPERTY(VisibleAnywhere, Category=RecordedQueryData)
+	EChaosVDSceneQueryMode Mode = EChaosVDSceneQueryMode::Invalid;
+
+	UPROPERTY(VisibleAnywhere, Category=RecordedQueryData)
+	FVector StartLocation = FVector::ZeroVector;
+
+	UPROPERTY(VisibleAnywhere, Category=RecordedQueryData)
+	FVector EndLocation = FVector::ZeroVector;
+
+	UPROPERTY(VisibleAnywhere, Category=RecordedQueryData)
+	int32 CollisionChannel = INDEX_NONE;
+
+	UPROPERTY(VisibleAnywhere, Category=RecordedQueryData)
+	FChaosVDCollisionQueryParams CollisionQueryParams;
+
+	UPROPERTY(VisibleAnywhere, Category=RecordedQueryData)
+	FChaosVDCollisionResponseParams CollisionResponseParams;
+
+	UPROPERTY(VisibleAnywhere, Category=RecordedQueryData)
+	FChaosVDCollisionObjectQueryParams CollisionObjectQueryParams;
+
+	TArray<FChaosVDQueryVisitStep> SQVisitData;
+
+	UPROPERTY(VisibleAnywhere, Category=QueryData)
+	TArray<FChaosVDQueryVisitStep> Hits;
+
+	TArray<int32> SubQueriesIDs;
+
+	bool bIsSelectedInEditor = false;
+
+	int32 CurrentVisitIndex = 0;
+};
+
+inline FArchive& operator<<(FArchive& Ar, FChaosVDQueryDataWrapper& Data)
+{
+	Data.Serialize(Ar);
+	return Ar;
+}
+
+template<>
+struct TStructOpsTypeTraits<FChaosVDQueryDataWrapper> : public TStructOpsTypeTraitsBase2<FChaosVDQueryDataWrapper>
+{
+	enum
+	{
+		WithSerializer = true,
+	};
+};
diff --git a/Engine/Source/Runtime/PhysicsCore/Private/ChaosScene.cpp b/Engine/Source/Runtime/PhysicsCore/Private/ChaosScene.cpp
--- a/Engine/Source/Runtime/PhysicsCore/Private/ChaosScene.cpp
+++ b/Engine/Source/Runtime/PhysicsCore/Private/ChaosScene.cpp
@@ -82,6 +82,7 @@ FChaosScene::FChaosScene(
 #if WITH_CHAOS_VISUAL_DEBUGGER
 	SceneSolver->GetChaosVDContextData().OwnerID = GetChaosVDContextData().Id;
 	SceneSolver->GetChaosVDContextData().Id = FChaosVDRuntimeModule::Get().GenerateUniqueID();
+	SceneSolver->GetChaosVDContextData().Type = static_cast<int32>(EChaosVDContextType::Solver);
 #endif
 
 	SceneSolver->PhysSceneHack = this;
diff --git a/Engine/Source/Runtime/PhysicsCore/Public/ChaosVDSQVisitorHelpers.h b/Engine/Source/Runtime/PhysicsCore/Public/ChaosVDSQVisitorHelpers.h
new file mode 100644
--- /dev/null
+++ b/Engine/Source/Runtime/PhysicsCore/Public/ChaosVDSQVisitorHelpers.h
@@ -0,0 +1,93 @@
+﻿// Copyright Epic Games, Inc. All Rights Reserved.
+
+#pragma once
+#if WITH_CHAOS_VISUAL_DEBUGGER
+
+#include "CollisionQueryFilterCallbackCore.h"
+#include "Chaos/ISpatialAcceleration.h"
+#include "ChaosVisualDebugger/ChaosVisualDebuggerTrace.h"
+#include "DataWrappers/ChaosVDQueryDataWrappers.h"
+
+/** Initializes a CVD Wrapper for a SQ visit with the provider data, and traces it as soon it goes iut of scope if we are recording */
+#ifndef CVD_TRACE_SCOPED_SCENE_QUERY_VISIT_HELPER
+	#define CVD_TRACE_SCOPED_SCENE_QUERY_VISIT_HELPER(Type, ParticleTransform, ParticleIndex, ShapeIndex, CurData) \
+	FChaosVDQueryVisitStep CVDSQVisitStepData; \
+	CVD_TRACE_SCOPED_SCENE_QUERY_VISIT(CVDSQVisitStepData); \
+	Chaos::VisualDebugger::TraceHelpers::FillCVDQueryVisitData(CVDSQVisitStepData, Type, ParticleTransform, ParticleIndex, ShapeIndex, CurData);
+#endif
+
+/** Fill a hit data entry on the current CVD Wrapper for a SQ visit in scope if we are recording
+ * This needs to be called in the same scope of a CVD_TRACE_SCOPED_SCENE_QUERY_VISIT_HELPER call
+ */
+#ifndef CVD_FILL_HIT_DATA_HELPER
+	#define CVD_FILL_HIT_DATA_HELPER(Hit, HitType) \
+	Chaos::VisualDebugger::TraceHelpers::FillCVDHitDataHelper(Hit, HitType, CVDSQVisitStepData);
+#endif
+
+namespace Chaos::VisualDebugger::TraceHelpers
+{
+	template <class THitType>
+	void FillCVDHitDataHelper(THitType& Hit, ECollisionQueryHitType HitType, FChaosVDQueryVisitStep& InOutVisitData);
+	inline void FillCVDQueryVisitData(FChaosVDQueryVisitStep& InOutVisitData, EChaosVDSceneQueryVisitorType Type, const FTransform& ParticleTransform, int32 ParticleIndex, int32 ShapeIndex, const Chaos::FQueryFastData* CurData)
+	{
+		if (!FChaosVisualDebuggerTrace::IsTracing())
+		{
+			return;
+		}
+
+		InOutVisitData.Type = Type;
+		InOutVisitData.ParticleTransform = ParticleTransform;
+		InOutVisitData.ParticleIndex = ParticleIndex;
+		InOutVisitData.ShapeIndex = ShapeIndex;
+
+		if (CurData)
+		{
+			InOutVisitData.QueryFastData.Dir = CurData->Dir;
+			InOutVisitData.QueryFastData.InvDir = CurData->InvDir;
+			InOutVisitData.QueryFastData.CurrentLength = CurData->CurrentLength;
+			InOutVisitData.QueryFastData.InvCurrentLength = CurData->InvCurrentLength;
+			InOutVisitData.QueryFastData.bParallel0 = CurData->bParallel[0];
+			InOutVisitData.QueryFastData.bParallel1 = CurData->bParallel[1];
+			InOutVisitData.QueryFastData.bParallel2 = CurData->bParallel[2];
+			InOutVisitData.QueryFastData.MarkAsValid();
+		}
+
+		InOutVisitData.MarkAsValid();
+	}
+
+	template <class THitType>
+	void FillCVDHitDataHelper(THitType& Hit, ECollisionQueryHitType HitType, FChaosVDQueryVisitStep& InOutVisitData)
+	{
+		if (!FChaosVisualDebuggerTrace::IsTracing())
+		{
+			return;
+		}
+
+		if constexpr (std::is_base_of_v<ChaosInterface::FLocationHit, THitType> || std::is_base_of_v<ChaosInterface::FPTLocationHit, THitType>)
+		{
+			InOutVisitData.HitData.Distance = Hit.Distance;
+			InOutVisitData.HitData.WorldPosition = Hit.WorldPosition;
+			InOutVisitData.HitData.WorldNormal = Hit.WorldNormal;
+			InOutVisitData.HitData.Flags = static_cast<uint16>(Hit.Flags.HitFlags);
+			InOutVisitData.HitData.FaceIdx = Hit.FaceIndex;
+		}
+
+		if constexpr (std::is_base_of_v<ChaosInterface::FQueryHit, THitType> || std::is_base_of_v<ChaosInterface::FPTQueryHit, THitType>)
+		{
+			InOutVisitData.HitData.FaceNormal = Hit.FaceNormal;
+		}
+
+		InOutVisitData.HitType = static_cast<EChaosVDCollisionQueryHitType>(HitType);
+
+		InOutVisitData.HitData.MarkAsValid();
+	}
+}
+
+#else
+	#ifndef CVD_TRACE_SCOPED_SCENE_QUERY_VISIT_HELPER
+		#define CVD_TRACE_SCOPED_SCENE_QUERY_VISIT_HELPER(Type, ParticleTransform, ParticleIndex, ShapeIndex, CurData)
+	#endif
+	#ifndef CVD_FILL_HIT_DATA_HELPER
+		#define CVD_FILL_HIT_DATA_HELPER(Hit, HitType)
+	#endif
+#endif
diff --git a/Engine/Source/Runtime/PhysicsCore/Public/SQVisitor.h b/Engine/Source/Runtime/PhysicsCore/Public/SQVisitor.h
--- a/Engine/Source/Runtime/PhysicsCore/Public/SQVisitor.h
+++ b/Engine/Source/Runtime/PhysicsCore/Public/SQVisitor.h
@@ -14,6 +14,12 @@
 
 #include <type_traits>
 
+#if WITH_CHAOS_VISUAL_DEBUGGER
+#include "DataWrappers/ChaosVDQueryDataWrappers.h"
+#endif
+
+#include "ChaosVDSQVisitorHelpers.h"
+
 #if CHAOS_DEBUG_DRAW
 extern PHYSICSCORE_API int32 ChaosSQDrawDebugVisitorQueries;
 extern PHYSICSCORE_API FAutoConsoleVariableRef CVarChaosSQDrawDebugQueries;
@@ -202,11 +208,15 @@ private:
 		bool bAllShapesIgnoredInPrefilter = true;
 		bool bHitBufferIncreased = false;
 #endif
-
-		for (const auto& Shape : Shapes)
+		
+		for (int32 ShapeIndex = 0; ShapeIndex < Shapes.Num(); ++ShapeIndex)
 		{
+			const TUniquePtr<FPerShapeData>& Shape = Shapes[ShapeIndex];
+
 			const FImplicitObject* Geom = Shape->GetGeometry();
 
+			CVD_TRACE_SCOPED_SCENE_QUERY_VISIT_HELPER(EChaosVDSceneQueryVisitorType::NarrowPhase, FTransform(GeometryParticle->R(), GeometryParticle->X()), Payload.UniqueIdx().Idx, ShapeIndex, CurData);
+
 			if (bTestShapeBounds)
 			{
 				FAABB3 InflatedWorldBounds;
@@ -221,6 +231,7 @@ private:
 
 					InflatedWorldBounds = FAABB3(Shape->GetWorldSpaceInflatedShapeBounds().Min() - WorldHalfExtent, Shape->GetWorldSpaceInflatedShapeBounds().Max() + WorldHalfExtent);
 				}
+	
 				if (SQ != ESQType::Overlap)
 				{
 					//todo: use fast raycast
@@ -355,10 +366,12 @@ private:
 
 						if (HitType != ECollisionQueryHitType::None)
 						{
-
 							//overlap never blocks
 							const bool bBlocker = (HitType == ECollisionQueryHitType::Block || bAnyHit || HitBuffer.WantsSingleResult());
 							HitBuffer.InsertHit(Hit, bBlocker);
+
+							CVD_FILL_HIT_DATA_HELPER(Hit, HitType);
+
 #if CHAOS_DEBUG_DRAW
 							bHitBufferIncreased = true;
 #endif
@@ -546,14 +559,22 @@ private:
 		const FShapesArray& Shapes = GeometryParticle->ShapesArray();
 		THitType Hit;
 		Hit.Actor = GeometryParticle;
-		for (const auto& Shape : Shapes)
+
+		for (int32 ShapeIndex = 0; ShapeIndex < Shapes.Num(); ++ShapeIndex)
 		{
+			const TUniquePtr<FPerShapeData>& Shape = Shapes[ShapeIndex];
+
+			CVD_TRACE_SCOPED_SCENE_QUERY_VISIT_HELPER(EChaosVDSceneQueryVisitorType::BroadPhase, FTransform(GeometryParticle->R(), GeometryParticle->X()), Payload.UniqueIdx().Idx, ShapeIndex, CurData);
+
 			ECollisionQueryHitType HitType = QueryFilterData.flags & FPhysicsQueryFlag::ePREFILTER ? QueryCallback.PreFilter(QueryFilterDataConcrete, *Shape, *GeometryParticle) : ECollisionQueryHitType::Block;
 			if (HitType != ECollisionQueryHitType::None)
 			{
 				const bool bBlocker = (HitType == ECollisionQueryHitType::Block || bAnyHit || HitBuffer.WantsSingleResult());
 				Hit.Shape = Shape.Get();
 				HitBuffer.InsertHit(Hit, bBlocker);
+
+				CVD_FILL_HIT_DATA_HELPER(Hit, HitType);
+
 				if (bAnyHit)
 				{
 					bContinue = false;
