[ChaosVD] Fix for an issue where mirrored geometry apears with its normals inverted

This is fallout of the refactor to support Instanced Static Mesh components made two weeks ago. I also added a cvar to be able to force CVD to only use normal Static Mesh components

diff --git a/Engine/Plugins/Experimental/ChaosVD/Source/ChaosVD/Private/ChaosVDParticleActor.cpp b/Engine/Plugins/Experimental/ChaosVD/Source/ChaosVD/Private/ChaosVDParticleActor.cpp
--- a/Engine/Plugins/Experimental/ChaosVD/Source/ChaosVD/Private/ChaosVDParticleActor.cpp
+++ b/Engine/Plugins/Experimental/ChaosVD/Source/ChaosVD/Private/ChaosVDParticleActor.cpp
@@ -18,6 +18,15 @@
 #include "Visualizers/ChaosVDParticleDataVisualizer.h"
 #include "Visualizers/ChaosVDSolverCollisionDataComponentVisualizer.h"
 
+namespace Chaos::VisualDebugger::Cvars
+{
+	static bool bForceStaticMeshComponentUse = false;
+	static FAutoConsoleVariableRef CVarChaosVDForceStaticMeshComponentUse(
+		TEXT("p.Chaos.VD.Tool.ForceStaticMeshComponentUse"),
+		bForceStaticMeshComponentUse,
+		TEXT("If true, static mesh components will be used instead of Instanced Static mesh components when recreating the geometry for each particle"));
+}
+
 AChaosVDParticleActor::AChaosVDParticleActor(const FObjectInitializer& ObjectInitializer) : Super(ObjectInitializer)
 {
 	RootComponent = CreateDefaultSubobject<USceneComponent>(TEXT("SceneComponent0"));
@@ -209,7 +218,7 @@ void AChaosVDParticleActor::UpdateGeometry(const Chaos::FConstImplicitObjectPtr&
 	OutExtractedGeometryDataHandles.Reserve(ObjectsToGenerateNum);
 
 	// Heightfields need to be created as Static meshes and use normal Static Mesh components because we need LODs for them due to their high triangle count
-	const bool bHasToUseStaticMeshComponent = FChaosVDGeometryBuilder::DoesImplicitContainType(InImplicitObject, Chaos::ImplicitObjectType::HeightField);
+	const bool bHasToUseStaticMeshComponent = Chaos::VisualDebugger::Cvars::bForceStaticMeshComponentUse || FChaosVDGeometryBuilder::DoesImplicitContainType(InImplicitObject, Chaos::ImplicitObjectType::HeightField);
 	constexpr int32 LODsToGenerateNum = 3;
 	constexpr int32 LODsToGenerateNumForInstancedStaticMesh = 0;
 
diff --git a/Engine/Plugins/Experimental/ChaosVD/Source/ChaosVD/Public/ChaosVDGeometryBuilder.h b/Engine/Plugins/Experimental/ChaosVD/Source/ChaosVD/Public/ChaosVDGeometryBuilder.h
--- a/Engine/Plugins/Experimental/ChaosVD/Source/ChaosVD/Public/ChaosVDGeometryBuilder.h
+++ b/Engine/Plugins/Experimental/ChaosVD/Source/ChaosVD/Public/ChaosVDGeometryBuilder.h
@@ -505,6 +505,12 @@ ComponentType* FChaosVDGeometryBuilder::GetAvailableInstancedStaticMeshComponent
 			return nullptr;
 		}
 
+		// If this is a Instanced Static Mesh Component, make sure we set the reverse culling flag correctly
+		if (UChaosVDInstancedStaticMeshComponent* AsInstancedMeshComponent = Component)
+		{
+			AsInstancedMeshComponent->bReverseCulling = EnumHasAnyFlags(MeshComponentAttributeFlags, EChaosVDMeshAttributesFlags::MirroredGeometry);
+		}
+
 		bOutIsNewComponent = true;
 
 		InstancedMeshComponentMapToSearch.Add(InExtractedGeometryDataHandle->GetGeometryKey(), Component);
