[ChaosVD] Fix for an issue where particle geometry was not properly unloaded when a solver is removed and other fixes

This change also adds improvements to unloading data in general by providing a progress bar when the CVD editor will be performing slow task like deleting actors or performing GC.

Also, fixed an issue where the folders used to group the particle actors by solver and particle type, were not deleted when loading a new file.

Lastly, fixed an issue where the pop up asking you to load a recording appears twice

diff --git a/Engine/Plugins/Experimental/ChaosVD/Source/ChaosVD/Private/Actors/ChaosVDSolverInfoActor.cpp b/Engine/Plugins/Experimental/ChaosVD/Source/ChaosVD/Private/Actors/ChaosVDSolverInfoActor.cpp
--- a/Engine/Plugins/Experimental/ChaosVD/Source/ChaosVD/Private/Actors/ChaosVDSolverInfoActor.cpp
+++ b/Engine/Plugins/Experimental/ChaosVD/Source/ChaosVD/Private/Actors/ChaosVDSolverInfoActor.cpp
@@ -4,9 +4,13 @@
 
 #include "ChaosVDParticleActor.h"
 #include "ChaosVDScene.h"
+#include "EditorActorFolders.h"
 #include "Components/ChaosVDSolverCollisionDataComponent.h"
 #include "Elements/Framework/TypedElementSelectionSet.h"
 #include "Engine/World.h"
+#include "Misc/ScopedSlowTask.h"
+
+#define LOCTEXT_NAMESPACE "ChaosVisualDebugger"
 
 AChaosVDSolverInfoActor::AChaosVDSolverInfoActor(const FObjectInitializer& ObjectInitializer) : Super(ObjectInitializer)
 {
@@ -26,10 +30,18 @@ void AChaosVDSolverInfoActor::SetScene(TWeakPtr<FChaosVDScene> InScene)
 
 void AChaosVDSolverInfoActor::RegisterParticleActor(int32 ParticleID, AChaosVDParticleActor* ParticleActor)
 {
+	if (!ensure(ParticleActor) || !ensure(ParticleActor->GetParticleData()))
+	{
+		return;
+	}
+
 	if (!SolverParticlesByID.Contains(ParticleID))
 	{
 		SolverParticlesByID.Add(ParticleID, ParticleActor);
 	}
+
+	ParticleActor->SetFolderPath(GetFolderPathForParticleType(ParticleActor->GetParticleData()->Type));
+	CreatedFolders.Add(ParticleActor->GetFolder());
 }
 
 AChaosVDParticleActor* AChaosVDSolverInfoActor::GetParticleActor(int32 ParticleID)
@@ -61,20 +73,56 @@ void AChaosVDSolverInfoActor::HandleColorsSettingsUpdated()
 	}
 }
 
-void AChaosVDSolverInfoActor::BeginDestroy()
+void AChaosVDSolverInfoActor::RemoveSolverFolders(UWorld* World)
 {
-	Super::BeginDestroy();
+	if (!World)
+	{
+		return;
+	}
+
+	TOptional<FFolder> ParentFolder;
+
+	for (const FFolder& Folder : CreatedFolders)
+	{
+		FActorFolders::Get().DeleteFolder(*World, Folder);
+
+		// All folders for the particles from this solver will have the same parent
+		if (!ParentFolder.IsSet())
+		{
+			ParentFolder = Folder.GetParent();
+		}
+	}
+	
+	if (ParentFolder.IsSet())
+	{
+		FActorFolders::Get().DeleteFolder(*World, ParentFolder.GetValue());
+	}
+}
 
+void AChaosVDSolverInfoActor::Destroyed()
+{
 	UWorld* World = GetWorld();
 	if (!World)
 	{
 		return;
 	}
+
+	constexpr float AmountOfWork = 1.0f;
+	const float PercentagePerElement = 1.0f / SolverParticlesByID.Num();
+
+	FScopedSlowTask CleaningParticleDataSlowTask(AmountOfWork, LOCTEXT("CleaningParticleDataMessage", "Cleaning Up Particle Data ..."));
+	CleaningParticleDataSlowTask.MakeDialog();
 	
 	for (const TPair<int32, AChaosVDParticleActor*>& ParticleVDInstanceWithID : SolverParticlesByID)
 	{
 		World->DestroyActor(ParticleVDInstanceWithID.Value);
+
+		CleaningParticleDataSlowTask.EnterProgressFrame(PercentagePerElement);
 	}
+
+	RemoveSolverFolders(World);
+
+	Super::Destroyed();
 }
 
 void AChaosVDSolverInfoActor::HandlePostSelectionChange(const UTypedElementSelectionSet* ChangesSelectionSet)
@@ -95,3 +143,21 @@ void AChaosVDSolverInfoActor::HandlePostSelectionChange(const UTypedElementSelec
 		}
 	}
 }
+
+FName AChaosVDSolverInfoActor::GetFolderPathForParticleType(EChaosVDParticleType ParticleType)
+{
+	if (const FName* FoundPathPtr = FolderPathByParticlePath.Find(ParticleType))
+	{
+		return *FoundPathPtr;
+	}
+	else
+	{
+		const FStringFormatOrderedArguments Args {SolverName, FString::FromInt(SolverID)};
+		const FName ParticleFolderPath = *FPaths::Combine(FString::Format(TEXT("Solver {0} | ID {1}"), Args), UEnum::GetDisplayValueAsText(ParticleType).ToString());
+
+		FolderPathByParticlePath.Add(ParticleType, ParticleFolderPath);
+		return ParticleFolderPath;
+	}
+}
+
+#undef LOCTEXT_NAMESPACE
diff --git a/Engine/Plugins/Experimental/ChaosVD/Source/ChaosVD/Private/ChaosVDParticleActor.cpp b/Engine/Plugins/Experimental/ChaosVD/Source/ChaosVD/Private/ChaosVDParticleActor.cpp
--- a/Engine/Plugins/Experimental/ChaosVD/Source/ChaosVD/Private/ChaosVDParticleActor.cpp
+++ b/Engine/Plugins/Experimental/ChaosVD/Source/ChaosVD/Private/ChaosVDParticleActor.cpp
@@ -95,7 +95,7 @@ void AChaosVDParticleActor::UpdateFromRecordedParticleData(const TSharedPtr<FCha
 
 	if (bHasNewTransform)
 	{
-		PerformTaskOnGeometry([NewParticleTransform](TSharedPtr<FChaosVDMeshDataInstanceHandle>& MeshDataHandle)
+		VisitGeometryInstances([NewParticleTransform](const TSharedRef<FChaosVDMeshDataInstanceHandle>& MeshDataHandle)
 		{
 			MeshDataHandle->SetWorldTransform(NewParticleTransform);
 		});
@@ -197,18 +197,13 @@ void AChaosVDParticleActor::UpdateGeometry(const Chaos::FConstImplicitObjectPtr&
 	// If the new implicit object is empty, then we can just clear all the mesh components and early out
 	if (ObjectsToGenerateNum == 0)
 	{
-		for (const TSharedPtr<FChaosVDMeshDataInstanceHandle>& MeshDataHandle : MeshDataHandles)
+		VisitGeometryInstances([this](const TSharedRef<FChaosVDMeshDataInstanceHandle>& MeshDataHandle)
 		{
-			if (!MeshDataHandle.IsValid())
-			{
-				continue;
-			}
-
 			if (IChaosVDGeometryComponent* AsGeometryComponent = Cast<IChaosVDGeometryComponent>(MeshDataHandle->GetMeshComponent()))
 			{
 				AsGeometryComponent->RemoveMeshInstance(MeshDataHandle);
 			}
-		}
+		});
 
 		MeshDataHandles.Reset();
 		return;
@@ -308,14 +303,24 @@ void AChaosVDParticleActor::SetScene(TWeakPtr<FChaosVDScene> InScene)
 	}
 }
 
-void AChaosVDParticleActor::BeginDestroy()
+void AChaosVDParticleActor::Destroyed()
 {
 	if (const TSharedPtr<FChaosVDScene>& ScenePtr = SceneWeakPtr.Pin())
 	{
 		ScenePtr->OnNewGeometryAvailable().Remove(GeometryUpdatedDelegate);
 	}
 
-	Super::BeginDestroy();
+	VisitGeometryInstances([](const TSharedRef<FChaosVDMeshDataInstanceHandle>& MeshDataHandle)
+	{
+		if (IChaosVDGeometryComponent* AsGeometryComponent = Cast<IChaosVDGeometryComponent>(MeshDataHandle->GetMeshComponent()))
+		{
+			AsGeometryComponent->RemoveMeshInstance(MeshDataHandle);
+		}
+	});
+
+	MeshDataHandles.Empty();
+
+	Super::Destroyed();
 }
 
 void AChaosVDParticleActor::GetVisualizationContext(FChaosVDVisualizationContext& OutVisualizationContext)
@@ -362,7 +367,7 @@ void AChaosVDParticleActor::SetIsTemporarilyHiddenInEditor(bool bIsHidden)
 {
 	Super::SetIsTemporarilyHiddenInEditor(bIsHidden);
 
-	PerformTaskOnGeometry([this, bIsHidden](TSharedPtr<FChaosVDMeshDataInstanceHandle>& MeshDataHandle)
+	VisitGeometryInstances([this, bIsHidden](const TSharedRef<FChaosVDMeshDataInstanceHandle>& MeshDataHandle)
 	{
 		FChaosVDGeometryComponentUtils::UpdateMeshVisibility(MeshDataHandle, ParticleDataPtr ? *ParticleDataPtr.Get() : FChaosVDParticleDataWrapper(), IsActive() && !bIsHidden);
 	});
@@ -434,7 +439,7 @@ void AChaosVDParticleActor::PushSelectionToProxies()
 {
 	Super::PushSelectionToProxies();
 
-	PerformTaskOnGeometry([this](TSharedPtr<FChaosVDMeshDataInstanceHandle>& MeshDataHandle)
+	VisitGeometryInstances([this](const TSharedRef<FChaosVDMeshDataInstanceHandle>& MeshDataHandle)
 	{
 		const bool bIsSelectedInEditor = IsSelectedInEditor();
 		MeshDataHandle->SetIsSelected(bIsSelectedInEditor);
@@ -467,7 +472,7 @@ const TArray<TSharedPtr<FChaosVDParticlePairMidPhase>>* AChaosVDParticleActor::G
 
 void AChaosVDParticleActor::UpdateShapeDataComponents()
 {
-	PerformTaskOnGeometry([this](TSharedPtr<FChaosVDMeshDataInstanceHandle>& MeshDataHandle)
+	VisitGeometryInstances([this](const TSharedRef<FChaosVDMeshDataInstanceHandle>& MeshDataHandle)
 	{
 		if (ParticleDataPtr)
 		{
@@ -478,7 +483,7 @@ void AChaosVDParticleActor::UpdateShapeDataComponents()
 
 void AChaosVDParticleActor::UpdateGeometryComponentsVisibility()
 {
-	PerformTaskOnGeometry([this](TSharedPtr<FChaosVDMeshDataInstanceHandle>& MeshDataHandle)
+	VisitGeometryInstances([this](const TSharedRef<FChaosVDMeshDataInstanceHandle>& MeshDataHandle)
 	{
 		if (ParticleDataPtr)
 		{
@@ -489,7 +494,7 @@ void AChaosVDParticleActor::UpdateGeometryComponentsVisibility()
 
 void AChaosVDParticleActor::UpdateGeometryColors()
 {
-	PerformTaskOnGeometry([this](TSharedPtr<FChaosVDMeshDataInstanceHandle>& MeshDataHandle)
+	VisitGeometryInstances([this](const TSharedRef<FChaosVDMeshDataInstanceHandle>& MeshDataHandle)
 	{
 		if (ParticleDataPtr)
 		{
diff --git a/Engine/Plugins/Experimental/ChaosVD/Source/ChaosVD/Private/ChaosVDPlaybackController.cpp b/Engine/Plugins/Experimental/ChaosVD/Source/ChaosVD/Private/ChaosVDPlaybackController.cpp
--- a/Engine/Plugins/Experimental/ChaosVD/Source/ChaosVD/Private/ChaosVDPlaybackController.cpp
+++ b/Engine/Plugins/Experimental/ChaosVD/Source/ChaosVD/Private/ChaosVDPlaybackController.cpp
@@ -119,7 +119,7 @@ void FChaosVDPlaybackController::PlayFromClosestKeyFrame_AssumesLocked(const int
 
 	if (CollapsedFrameData.SolverSteps.Num() > 0)
 	{
-		InSceneToControl.UpdateFromRecordedStepData(InTrackID, CollapsedFrameData.DebugName, CollapsedFrameData.SolverSteps[0], CollapsedFrameData);
+		InSceneToControl.UpdateFromRecordedStepData(InTrackID, CollapsedFrameData.SolverSteps[0], CollapsedFrameData);
 	}
 }
 
@@ -188,7 +188,7 @@ void FChaosVDPlaybackController::GoToRecordedSolverStep_AssumesLocked(const int3
 					{
 						if (SolverFrameData && SolverFrameData->SolverSteps.IsValidIndex(CurrentTrackInfo->LockedOnStep))
 						{
-							SceneToControlSharedPtr->UpdateFromRecordedStepData(InTrackID, SolverFrameData->DebugName, SolverFrameData->SolverSteps[StepIndex], *SolverFrameData);
+							SceneToControlSharedPtr->UpdateFromRecordedStepData(InTrackID, SolverFrameData->SolverSteps[StepIndex], *SolverFrameData);
 						}
 						else
 						{
@@ -202,7 +202,7 @@ void FChaosVDPlaybackController::GoToRecordedSolverStep_AssumesLocked(const int3
 				{
 					if (SolverFrameData && SolverFrameData->SolverSteps.IsValidIndex(Step))
 					{
-						SceneToControlSharedPtr->UpdateFromRecordedStepData(InTrackID, SolverFrameData->DebugName, SolverFrameData->SolverSteps[Step], *SolverFrameData);
+						SceneToControlSharedPtr->UpdateFromRecordedStepData(InTrackID, SolverFrameData->SolverSteps[Step], *SolverFrameData);
 					}
 					else
 					{
diff --git a/Engine/Plugins/Experimental/ChaosVD/Source/ChaosVD/Private/ChaosVDScene.cpp b/Engine/Plugins/Experimental/ChaosVD/Source/ChaosVD/Private/ChaosVDScene.cpp
--- a/Engine/Plugins/Experimental/ChaosVD/Source/ChaosVD/Private/ChaosVDScene.cpp
+++ b/Engine/Plugins/Experimental/ChaosVD/Source/ChaosVD/Private/ChaosVDScene.cpp
@@ -27,6 +27,7 @@
 #include "Selection.h"
 #include "UObject/Package.h"
 #include "WorldPersistentFolders.h"
+#include "Engine/Level.h"
 
 #define LOCTEXT_NAMESPACE "ChaosVisualDebugger"
 
@@ -36,7 +37,7 @@ FChaosVDScene::~FChaosVDScene() = default;
 
 namespace ChaosVDSceneUIOptions
 {
-	constexpr float DelayToShowProgressDialogThreshold = 1.5f;
+	constexpr float DelayToShowProgressDialogThreshold = 1.0f;
 	constexpr bool bShowCancelButton = false;
 	constexpr bool bAllowInPIE = false;
 }
@@ -76,6 +77,10 @@ void FChaosVDScene::Initialize()
 
 void FChaosVDScene::DeInitialize()
 {
+	constexpr float AmountOfWork = 1.0f;
+	FScopedSlowTask ClosingSceneSlowTask(AmountOfWork, LOCTEXT("ClosingSceneMessage", "Closing Scene ..."));
+	ClosingSceneSlowTask.MakeDialog();
+
 	if (!ensure(bIsInitialized))
 	{
 		return;
@@ -87,12 +92,12 @@ void FChaosVDScene::DeInitialize()
 		Settings->OnColorSettingsChanged().RemoveAll(this);
 	}
 
+	CleanUpScene();
+
 	DeInitializeSelectionSets();
 
 	GeometryGenerator.Reset();
 
-	CleanUpScene();
-
 	if (PhysicsVDWorld)
 	{
 		PhysicsVDWorld->RemoveOnActorDestroyededHandler(ActorDestroyedHandle);
@@ -103,8 +108,15 @@ void FChaosVDScene::DeInitialize()
 		PhysicsVDWorld->MarkAsGarbage();
 		PhysicsVDWorld = nullptr;
 	}
-	
-	CollectGarbage(GARBAGE_COLLECTION_KEEPFLAGS);
+
+	{
+		FScopedSlowTask CollectingGarbageSlowTask(1, LOCTEXT("CollectingGarbageDataMessage", "Collecting Garbage ..."));
+		CollectingGarbageSlowTask.MakeDialog();
+
+		CollectGarbage(GARBAGE_COLLECTION_KEEPFLAGS);
+
+		CollectingGarbageSlowTask.EnterProgressFrame();
+	}
 
 	bIsInitialized = false;
 }
@@ -118,7 +130,7 @@ void FChaosVDScene::AddReferencedObjects(FReferenceCollector& Collector)
 	Collector.AddReferencedObject(ComponentSelection);
 }
 
-void FChaosVDScene::UpdateFromRecordedStepData(const int32 SolverID, const FString& SolverName, const FChaosVDStepData& InRecordedStepData, const FChaosVDSolverFrameData& InFrameData)
+void FChaosVDScene::UpdateFromRecordedStepData(const int32 SolverID, const FChaosVDStepData& InRecordedStepData, const FChaosVDSolverFrameData& InFrameData)
 {
 	AChaosVDSolverInfoActor* SolverSceneData = SolverDataContainerBySolverID.FindChecked(SolverID);
 
@@ -134,7 +146,7 @@ void FChaosVDScene::UpdateFromRecordedStepData(const int32 SolverID, const FStri
 		constexpr float AmountOfWork = 1.0f;
 		const float PercentagePerElement = 1.0f / InRecordedStepData.RecordedParticlesData.Num();
 
-		const FText ProgressBarTitle = FText::Format(FTextFormat(LOCTEXT("ProcessingParticleData", "Processing Particle Data for {0} Solver with ID {1} ...")), FText::FromString(SolverName), FText::AsNumber(SolverID));
+		const FText ProgressBarTitle = FText::Format(FTextFormat(LOCTEXT("ProcessingParticleData", "Processing Particle Data for {0} Solver with ID {1} ...")), FText::FromString(SolverSceneData->GetSolverName()), FText::AsNumber(SolverID));
 		FScopedSlowTask UpdatingSceneSlowTask(AmountOfWork, ProgressBarTitle);
 		UpdatingSceneSlowTask.MakeDialogDelayed(ChaosVDSceneUIOptions::DelayToShowProgressDialogThreshold, ChaosVDSceneUIOptions::bShowCancelButton, ChaosVDSceneUIOptions::bAllowInPIE);
 	
@@ -164,11 +176,6 @@ void FChaosVDScene::UpdateFromRecordedStepData(const int32 SolverID, const FStri
 			{
 				if (AChaosVDParticleActor* NewParticleVDInstance = SpawnParticleFromRecordedData(Particle, InFrameData))
 				{
-					FStringFormatOrderedArguments Args {SolverName, FString::FromInt(SolverID)};
-					const FName FolderPath = *FPaths::Combine(FString::Format(TEXT("Solver {0} | ID {1}"), Args), UEnum::GetDisplayValueAsText(NewParticleVDInstance->GetParticleData()->Type).ToString());
-
-					NewParticleVDInstance->SetFolderPath(FolderPath);
-
 					// TODO: Precalculate the max num of entries we would see in the loaded file, and use that number to pre-allocate this map
 					SolverSceneData->RegisterParticleActor(ParticleVDInstanceID, NewParticleVDInstance);
 				}
@@ -231,6 +238,29 @@ void FChaosVDScene::HandleNewGeometryData(const Chaos::FConstImplicitObjectPtr&
 	NewGeometryAvailableDelegate.Broadcast(GeometryData, GeometryID);
 }
 
+void FChaosVDScene::CreateSolverInfoActor(int32 SolverID)
+{
+	if (!SolverDataContainerBySolverID.Contains(SolverID))
+	{
+		AChaosVDSolverInfoActor* SolverDataInfo = PhysicsVDWorld->SpawnActor<AChaosVDSolverInfoActor>();
+		check(SolverDataInfo);
+
+		const FName FolderPath("ChaosVisualDebugger/SolverDataContainer");
+		SolverDataInfo->SetFolderPath(FolderPath);
+
+		const bool bIsServer = IsSolverForServer(SolverID);
+		FString SolverName = LoadedRecording->GetSolverName_AssumedLocked(SolverID);
+
+		SolverDataInfo->SetSolverID(SolverID);
+		SolverDataInfo->SetSolverName(SolverName);
+		SolverDataInfo->SetScene(AsWeak());
+		SolverDataInfo->SetIsServer(bIsServer);
+
+		SolverDataContainerBySolverID.Add(SolverID, SolverDataInfo);
+	}
+}
+
+
 void FChaosVDScene::HandleEnterNewGameFrame(int32 FrameNumber, const TArray<int32>& AvailableSolversIds, const FChaosVDGameFrameData& InNewGameFrameData)
 {
 	// Currently the particle actors from all the solvers are in the same level, and we manage them by keeping track
@@ -246,18 +276,7 @@ void FChaosVDScene::HandleEnterNewGameFrame(int32 FrameNumber, const TArray<int3
 	{
 		AvailableSolversSet.Add(SolverID);
 
-		if (!SolverDataContainerBySolverID.Contains(SolverID))
-		{
-			AChaosVDSolverInfoActor* CollisionDataContainer = PhysicsVDWorld->SpawnActor<AChaosVDSolverInfoActor>();
-			check(CollisionDataContainer);
-
-			const bool bIsServer = LoadedRecording->GetSolverName_AssumedLocked(SolverID).Contains(TEXT("Server"));
-
-			CollisionDataContainer->SetSolverID(SolverID);
-			CollisionDataContainer->SetScene(AsWeak());
-			CollisionDataContainer->SetIsServer(bIsServer);
-			SolverDataContainerBySolverID.Add(SolverID, CollisionDataContainer);
-		}
+		CreateSolverInfoActor(SolverID);
 	}
 
 	int32 AmountRemoved = 0;
@@ -294,6 +313,12 @@ void FChaosVDScene::HandleEnterNewGameFrame(int32 FrameNumber, const TArray<int3
 
 void FChaosVDScene::CleanUpScene()
 {
+	constexpr float AmountOfWork = 1.0f;
+	const float PercentagePerElement = 1.0f / SolverDataContainerBySolverID.Num();
+
+	FScopedSlowTask CleaningSceneSlowTask(AmountOfWork, LOCTEXT("CleaningupSceneSolverMessage", "Clearing Solver Data ..."));
+	CleaningSceneSlowTask.MakeDialog();
+
 	ClearSelectionAndNotify();
 
 	if (PhysicsVDWorld)
@@ -301,6 +326,7 @@ void FChaosVDScene::CleanUpScene()
 		for (const TPair<int32, AChaosVDSolverInfoActor*>& SolverDataInfoWithID : SolverDataContainerBySolverID)
 		{
 			PhysicsVDWorld->DestroyActor(SolverDataInfoWithID.Value);
+			CleaningSceneSlowTask.EnterProgressFrame(PercentagePerElement);
 		}
 	}
 
@@ -469,6 +495,11 @@ UWorld* FChaosVDScene::CreatePhysicsVDWorld()
 										  .SetTransactional( false )
 	);
 
+	if (ULevel* Level = NewWorld->GetCurrentLevel())
+	{
+		Level->SetUseActorFolders(true);
+	}
+
 	CreateBaseLights(NewWorld);
 	CreateSceneQueriesContainer(NewWorld);
 	CreateMeshComponentsContainer(NewWorld);
diff --git a/Engine/Plugins/Experimental/ChaosVD/Source/ChaosVD/Public/Actors/ChaosVDSolverInfoActor.h b/Engine/Plugins/Experimental/ChaosVD/Source/ChaosVD/Public/Actors/ChaosVDSolverInfoActor.h
--- a/Engine/Plugins/Experimental/ChaosVD/Source/ChaosVD/Public/Actors/ChaosVDSolverInfoActor.h
+++ b/Engine/Plugins/Experimental/ChaosVD/Source/ChaosVD/Public/Actors/ChaosVDSolverInfoActor.h
@@ -10,6 +10,8 @@
 class AChaosVDParticleActor;
 class UChaosVDSolverCollisionDataComponent;
 
+enum class EChaosVDParticleType : uint8;
+
 UCLASS()
 class AChaosVDSolverInfoActor : public AActor, public FChaosVDSceneObjectBase, public FChaosVDSceneSelectionObserver
 {
@@ -22,6 +24,9 @@ public:
 	void SetSolverID(int32 InSolverID) { SolverID = InSolverID; }
 	int32 GetSolverID() const { return SolverID; }
 
+	void SetSolverName(const FString& InSolverName) { SolverName = InSolverName; }
+	const FString& GetSolverName() { return SolverName; }
+
 	void SetIsServer(bool bInIsServer) { bIsServer = bInIsServer; }
 	bool GetIsServer() const { return bIsServer; }
 
@@ -41,19 +46,25 @@ public:
 
 	void HandleVisibilitySettingsUpdated();
 	void HandleColorsSettingsUpdated();
+	void RemoveSolverFolders(UWorld* World);
 
-	virtual void BeginDestroy() override;
+	virtual void Destroyed() override;
 
 protected:
 
 	virtual void HandlePostSelectionChange(const UTypedElementSelectionSet* ChangesSelectionSet) override;
 
+	FName GetFolderPathForParticleType(EChaosVDParticleType ParticleType);
+
 	UPROPERTY(VisibleAnywhere, Category="Solver Data")
 	int32 SolverID = INDEX_NONE;
 
 	UPROPERTY(VisibleAnywhere, Category="Solver Data")
 	FTransform SimulationTransform;
 
+	UPROPERTY(VisibleAnywhere, Category="Solver Data")
+	FString SolverName;
+
 	UPROPERTY()
 	TObjectPtr<UChaosVDSolverCollisionDataComponent> CollisionDataComponent;
 
@@ -62,5 +73,9 @@ protected:
 
 	TArray<int32> SelectedParticlesID;
 
+	TSortedMap<EChaosVDParticleType, FName> FolderPathByParticlePath;
+
+	TSet<FFolder> CreatedFolders;
+
 	bool bIsServer;
 };
diff --git a/Engine/Plugins/Experimental/ChaosVD/Source/ChaosVD/Public/ChaosVDParticleActor.h b/Engine/Plugins/Experimental/ChaosVD/Source/ChaosVD/Public/ChaosVDParticleActor.h
--- a/Engine/Plugins/Experimental/ChaosVD/Source/ChaosVD/Public/ChaosVDParticleActor.h
+++ b/Engine/Plugins/Experimental/ChaosVD/Source/ChaosVD/Public/ChaosVDParticleActor.h
@@ -56,7 +56,7 @@ public:
 
 	virtual void SetScene(TWeakPtr<FChaosVDScene> InScene) override;
 
-	virtual void BeginDestroy() override;
+	virtual void Destroyed() override;
 
 	virtual const FChaosVDParticleDataWrapper* GetParticleData() override { return ParticleDataPtr.Get(); }
 	virtual void GetVisualizationContext(FChaosVDVisualizationContext& OutVisualizationContext) override;
@@ -108,7 +108,7 @@ protected:
 	void UpdateShapeDataComponents();
 
 	template<typename TTaskCallback>
-	void PerformTaskOnGeometry(const TTaskCallback& TaskToPerform);
+	void VisitGeometryInstances(const TTaskCallback& VisitorCallback);
 	
 	UPROPERTY(EditAnywhere, Category = "Viewport Visualization Flags", meta = (Bitmask, BitmaskEnum = "/Script/ChaosVD.EChaosVDParticleDataVisualizationFlags"))
 	uint8 LocalParticleDataVisualizationFlags;
@@ -137,14 +137,14 @@ protected:
 	friend FChaosVDParticleActorCustomization;
 };
 
-template <typename TTaskCallback>
-void AChaosVDParticleActor::PerformTaskOnGeometry(const TTaskCallback& TaskToPerform)
+template <typename TVisitorCallback>
+void AChaosVDParticleActor::VisitGeometryInstances(const TVisitorCallback& VisitorCallback)
 {
 	for (TSharedPtr<FChaosVDMeshDataInstanceHandle>& MeshDataHandle : MeshDataHandles)
 	{
 		if (MeshDataHandle)
 		{
-			TaskToPerform(MeshDataHandle);
+			VisitorCallback(MeshDataHandle.ToSharedRef());
 		}
 	}
 }
diff --git a/Engine/Plugins/Experimental/ChaosVD/Source/ChaosVD/Public/ChaosVDScene.h b/Engine/Plugins/Experimental/ChaosVD/Source/ChaosVD/Public/ChaosVDScene.h
--- a/Engine/Plugins/Experimental/ChaosVD/Source/ChaosVD/Public/ChaosVDScene.h
+++ b/Engine/Plugins/Experimental/ChaosVD/Source/ChaosVD/Public/ChaosVDScene.h
@@ -51,13 +51,13 @@ public:
 	FChaosVDSceneUpdatedDelegate& OnSceneUpdated() { return SceneUpdatedDelegate; }
 
 	/** Updates, Adds and Remove actors to match the provided Step Data */
-	void UpdateFromRecordedStepData(const int32 SolverID, const FString& SolverName, const FChaosVDStepData& InRecordedStepData, const FChaosVDSolverFrameData& InFrameData);
+	void UpdateFromRecordedStepData(const int32 SolverID, const FChaosVDStepData& InRecordedStepData, const FChaosVDSolverFrameData& InFrameData);
 
 	void UpdateParticlesCollisionData(const FChaosVDStepData& InRecordedStepData, int32 SolverID);
 
 	// No need to deprecate the old version since it is not a public API nor inline 
 	void HandleNewGeometryData(const Chaos::FConstImplicitObjectPtr& Geometry, const uint32 GeometryID) const;
-	
+
 	void HandleEnterNewGameFrame(int32 FrameNumber, const TArray<int32>& AvailableSolversIds, const FChaosVDGameFrameData& InNewGameFrameData);
 
 	/** Deletes all actors of the Scene and underlying UWorld */
@@ -120,10 +120,14 @@ private:
 
 	/** Returns the ID used to track this recorded particle data */
 	int32 GetIDForRecordedParticleData(const TSharedPtr<FChaosVDParticleDataWrapper>& InParticleData) const;
+
 	void CreateBaseLights(UWorld* TargetWorld) const;
 
 	void CreateSceneQueriesContainer(UWorld* TargetWorld);
 
+	/** Creates an actor that will contain all solver data for the provided Solver ID*/
+	void CreateSolverInfoActor(int32 SolverID);
+
 	AActor* CreateMeshComponentsContainer(UWorld* TargetWorld);
 
 	/** Creates the instance of the World which will be used the recorded data*/
diff --git a/Engine/Source/Runtime/Experimental/ChaosVisualDebugger/Private/ChaosVDRuntimeModule.cpp b/Engine/Source/Runtime/Experimental/ChaosVisualDebugger/Private/ChaosVDRuntimeModule.cpp
--- a/Engine/Source/Runtime/Experimental/ChaosVisualDebugger/Private/ChaosVDRuntimeModule.cpp
+++ b/Engine/Source/Runtime/Experimental/ChaosVisualDebugger/Private/ChaosVDRuntimeModule.cpp
@@ -57,8 +57,6 @@ void FChaosVDRuntimeModule::StartupModule()
 	// Make sure it is off until we support auto trace
 	UE::Trace::ToggleChannel(TEXT("ChaosVDChannel"), false);
 #endif
-
-	FTraceAuxiliary::OnTraceStopped.AddRaw(this, &FChaosVDRuntimeModule::HandleTraceStopRequest);
 }
 
 void FChaosVDRuntimeModule::ShutdownModule()
@@ -114,6 +112,9 @@ void FChaosVDRuntimeModule::StartRecording(const TArray<FString>& Args)
 		return;
 	}
 
+	// Start Listening for Trace Stopped events, in case Trace is stopped outside our control so we can gracefully stop CVD recording and log a warning 
+	FTraceAuxiliary::OnTraceStopped.AddRaw(this, &FChaosVDRuntimeModule::HandleTraceStopRequest);
+
 	{
 		FReadScopeLock ReadLock(DelegatesRWLock);
 		RecordingStartedDelegate.Broadcast();
@@ -189,6 +190,8 @@ void FChaosVDRuntimeModule::StopRecording()
 	{
 		return;
 	}
+	
+	FTraceAuxiliary::OnTraceStopped.RemoveAll(this);
 
 #if UE_TRACE_ENABLED
 	
