[ChaosVD] Added an option to override the playback framerate. And fixed and assert that happened when the Solver tracks was closed and then re-opened.

Now you can set a desired framerate that will override the one that was recorded. This is useful to see things in "slow mo", compensate for a low recorded framerate or to ensure playback is not halted for large recorded frame times caused by an ensure, hitch or a breakpoint.

diff --git a/Engine/Plugins/Experimental/ChaosVD/Source/ChaosVD/Private/ChaosVDEditorSettings.cpp b/Engine/Plugins/Experimental/ChaosVD/Source/ChaosVD/Private/ChaosVDEditorSettings.cpp
--- a/Engine/Plugins/Experimental/ChaosVD/Source/ChaosVD/Private/ChaosVDEditorSettings.cpp
+++ b/Engine/Plugins/Experimental/ChaosVD/Source/ChaosVD/Private/ChaosVDEditorSettings.cpp
@@ -86,6 +86,11 @@ void UChaosVDEditorSettings::PostEditChangeProperty(FPropertyChangedEvent& Prope
 	{
 		FarClippingOverrideChangedDelegate.Broadcast(this);
 	}
+	else if (PropertyName == GET_MEMBER_NAME_CHECKED(UChaosVDEditorSettings, bPlaybackAtRecordedFrameRate) ||
+			 PropertyName == GET_MEMBER_NAME_CHECKED(UChaosVDEditorSettings, TargetFrameRateOverride))
+	{
+		PlaybackSettingsChangedDelegate.Broadcast(this);
+	}
 
 	// TODO: If we keep this object as the main setting object,
 	// we should have a single event for what changed and an enum flags that the listener could use to decide if cares about the change
diff --git a/Engine/Plugins/Experimental/ChaosVD/Source/ChaosVD/Private/ChaosVDPlaybackController.cpp b/Engine/Plugins/Experimental/ChaosVD/Source/ChaosVD/Private/ChaosVDPlaybackController.cpp
--- a/Engine/Plugins/Experimental/ChaosVD/Source/ChaosVD/Private/ChaosVDPlaybackController.cpp
+++ b/Engine/Plugins/Experimental/ChaosVD/Source/ChaosVD/Private/ChaosVDPlaybackController.cpp
@@ -2,6 +2,7 @@
 
 #include "ChaosVDPlaybackController.h"
 
+#include "ChaosVDEditorSettings.h"
 #include "ChaosVDModule.h"
 #include "ChaosVDPlaybackControllerInstigator.h"
 #include "ChaosVDRecording.h"
@@ -16,6 +17,12 @@ FChaosVDPlaybackController::FChaosVDPlaybackController(const TWeakPtr<FChaosVDSc
 	SceneToControl = InSceneToControl;
 
 	RecordingStoppedHandle = FChaosVDRuntimeModule::Get().RegisterRecordingStopCallback(FChaosVDRecordingStateChangedDelegate::FDelegate::CreateRaw(this, &FChaosVDPlaybackController::HandleDisconnectedFromSession));
+
+	if (UChaosVDEditorSettings* Settings = GetMutableDefault<UChaosVDEditorSettings>())
+	{
+		Settings->OnPlaybackSettingsChanged().AddRaw(this, &FChaosVDPlaybackController::HandleFrameRateOverrideSettingsChanged);
+		HandleFrameRateOverrideSettingsChanged(Settings);
+	}
 }
 
 FChaosVDPlaybackController::~FChaosVDPlaybackController()
@@ -25,6 +32,11 @@ FChaosVDPlaybackController::~FChaosVDPlaybackController()
 	{
 		FChaosVDRuntimeModule::Get().RemoveRecordingStopCallback(RecordingStoppedHandle);
 	}
+	
+	if (UChaosVDEditorSettings* Settings = GetMutableDefault<UChaosVDEditorSettings>())
+	{
+		Settings->OnPlaybackSettingsChanged().RemoveAll(this);
+	}
 
 	UnloadCurrentRecording(EChaosVDUnloadRecordingFlags::Silent);
 }
@@ -142,6 +154,16 @@ void FChaosVDPlaybackController::EnqueueGeometryDataUpdate(const Chaos::FConstIm
 	GeometryDataUpdateGTQueue.Enqueue({NewGeometry, GeometryID });
 }
 
+void FChaosVDPlaybackController::HandleFrameRateOverrideSettingsChanged(UChaosVDEditorSettings* CVDSettings)
+{
+	if (!CVDSettings)
+	{
+		return;
+	}
+
+	CurrentFrameRateOverride = CVDSettings->bPlaybackAtRecordedFrameRate ? InvalidFrameRateOverride : CVDSettings->TargetFrameRateOverride;
+}
+
 void FChaosVDPlaybackController::GoToRecordedSolverStep_AssumesLocked(const int32 InTrackID, const int32 FrameNumber, const int32 Step, FGuid InstigatorID, int32 Attempts)
 {
 	if (const TSharedPtr<FChaosVDScene> SceneToControlSharedPtr = SceneToControl.Pin())
@@ -633,6 +655,51 @@ void FChaosVDPlaybackController::HandleDisconnectedFromSession()
 	bHasPendingGTUpdateBroadcast = true;
 }
 
+float FChaosVDPlaybackController::GetFrameTimeOverride() const
+{
+	constexpr int32 MinimumFrameRateOverride = 1;
+	return CurrentFrameRateOverride >= MinimumFrameRateOverride ? 1.0f / static_cast<float>(CurrentFrameRateOverride) : InvalidFrameRateOverride;
+}
+
+float FChaosVDPlaybackController::GetFrameTimeForTrack(EChaosVDTrackType TrackType, int32 TrackID, const FChaosVDTrackInfo& TrackInfo) const
+{
+	const float TargetFrameTimeOverride = GetFrameTimeOverride();
+	const bool bHastFrameRateOverride = !FMath::IsNearlyEqual(TargetFrameTimeOverride, FChaosVDPlaybackController::InvalidFrameRateOverride);
+	if (bHastFrameRateOverride)
+	{
+		return TargetFrameTimeOverride;
+	}
+
+	float CurrentTargetFrameTime = FallbackFrameTime;
+	if (LoadedRecording)
+	{
+		switch(TrackType)
+		{
+			case EChaosVDTrackType::Solver:
+				{
+					if (const FChaosVDSolverFrameData* FrameData = LoadedRecording->GetSolverFrameData_AssumesLocked(TrackID, TrackInfo.CurrentFrame))
+					{
+						CurrentTargetFrameTime = FrameData->GetFrameTime();
+					}
+					break;
+				}
+			case EChaosVDTrackType::Game:
+				{
+					if (const FChaosVDGameFrameData* FrameData = LoadedRecording->GetGameFrameData_AssumesLocked(TrackInfo.CurrentFrame))
+					{
+						CurrentTargetFrameTime = FrameData->GetFrameTime();
+					}
+
+					break;
+				}
+			default:
+				break;
+		}
+	}
+
+	return CurrentTargetFrameTime;
+}
+
 void FChaosVDPlaybackController::UpdateSolverTracksData()
 {
 	const TMap<int32, TArray<FChaosVDSolverFrameData>>& SolversByID = LoadedRecording->GetAvailableSolvers_AssumesLocked();
diff --git a/Engine/Plugins/Experimental/ChaosVD/Source/ChaosVD/Private/Widgets/SChaosVDPlaybackViewport.cpp b/Engine/Plugins/Experimental/ChaosVD/Source/ChaosVD/Private/Widgets/SChaosVDPlaybackViewport.cpp
--- a/Engine/Plugins/Experimental/ChaosVD/Source/ChaosVD/Private/Widgets/SChaosVDPlaybackViewport.cpp
+++ b/Engine/Plugins/Experimental/ChaosVD/Source/ChaosVD/Private/Widgets/SChaosVDPlaybackViewport.cpp
@@ -189,24 +189,28 @@ void SChaosVDPlaybackViewport::HandlePlaybackControllerDataUpdated(TWeakPtr<FCha
 
 void SChaosVDPlaybackViewport::HandleControllerTrackFrameUpdated(TWeakPtr<FChaosVDPlaybackController> InController, const FChaosVDTrackInfo* UpdatedTrackInfo, FGuid InstigatorGuid)
 {
-	if (TSharedPtr<FChaosVDPlaybackController> ControllerSharedPtr = InController.Pin())
+	if (const TSharedPtr<FChaosVDPlaybackController> ControllerSharedPtr = InController.Pin())
 	{
-		// The frame number we receive could be from a Solver track, so make sure it is converted to the correct game track frame number 
-		int32 GameTrackFrame = ControllerSharedPtr->ConvertCurrentFrameToOtherTrackFrame(UpdatedTrackInfo, ControllerSharedPtr->GetTrackInfo(EChaosVDTrackType::Game, FChaosVDPlaybackController::GameTrackID));
-
-		// Something other than us advanced the game track, so make sure the timeline widget is updated
-		if (InstigatorGuid != GetInstigatorID())
+		// The frame number we receive could be from a Solver track, so make sure it is converted to the correct game track frame number
+		if (const FChaosVDTrackInfo* GameTrackInfo = ControllerSharedPtr->GetTrackInfo(EChaosVDTrackType::Game, FChaosVDPlaybackController::GameTrackID))
 		{
-			GameFramesTimelineWidget->SetCurrentTimelineFrame(GameTrackFrame, EChaosVDSetTimelineFrameFlags::None);
-		}
+			const int32 GameTrackFrame = ControllerSharedPtr->ConvertCurrentFrameToOtherTrackFrame(UpdatedTrackInfo, GameTrackInfo);
 
-		if (TSharedPtr<FChaosVDRecording> RecordingData = ControllerSharedPtr->GetCurrentRecording().Pin())
-		{
-			if (FChaosVDGameFrameData* FrameData = RecordingData->GetGameFrameData_AssumesLocked(GameTrackFrame))
+			// Something other than us advanced the game track, so make sure the timeline widget is updated
+			if (InstigatorGuid != GetInstigatorID())
+			{
+				GameFramesTimelineWidget->SetCurrentTimelineFrame(GameTrackFrame, EChaosVDSetTimelineFrameFlags::None);
+			}
+
+			GameFramesTimelineWidget->SetTargetFrameTime(ControllerSharedPtr->GetFrameTimeForTrack(EChaosVDTrackType::Game, FChaosVDPlaybackController::GameTrackID, *GameTrackInfo));
+
+			if (const TSharedPtr<FChaosVDRecording> RecordingData = ControllerSharedPtr->GetCurrentRecording().Pin())
 			{
-				PlaybackViewportClient->PerformSelectedTrackingForFrame(FrameData);
-				GameFramesTimelineWidget->SetTargetFrameTime(FrameData->GetFrameTime());
-			}	
+				if (FChaosVDGameFrameData* FrameData = RecordingData->GetGameFrameData_AssumesLocked(GameTrackFrame))
+				{
+					PlaybackViewportClient->PerformSelectedTrackingForFrame(FrameData);
+				}
+			}
 		}
 	}
 }
diff --git a/Engine/Plugins/Experimental/ChaosVD/Source/ChaosVD/Private/Widgets/SChaosVDSolverPlaybackControls.cpp b/Engine/Plugins/Experimental/ChaosVD/Source/ChaosVD/Private/Widgets/SChaosVDSolverPlaybackControls.cpp
--- a/Engine/Plugins/Experimental/ChaosVD/Source/ChaosVD/Private/Widgets/SChaosVDSolverPlaybackControls.cpp
+++ b/Engine/Plugins/Experimental/ChaosVD/Source/ChaosVD/Private/Widgets/SChaosVDSolverPlaybackControls.cpp
@@ -163,13 +163,7 @@ void SChaosVDSolverPlaybackControls::HandleControllerTrackFrameUpdated(TWeakPtr<
 				UpdateStepsWidgetForFrame(*CurrentPlaybackControllerPtr.Get(), SolverTrackInfo->CurrentFrame, SolverTrackInfo->CurrentStep);
 			}
 
-			if (const TSharedPtr<FChaosVDRecording> RecordingData = CurrentPlaybackControllerPtr->GetCurrentRecording().Pin())
-			{
-				if (const FChaosVDSolverFrameData* FrameData = RecordingData->GetSolverFrameData_AssumesLocked(SolverID, SolverTrackInfo->CurrentFrame))
-				{
-					FramesTimelineWidget->SetTargetFrameTime(FrameData->GetFrameTime());
-				}	
-			}
+			FramesTimelineWidget->SetTargetFrameTime(CurrentPlaybackControllerPtr->GetFrameTimeForTrack(EChaosVDTrackType::Solver, SolverID, *SolverTrackInfo));
 		}
 	}
 }
diff --git a/Engine/Plugins/Experimental/ChaosVD/Source/ChaosVD/Private/Widgets/SChaosVDSolverTracks.cpp b/Engine/Plugins/Experimental/ChaosVD/Source/ChaosVD/Private/Widgets/SChaosVDSolverTracks.cpp
--- a/Engine/Plugins/Experimental/ChaosVD/Source/ChaosVD/Private/Widgets/SChaosVDSolverTracks.cpp
+++ b/Engine/Plugins/Experimental/ChaosVD/Source/ChaosVD/Private/Widgets/SChaosVDSolverTracks.cpp
@@ -2,6 +2,7 @@
 
 #include "Widgets/SChaosVDSolverTracks.h"
 
+#include "ChaosVDModule.h"
 #include "ChaosVDPlaybackController.h"
 #include "Widgets/SChaosVDPlaybackViewport.h"
 #include "Widgets/SChaosVDSolverPlaybackControls.h"
@@ -9,8 +10,6 @@
 
 void SChaosVDSolverTracks::Construct(const FArguments& InArgs, TWeakPtr<FChaosVDPlaybackController> InPlaybackController)
 {
-	RegisterNewController(InPlaybackController);
-
 	ChildSlot
 	[
 		SAssignNew(SolverTracksListWidget, SListView<TSharedPtr<FChaosVDTrackInfo>>)
@@ -19,6 +18,10 @@ void SChaosVDSolverTracks::Construct(const FArguments& InArgs, TWeakPtr<FChaosVD
 			.ListViewStyle(&FAppStyle::Get().GetWidgetStyle<FTableViewStyle>("SimpleListView"))
 			.OnGenerateRow(this, &SChaosVDSolverTracks::MakeSolverTrackControlsFromTrackInfo)
 	];
+
+	ensure(InPlaybackController.IsValid());
+
+	RegisterNewController(InPlaybackController);
 	
 	if (const TSharedPtr<FChaosVDPlaybackController> CurrentPlaybackControllerPtr = InPlaybackController.Pin())
 	{
@@ -27,6 +30,10 @@ void SChaosVDSolverTracks::Construct(const FArguments& InArgs, TWeakPtr<FChaosVD
 			HandleControllerTrackFrameUpdated(InPlaybackController, GameTrackInfo, InvalidGuid);
 		}
 	}
+	else
+	{
+		UE_LOG(LogChaosVDEditor, Error, TEXT("[%s] Solver tracks contructed with an invalid player controler. The solver tracks widget will not be functional"), ANSI_TO_TCHAR(__FUNCTION__))	
+	}
 }
 
 void SChaosVDSolverTracks::HandlePlaybackControllerDataUpdated(TWeakPtr<FChaosVDPlaybackController> InPlaybackController)
diff --git a/Engine/Plugins/Experimental/ChaosVD/Source/ChaosVD/Public/ChaosVDEditorSettings.h b/Engine/Plugins/Experimental/ChaosVD/Source/ChaosVD/Public/ChaosVDEditorSettings.h
--- a/Engine/Plugins/Experimental/ChaosVD/Source/ChaosVD/Public/ChaosVDEditorSettings.h
+++ b/Engine/Plugins/Experimental/ChaosVD/Source/ChaosVD/Public/ChaosVDEditorSettings.h
@@ -21,29 +21,43 @@ DECLARE_MULTICAST_DELEGATE_OneParam(FChaosVDSettingChaged, UChaosVDEditorSetting
 UENUM()
 enum class EChaosVDActorTrackingMode
 {
+	/** Follow the selected target keeping the specified distance from it. */
 	ByDistanceOffset,
+	/** Follow the selected target keeping its bounding box inside of CVD's camera view. */
 	ByBoundingBox,
-	MatchTransform
+	/** Snap CVD's camera to the transform of the current selected object */
+	MatchTransform 
 };
 
 UENUM()
 enum class EChaosVDActorTrackingTarget
 {
+	/** Disable Camera Auto-Tracking */
 	Disabled,
+	/** Follow the current selected object */
 	SelectedObject,
+	/** Follow the selected recorded transform ID */
 	RecordedTransform,
-	RecordedLocation,
+	/** Follow the selected recorded location ID */
+	RecordedLocation
 };
 
 UENUM(meta = (Bitflags, UseEnumValuesAsMaskValuesInEditor = "true"))
 enum class EChaosVDGeometryVisibilityFlags : uint8
 {
+	/** Draws all geometry that is for Query Only */
 	Query = 1 << 1,
+	/** Draws all geometry that is for [Physics Collision] or [Physics Collision and Query only] */
 	Simulated = 1 << 2,
+	/** Draws all simple geometry */
 	Simple = 1 << 3,
+	/** Draws all complex geometry */
 	Complex = 1 << 4,
-	ShowHeightfields = 1 << 5, // Selecting this will show heightfields even if complex is not selected
+	/** Draws heightfields even if complex is not selected */
+	ShowHeightfields = 1 << 5,
+	/** Draws all particles that are in a disabled state */
 	ShowDisabledParticles = 1 << 6,
+	/** Draws all triangle lines that forms each geometry object. Similar to the wireframe view, but the polygons are till filled and shaded */
 	ShowTriangleEdges = 1 << 7,
 };
 ENUM_CLASS_FLAGS(EChaosVDGeometryVisibilityFlags)
@@ -157,10 +171,10 @@ struct FChaosDebugDrawColorsByClientServer
 UENUM()
 enum class EChaosVDParticleDebugColorMode
 {
-	None,
-	State,
-	ShapeType,
-	ClientServer,
+	None, // Draw particles with the default gray color
+	State, // Draw particles with a specific color based on the recorded particle state
+	ShapeType, // Draw particles with a specific color based on their shape type
+	ClientServer, // Draw particles with a specific color based on if they are a Server Particle or Client particle
 };
 
 UCLASS(config = Engine)
@@ -169,18 +183,31 @@ class UChaosVDEditorSettings : public UObject
 	GENERATED_BODY()
 public:
 
-	UPROPERTY(EditAnywhere, Category = "Viewport")
-	float FarClippingOverride = 10000.0f;
+	/** If true, playback will respect the recorded frame times */
+	UPROPERTY(EditAnywhere, Category = "Playback Settings")
+	bool bPlaybackAtRecordedFrameRate = true;
 
+	/** If play at recorded frame rate is disabled, CVD will attempt to play the recording at the specified frame rate */
+	UPROPERTY(EditAnywhere, Category = "Playback Settings", meta=(EditCondition = "!bPlaybackAtRecordedFrameRate", EditConditionHides, ClampMin = 1))
+	int32 TargetFrameRateOverride = 60;
+
+	/** Sets the desired far clipping for CVD's viewport */
+	UPROPERTY(EditAnywhere, Category = "Viewport", meta=(ClampMin = 1))
+	float FarClippingOverride = 20000.0f;
+
+	/** Set of flags to enable/disable visualization of specific particle data as debug draw */
 	UPROPERTY(EditAnywhere, Category = "Viewport Visualization Flags", meta = (Bitmask, BitmaskEnum = "/Script/ChaosVD.EChaosVDParticleDataVisualizationFlags"))
 	uint32 GlobalParticleDataVisualizationFlags = 0;
-	
+
+	/** Set of flags to enable/disable visualization of specific collision data as debug draw */
 	UPROPERTY(EditAnywhere, Category = "Viewport Visualization Flags", meta = (Bitmask, BitmaskEnum = "/Script/ChaosVD.EChaosVDCollisionVisualizationFlags"))
 	uint32 GlobalCollisionDataVisualizationFlags = 0;
 
+	/** Set of flags to enable/disable visualization of specific scene queries data as debug draw */
 	UPROPERTY(EditAnywhere, Category = "Viewport Visualization Flags", meta = (Bitmask, BitmaskEnum = "/Script/ChaosVD.EChaosVDSceneQueryVisualizationFlags"))
 	uint32 GlobalSceneQueriesVisualizationFlags = static_cast<uint8>(EChaosVDSceneQueryVisualizationFlags::DrawClientQueries | EChaosVDSceneQueryVisualizationFlags::DrawServerQueries |  EChaosVDSceneQueryVisualizationFlags::DrawHits | EChaosVDSceneQueryVisualizationFlags::DrawLineTraceQueries);
 
+	/** If true, text information (if available) will be drawn alongside any other debug draw shape */
 	UPROPERTY(EditAnywhere, Category = "Viewport Visualization")
 	bool bShowDebugText = false;
 
@@ -199,18 +226,23 @@ public:
 	UPROPERTY(EditAnywhere, Category = "Viewport Visualization", meta = (EditCondition = "ParticleColorMode == EChaosVDParticleDebugColorMode::ClientServer", EditConditionHides))
 	FChaosDebugDrawColorsByClientServer ColorsByClientServer;
 
+	/** Sets what should be auto tracked by CVD's Camera */
 	UPROPERTY(EditAnywhere, Category = "Viewport Tracking")
 	EChaosVDActorTrackingTarget TrackingTarget;
 
+	/** Sets how tracking should be performed by CVD's Camera */
 	UPROPERTY(EditAnywhere, Category = "Viewport Tracking", meta=(EditCondition = "TrackingTarget != EChaosVDActorTrackingTarget::Disabled", EditConditionHides))
 	EChaosVDActorTrackingMode TrackingOptions;
 
+	/** Distance from the selected tracking target CVD's camera should attempt to be at */
 	UPROPERTY(EditAnywhere, Category = "Viewport Tracking", meta=(EditCondition = "TrackingOptions == EChaosVDActorTrackingMode::ByDistanceOffset && TrackingTarget != EChaosVDActorTrackingTarget::Disabled", EditConditionHides))
 	float TrackingDistanceOffset = 1500.0f;
 
+	/** By how much we should expand the bounding box used to track a target by bounding box. Used to see more of the screen while tracking in Bounding Box mode */
 	UPROPERTY(EditAnywhere, Category = "Viewport Tracking", meta=(EditCondition = "TrackingOptions == EChaosVDActorTrackingMode::ByBoundingBox && TrackingTarget != EChaosVDActorTrackingTarget::Disabled", EditConditionHides))
 	float ExpandViewTrackingBy = 60.0f;
 
+	/** Set of flags to enable/disable visibility of specific types of geometry/particles */
 	UPROPERTY(EditAnywhere, Category = "Geometry Visibility", meta = (Bitmask, BitmaskEnum = "/Script/ChaosVD.EChaosVDGeometryVisibilityFlags"))
 	uint8 GeometryVisibilityFlags = static_cast<uint8>(EChaosVDGeometryVisibilityFlags::Simulated | EChaosVDGeometryVisibilityFlags::Simple |  EChaosVDGeometryVisibilityFlags::ShowHeightfields);
 
@@ -237,6 +269,8 @@ public:
 
 	FChaosVDSettingChaged& OnFarClippingOverrideChanged() { return FarClippingOverrideChangedDelegate; }
 
+	FChaosVDSettingChaged& OnPlaybackSettingsChanged() { return PlaybackSettingsChangedDelegate ; }
+
 	TSharedPtr<FName> SelectedTrackedTransformName;
 	TSharedPtr<FName> SelectedTrackedLocationName;
 
@@ -244,4 +278,5 @@ protected:
 	FChaosVDSettingChaged VisibilitySettingsChangedDelegate;
 	FChaosVDSettingChaged ColorsSettingsChangedDelegate;
 	FChaosVDSettingChaged FarClippingOverrideChangedDelegate;
+	FChaosVDSettingChaged PlaybackSettingsChangedDelegate;
 };
diff --git a/Engine/Plugins/Experimental/ChaosVD/Source/ChaosVD/Public/ChaosVDPlaybackController.h b/Engine/Plugins/Experimental/ChaosVD/Source/ChaosVD/Public/ChaosVDPlaybackController.h
--- a/Engine/Plugins/Experimental/ChaosVD/Source/ChaosVD/Public/ChaosVDPlaybackController.h
+++ b/Engine/Plugins/Experimental/ChaosVD/Source/ChaosVD/Public/ChaosVDPlaybackController.h
@@ -10,6 +10,7 @@
 #include "Delegates/Delegate.h"
 #include "HAL/ThreadSafeBool.h"
 
+class UChaosVDEditorSettings;
 struct FChaosVDTraceSessionDescriptor;
 struct FChaosVDTrackInfo;
 class FChaosVDScene;
@@ -70,6 +71,8 @@ public:
 
 	/** ID used for the Game Track */
 	static constexpr int32 GameTrackID  = 0 ;
+	static constexpr int32 InvalidFrameRateOverride  = -1 ;
+	static constexpr float FallbackFrameTime = 1.0f / 60.0f ;
 
 	FChaosVDPlaybackController(const TWeakPtr<FChaosVDScene>& InSceneToControl);
 	virtual ~FChaosVDPlaybackController() override;
@@ -223,6 +226,9 @@ public:
 	void RequestUnpause() { bPauseRequested = false; }
 	bool HasPauseRequest() const { return bPauseRequested; }
 
+	float GetFrameTimeOverride() const;
+	float GetFrameTimeForTrack(EChaosVDTrackType TrackType, int32 TrackID, const FChaosVDTrackInfo& TrackInfo) const;
+
 protected:
 
 	/** Updates (or adds) solvers data from the loaded recording to the solver tracks */
@@ -246,6 +252,8 @@ protected:
 	/** Add the provided Geometry info data to the queue. The update will be broadcast in the game thread */
 	void EnqueueGeometryDataUpdate(const Chaos::FConstImplicitObjectPtr& NewGeometry, const uint32 GeometryID);
 
+	void HandleFrameRateOverrideSettingsChanged(UChaosVDEditorSettings* CVDSettings);
+
 	/** Map containing all track info, by track type*/
 	TMap<EChaosVDTrackType, TrackInfoByIDMap> TrackInfoPerType;
 
@@ -275,6 +283,8 @@ protected:
 	int32 MaxFramesLaggingBehindDuringLiveSession = 50;
 	int32 MinFramesLaggingBehindDuringLiveSession = 5;
 
+	int32 CurrentFrameRateOverride = InvalidFrameRateOverride;
+
 	bool bPauseRequested = false;
 
 	FDelegateHandle RecordingStoppedHandle;
