[ChaosVD] Fixed an issue where the Collision Data Inspector will not show data after being closed and automatically open when clicking a contact in the viewport, or using the Show Collision Data in the particle's details panel.

This also addresses an issue where the data was not being shown when we only have a Recorded midphase for a particle pair, but not constraint/contact data. I added a fallback to just show the midplane data in that case.

diff --git a/Engine/Plugins/Experimental/ChaosVD/Source/ChaosVD/Private/ChaosVDCollisionDataDetailsTab.cpp b/Engine/Plugins/Experimental/ChaosVD/Source/ChaosVD/Private/ChaosVDCollisionDataDetailsTab.cpp
--- a/Engine/Plugins/Experimental/ChaosVD/Source/ChaosVD/Private/ChaosVDCollisionDataDetailsTab.cpp
+++ b/Engine/Plugins/Experimental/ChaosVD/Source/ChaosVD/Private/ChaosVDCollisionDataDetailsTab.cpp
@@ -13,9 +13,8 @@ FChaosVDCollisionDataDetailsTab::~FChaosVDCollisionDataDetailsTab()
 {
 }
 
-TSharedRef<SDockTab> FChaosVDCollisionDataDetailsTab::HandleTabSpawned(const FSpawnTabArgs& Args)
+TSharedRef<SDockTab> FChaosVDCollisionDataDetailsTab::HandleTabSpawnRequest(const FSpawnTabArgs& Args)
 {
-
 	TSharedRef<SDockTab> DetailsPanelTab =
 		SNew(SDockTab)
 		.TabRole(ETabRole::MajorTab)
@@ -27,9 +26,16 @@ TSharedRef<SDockTab> FChaosVDCollisionDataDetailsTab::HandleTabSpawned(const FSp
 		SAssignNew(CollisionDataInspector, SChaosVDCollisionDataInspector, GetChaosVDScene())
 	);
 
-	OnTabSpawned().Broadcast(DetailsPanelTab);
+	HandleTabSpawned(DetailsPanelTab);
 
 	return DetailsPanelTab;
 }
 
+void FChaosVDCollisionDataDetailsTab::HandleTabClosed(TSharedRef<SDockTab> InTabClosed)
+{
+	FChaosVDTabSpawnerBase::HandleTabClosed(InTabClosed);
+
+	CollisionDataInspector.Reset();
+}
+
 #undef LOCTEXT_NAMESPACE
diff --git a/Engine/Plugins/Experimental/ChaosVD/Source/ChaosVD/Private/ChaosVDEditorVisualizationSettingsTab.cpp b/Engine/Plugins/Experimental/ChaosVD/Source/ChaosVD/Private/ChaosVDEditorVisualizationSettingsTab.cpp
--- a/Engine/Plugins/Experimental/ChaosVD/Source/ChaosVD/Private/ChaosVDEditorVisualizationSettingsTab.cpp
+++ b/Engine/Plugins/Experimental/ChaosVD/Source/ChaosVD/Private/ChaosVDEditorVisualizationSettingsTab.cpp
@@ -20,7 +20,7 @@
 
 #define LOCTEXT_NAMESPACE "ChaosVisualDebugger"
 
-TSharedRef<SDockTab> FChaosVDEditorVisualizationSettingsTab::HandleTabSpawned(const FSpawnTabArgs& Args)
+TSharedRef<SDockTab> FChaosVDEditorVisualizationSettingsTab::HandleTabSpawnRequest(const FSpawnTabArgs& Args)
 {
 
 	TSharedRef<SDockTab> DetailsPanelTab = SNew(SDockTab)
@@ -67,7 +67,7 @@ TSharedRef<SDockTab> FChaosVDEditorVisualizationSettingsTab::HandleTabSpawned(co
 
 	DetailsPanelTab->SetTabIcon(FChaosVDStyle::Get().GetBrush("TabIconDetailsPanel"));
 
-	OnTabSpawned().Broadcast(DetailsPanelTab);
+	HandleTabSpawned(DetailsPanelTab);
 
 	return DetailsPanelTab;
 }
diff --git a/Engine/Plugins/Experimental/ChaosVD/Source/ChaosVD/Private/ChaosVDObjectDetailsTab.cpp b/Engine/Plugins/Experimental/ChaosVD/Source/ChaosVD/Private/ChaosVDObjectDetailsTab.cpp
--- a/Engine/Plugins/Experimental/ChaosVD/Source/ChaosVD/Private/ChaosVDObjectDetailsTab.cpp
+++ b/Engine/Plugins/Experimental/ChaosVD/Source/ChaosVD/Private/ChaosVDObjectDetailsTab.cpp
@@ -24,40 +24,58 @@
 
 class SSubobjectEditor;
 
-TSharedRef<SDockTab> FChaosVDObjectDetailsTab::HandleTabSpawned(const FSpawnTabArgs& Args)
+TSharedRef<SDockTab> FChaosVDObjectDetailsTab::HandleTabSpawnRequest(const FSpawnTabArgs& Args)
 {
-	TSharedPtr<FChaosVDScene> ScenePtr = GetChaosVDScene().Pin();
-	check(ScenePtr);
-
-	RegisterSelectionSetObject(ScenePtr->GetElementSelectionSet());
 
 	TSharedRef<SDockTab> DetailsPanelTab =
-		SNew(SDockTab)
-		.TabRole(ETabRole::MajorTab)
-		.Label(LOCTEXT("DetailsPanel", "Details"))
-		.ToolTipText(LOCTEXT("DetailsPanelToolTip", "See the details of the selected object"));
-
-	DetailsPanelTab->SetContent
-	(
-		SNew(SVerticalBox)
-		+SVerticalBox::Slot()
-		[
-			SAssignNew(DetailsPanelView, SChaosVDDetailsView)
-		]
-		+SVerticalBox::Slot()
-		.AutoHeight()
-		[
-			GenerateShowCollisionDataButton().ToSharedRef()
-		]
-	);
+	SNew(SDockTab)
+	.TabRole(ETabRole::MajorTab)
+	.Label(LOCTEXT("DetailsPanel", "Details"))
+	.ToolTipText(LOCTEXT("DetailsPanelToolTip", "See the details of the selected object"));
+
+	if (const TSharedPtr<FChaosVDScene> ScenePtr = GetChaosVDScene().Pin())
+	{
+		RegisterSelectionSetObject(ScenePtr->GetElementSelectionSet());
+
+		DetailsPanelTab->SetContent
+		(
+			SNew(SVerticalBox)
+			+SVerticalBox::Slot()
+			[
+				SAssignNew(DetailsPanelView, SChaosVDDetailsView)
+			]
+			+SVerticalBox::Slot()
+			.AutoHeight()
+			[
+				GenerateShowCollisionDataButton().ToSharedRef()
+			]
+		);
+
+		// If we closed the tab and opened it again with an object already selected, try to restore the selected object view
+		if (DetailsPanelView.IsValid() && CurrentSelectedObject.IsValid())
+		{
+			DetailsPanelView->SetSelectedObject(CurrentSelectedObject.Get());
+		}
+	}
+	else
+	{
+		DetailsPanelTab->SetContent(GenerateErrorWidget());
+	}
 
 	DetailsPanelTab->SetTabIcon(FChaosVDStyle::Get().GetBrush("TabIconDetailsPanel"));
 	
-	OnTabSpawned().Broadcast(DetailsPanelTab);
+	HandleTabSpawned(DetailsPanelTab);
 	
 	return DetailsPanelTab;
 }
 
+void FChaosVDObjectDetailsTab::HandleTabClosed(TSharedRef<SDockTab> InTabClosed)
+{
+	FChaosVDTabSpawnerBase::HandleTabClosed(InTabClosed);
+	
+	DetailsPanelView.Reset();
+}
+
 void FChaosVDObjectDetailsTab::HandlePostSelectionChange(const UTypedElementSelectionSet* ChangedSelectionSet)
 {
 	TArray<AActor*> SelectedActors = ChangedSelectionSet->GetSelectedObjects<AActor>();
@@ -68,7 +86,11 @@ void FChaosVDObjectDetailsTab::HandlePostSelectionChange(const UTypedElementSele
 		ensure(SelectedActors.Num() == 1);
 
 		CurrentSelectedObject = SelectedActors[0];
-		DetailsPanelView->SetSelectedObject(CurrentSelectedObject.Get());
+
+		if (DetailsPanelView)
+		{
+			DetailsPanelView->SetSelectedObject(CurrentSelectedObject.Get());
+		}
 	}
 	else
 	{
@@ -146,13 +168,14 @@ FReply FChaosVDObjectDetailsTab::ShowCollisionDataForSelectedObject()
 		return FReply::Handled();
 	}
 
-	if (TSharedPtr<FChaosVDCollisionDataDetailsTab> CollisionDataTab = OwningTabPtr->GetTabSpawnerInstance<FChaosVDCollisionDataDetailsTab>(FChaosVDTabID::CollisionDataDetails).Pin())
+	if (const TSharedPtr<FChaosVDCollisionDataDetailsTab> CollisionDataTab = OwningTabPtr->GetTabSpawnerInstance<FChaosVDCollisionDataDetailsTab>(FChaosVDTabID::CollisionDataDetails).Pin())
 	{
-		if (TSharedPtr<SChaosVDCollisionDataInspector> CollisionInspector = CollisionDataTab->GetCollisionInspectorInstance().Pin())
+		if (const TSharedPtr<FTabManager> TabManager = OwningTabPtr->GetTabManager())
 		{
-			if (TSharedPtr<FTabManager> TabManager = OwningTabPtr->GetTabManager())
+			TabManager->TryInvokeTab(FChaosVDTabID::CollisionDataDetails);
+
+			if (const TSharedPtr<SChaosVDCollisionDataInspector> CollisionInspector = CollisionDataTab->GetCollisionInspectorInstance().Pin())
 			{
-				TabManager->TryInvokeTab(FChaosVDTabID::CollisionDataDetails);
 				CollisionInspector->SetCollisionDataProviderObjectToInspect(CollisionDataProvider);
 			}
 		}
diff --git a/Engine/Plugins/Experimental/ChaosVD/Source/ChaosVD/Private/ChaosVDOutputLogTab.cpp b/Engine/Plugins/Experimental/ChaosVD/Source/ChaosVD/Private/ChaosVDOutputLogTab.cpp
--- a/Engine/Plugins/Experimental/ChaosVD/Source/ChaosVD/Private/ChaosVDOutputLogTab.cpp
+++ b/Engine/Plugins/Experimental/ChaosVD/Source/ChaosVD/Private/ChaosVDOutputLogTab.cpp
@@ -9,7 +9,7 @@
 
 #define LOCTEXT_NAMESPACE "ChaosVisualDebugger"
 
-TSharedRef<SDockTab> FChaosVDOutputLogTab::HandleTabSpawned(const FSpawnTabArgs& Args)
+TSharedRef<SDockTab> FChaosVDOutputLogTab::HandleTabSpawnRequest(const FSpawnTabArgs& Args)
 {
 	FOutputLogCreationParams Params;
 	Params.bCreateDockInLayoutButton = true;
@@ -27,7 +27,7 @@ TSharedRef<SDockTab> FChaosVDOutputLogTab::HandleTabSpawned(const FSpawnTabArgs&
 		FOutputLogModule::Get().MakeOutputLogWidget(Params)
 	);
 
-	OnTabSpawned().Broadcast(OutputLogTab);
+	HandleTabSpawned(OutputLogTab);
 
 	return OutputLogTab;
 }
diff --git a/Engine/Plugins/Experimental/ChaosVD/Source/ChaosVD/Private/ChaosVDPlaybackViewportTab.cpp b/Engine/Plugins/Experimental/ChaosVD/Source/ChaosVD/Private/ChaosVDPlaybackViewportTab.cpp
--- a/Engine/Plugins/Experimental/ChaosVD/Source/ChaosVD/Private/ChaosVDPlaybackViewportTab.cpp
+++ b/Engine/Plugins/Experimental/ChaosVD/Source/ChaosVD/Private/ChaosVDPlaybackViewportTab.cpp
@@ -12,7 +12,7 @@
 
 #define LOCTEXT_NAMESPACE "ChaosVisualDebugger"
 
-TSharedRef<SDockTab> FChaosVDPlaybackViewportTab::HandleTabSpawned(const FSpawnTabArgs& Args)
+TSharedRef<SDockTab> FChaosVDPlaybackViewportTab::HandleTabSpawnRequest(const FSpawnTabArgs& Args)
 {
 	TSharedRef<SDockTab> ViewportTab =
 		SNew(SDockTab)
@@ -33,8 +33,17 @@ TSharedRef<SDockTab> FChaosVDPlaybackViewportTab::HandleTabSpawned(const FSpawnT
 	}
 	
 	ViewportTab->SetTabIcon(FChaosVDStyle::Get().GetBrush("TabIconPlaybackViewport"));
+	
+	HandleTabSpawned(ViewportTab);
 
 	return ViewportTab;
 }
 
-#undef LOCTEXT_NAMESPACE
\ No newline at end of file
+void FChaosVDPlaybackViewportTab::HandleTabClosed(TSharedRef<SDockTab> InTabClosed)
+{
+	FChaosVDTabSpawnerBase::HandleTabClosed(InTabClosed);
+
+	PlaybackViewportWidget.Reset();
+}
+
+#undef LOCTEXT_NAMESPACE
diff --git a/Engine/Plugins/Experimental/ChaosVD/Source/ChaosVD/Private/ChaosVDSceneQueryDataInspectorTab.cpp b/Engine/Plugins/Experimental/ChaosVD/Source/ChaosVD/Private/ChaosVDSceneQueryDataInspectorTab.cpp
--- a/Engine/Plugins/Experimental/ChaosVD/Source/ChaosVD/Private/ChaosVDSceneQueryDataInspectorTab.cpp
+++ b/Engine/Plugins/Experimental/ChaosVD/Source/ChaosVD/Private/ChaosVDSceneQueryDataInspectorTab.cpp
@@ -12,7 +12,7 @@ FChaosVDSceneQueryDataInspectorTab::~FChaosVDSceneQueryDataInspectorTab()
 {
 }
 
-TSharedRef<SDockTab> FChaosVDSceneQueryDataInspectorTab::HandleTabSpawned(const FSpawnTabArgs& Args)
+TSharedRef<SDockTab> FChaosVDSceneQueryDataInspectorTab::HandleTabSpawnRequest(const FSpawnTabArgs& Args)
 {
 	TSharedRef<SDockTab> DetailsPanelTab =
 	SNew(SDockTab)
@@ -32,9 +32,16 @@ TSharedRef<SDockTab> FChaosVDSceneQueryDataInspectorTab::HandleTabSpawned(const
 		DetailsPanelTab->SetContent(GenerateErrorWidget());
 	}
 
-	OnTabSpawned().Broadcast(DetailsPanelTab);
+	HandleTabSpawned(DetailsPanelTab);
 
 	return DetailsPanelTab;
 }
 
+void FChaosVDSceneQueryDataInspectorTab::HandleTabClosed(TSharedRef<SDockTab> InTabClosed)
+{
+	FChaosVDTabSpawnerBase::HandleTabClosed(InTabClosed);
+
+	SceneQueryDataInspector.Reset();
+}
+
 #undef LOCTEXT_NAMESPACE
diff --git a/Engine/Plugins/Experimental/ChaosVD/Source/ChaosVD/Private/ChaosVDSolversTracksTab.cpp b/Engine/Plugins/Experimental/ChaosVD/Source/ChaosVD/Private/ChaosVDSolversTracksTab.cpp
--- a/Engine/Plugins/Experimental/ChaosVD/Source/ChaosVD/Private/ChaosVDSolversTracksTab.cpp
+++ b/Engine/Plugins/Experimental/ChaosVD/Source/ChaosVD/Private/ChaosVDSolversTracksTab.cpp
@@ -10,9 +10,9 @@
 
 #define LOCTEXT_NAMESPACE "ChaosVisualDebugger"
 
-TSharedRef<SDockTab> FChaosVDSolversTracksTab::HandleTabSpawned(const FSpawnTabArgs& Args)
+TSharedRef<SDockTab> FChaosVDSolversTracksTab::HandleTabSpawnRequest(const FSpawnTabArgs& Args)
 {
-	TSharedRef<SDockTab> ViewportTab =
+	TSharedRef<SDockTab> SolverTracksTab =
 	SNew(SDockTab)
 	.TabRole(ETabRole::PanelTab)
 	.Label(LOCTEXT("SolverTracksTabLabel", "Available Solvers"))
@@ -20,21 +20,26 @@ TSharedRef<SDockTab> FChaosVDSolversTracksTab::HandleTabSpawned(const FSpawnTabA
 
 	if (const TSharedPtr<SChaosVDMainTab> MainTabPtr = OwningTabWidget.Pin())
 	{
-		ViewportTab->SetContent
+		SolverTracksTab->SetContent
 		(
 			SAssignNew(SolverTracksWidget, SChaosVDSolverTracks, MainTabPtr->GetChaosVDEngineInstance()->GetPlaybackController())
 		);
 	}
 	else
 	{
-		ViewportTab->SetContent(GenerateErrorWidget());
+		SolverTracksTab->SetContent(GenerateErrorWidget());
 	}
 
-	ViewportTab->SetTabIcon(FChaosVDStyle::Get().GetBrush("TabIconPlaybackViewport"));
+	SolverTracksTab->SetTabIcon(FChaosVDStyle::Get().GetBrush("TabIconPlaybackViewport"));
 
-	OnTabSpawned().Broadcast(ViewportTab);
+	HandleTabSpawned(SolverTracksTab);
 
-	return ViewportTab;
+	return SolverTracksTab;
+}
+
+void FChaosVDSolversTracksTab::HandleTabClosed(TSharedRef<SDockTab> InTabClosed)
+{
+	SolverTracksWidget.Reset();
 }
 
 #undef LOCTEXT_NAMESPACE
diff --git a/Engine/Plugins/Experimental/ChaosVD/Source/ChaosVD/Private/ChaosVDTabSpawnerBase.cpp b/Engine/Plugins/Experimental/ChaosVD/Source/ChaosVD/Private/ChaosVDTabSpawnerBase.cpp
--- a/Engine/Plugins/Experimental/ChaosVD/Source/ChaosVD/Private/ChaosVDTabSpawnerBase.cpp
+++ b/Engine/Plugins/Experimental/ChaosVD/Source/ChaosVD/Private/ChaosVDTabSpawnerBase.cpp
@@ -12,12 +12,19 @@
 FChaosVDTabSpawnerBase::FChaosVDTabSpawnerBase(const FName& InTabID, TSharedPtr<FTabManager> InTabManager, TWeakPtr<SChaosVDMainTab> InOwningTabWidget)
 {
 	OwningTabWidget = InOwningTabWidget;
-	InTabManager->RegisterTabSpawner(InTabID, FOnSpawnTab::CreateRaw(this, &FChaosVDTabSpawnerBase::HandleTabSpawned));
+	InTabManager->RegisterTabSpawner(InTabID, FOnSpawnTab::CreateRaw(this, &FChaosVDTabSpawnerBase::HandleTabSpawnRequest));
 }
 
-void FChaosVDTabSpawnerBase::HandleTabClosed(const TSharedRef<SDockTab>& InTabClosed)
+void FChaosVDTabSpawnerBase::HandleTabClosed(TSharedRef<SDockTab> InTabClosed)
 {
 	OnTabDestroyed().Broadcast(InTabClosed);
+	InTabClosed->SetOnTabClosed(SDockTab::FOnTabClosedCallback());
+}
+
+void FChaosVDTabSpawnerBase::HandleTabSpawned(TSharedRef<SDockTab> InTabClosed)
+{
+	OnTabSpawned().Broadcast(InTabClosed);
+	InTabClosed->SetOnTabClosed(SDockTab::FOnTabClosedCallback::CreateRaw(this, &FChaosVDTabSpawnerBase::HandleTabClosed));
 }
 
 TSharedRef<SWidget> FChaosVDTabSpawnerBase::GenerateErrorWidget()
diff --git a/Engine/Plugins/Experimental/ChaosVD/Source/ChaosVD/Private/ChaosVDWorldOutlinerTab.cpp b/Engine/Plugins/Experimental/ChaosVD/Source/ChaosVD/Private/ChaosVDWorldOutlinerTab.cpp
--- a/Engine/Plugins/Experimental/ChaosVD/Source/ChaosVD/Private/ChaosVDWorldOutlinerTab.cpp
+++ b/Engine/Plugins/Experimental/ChaosVD/Source/ChaosVD/Private/ChaosVDWorldOutlinerTab.cpp
@@ -41,7 +41,7 @@ void FChaosVDWorldOutlinerTab::CreateWorldOutlinerWidget()
 	SceneOutlinerWidget = SceneOutlinerModule.CreateSceneOutliner(InitOptions);
 }
 
-TSharedRef<SDockTab> FChaosVDWorldOutlinerTab::HandleTabSpawned(const FSpawnTabArgs& Args)
+TSharedRef<SDockTab> FChaosVDWorldOutlinerTab::HandleTabSpawnRequest(const FSpawnTabArgs& Args)
 {
 	CreateWorldOutlinerWidget();
 
@@ -58,9 +58,16 @@ TSharedRef<SDockTab> FChaosVDWorldOutlinerTab::HandleTabSpawned(const FSpawnTabA
 
 	OutlinerTab->SetTabIcon(FChaosVDStyle::Get().GetBrush("TabIconWorldOutliner"));
 
-	OnTabSpawned().Broadcast(OutlinerTab);
+	HandleTabSpawned(OutlinerTab);
 
 	return OutlinerTab;
 }
 
+void FChaosVDWorldOutlinerTab::HandleTabClosed(TSharedRef<SDockTab> InTabClosed)
+{
+	FChaosVDTabSpawnerBase::HandleTabClosed(InTabClosed);
+
+	SceneOutlinerWidget.Reset();
+}
+
 #undef LOCTEXT_NAMESPACE
diff --git a/Engine/Plugins/Experimental/ChaosVD/Source/ChaosVD/Private/Visualizers/ChaosVDSolverCollisionDataComponentVisualizer.cpp b/Engine/Plugins/Experimental/ChaosVD/Source/ChaosVD/Private/Visualizers/ChaosVDSolverCollisionDataComponentVisualizer.cpp
--- a/Engine/Plugins/Experimental/ChaosVD/Source/ChaosVD/Private/Visualizers/ChaosVDSolverCollisionDataComponentVisualizer.cpp
+++ b/Engine/Plugins/Experimental/ChaosVD/Source/ChaosVD/Private/Visualizers/ChaosVDSolverCollisionDataComponentVisualizer.cpp
@@ -110,21 +110,23 @@ bool FChaosVDSolverCollisionDataComponentVisualizer::VisProxyHandleClick(FEditor
 		return false;
 	}
 
-	const TSharedPtr<SChaosVDCollisionDataInspector> CollisionInspector = CollisionDataDetailsTab->GetCollisionInspectorInstance().Pin();
-	if (!CollisionInspector)
+	const TSharedPtr<FTabManager> TabManager = MainTabToolkitHost->GetTabManager();
+	if (!TabManager)
 	{
 		return false;
 	}
 
-	const TSharedPtr<FTabManager> TabManager = MainTabToolkitHost->GetTabManager();
-	if (!TabManager)
+	TabManager->TryInvokeTab(FChaosVDTabID::CollisionDataDetails);
+	
+	const TSharedPtr<SChaosVDCollisionDataInspector> CollisionInspector = CollisionDataDetailsTab->GetCollisionInspectorInstance().Pin();
+	if (!CollisionInspector)
 	{
 		return false;
 	}
-	
+
 	if (ContactProxy->ContactFinder.OwningMidPhase.Pin())
 	{
-		TabManager->TryInvokeTab(FChaosVDTabID::CollisionDataDetails);
+		
 		CollisionInspector->SetSingleContactDataToInspect(ContactProxy->ContactFinder);
 		
 		ClearCurrentSelection();
diff --git a/Engine/Plugins/Experimental/ChaosVD/Source/ChaosVD/Private/Widgets/SChaosVDCollisionDataInspector.cpp b/Engine/Plugins/Experimental/ChaosVD/Source/ChaosVD/Private/Widgets/SChaosVDCollisionDataInspector.cpp
--- a/Engine/Plugins/Experimental/ChaosVD/Source/ChaosVD/Private/Widgets/SChaosVDCollisionDataInspector.cpp
+++ b/Engine/Plugins/Experimental/ChaosVD/Source/ChaosVD/Private/Widgets/SChaosVDCollisionDataInspector.cpp
@@ -34,8 +34,8 @@ void SChaosVDCollisionDataInspector::Construct(const FArguments& InArgs, const T
 		ScenePtr->OnSceneUpdated().AddRaw(this, &SChaosVDCollisionDataInspector::HandleSceneUpdated);
 	}
 
-	ContactDataDetailsView = CreateCollisionDataDetailsView();
-	ConstraintDataDetailsView = CreateCollisionDataDetailsView();
+	MainCollisionDataDetailsView = CreateCollisionDataDetailsView();
+	SecondaryCollisionDataDetailsPanel = CreateCollisionDataDetailsView();
 
 	ChildSlot
 	[
@@ -51,7 +51,7 @@ void SChaosVDCollisionDataInspector::Construct(const FArguments& InArgs, const T
 		.AutoHeight()
 		[
 			SAssignNew(CollisionDataAvailableList, SChaosVDNameListPicker)
-			.OnNameSleceted_Raw(this, &SChaosVDCollisionDataInspector::HandleSessionNameSelected)
+			.OnNameSleceted_Raw(this, &SChaosVDCollisionDataInspector::HandleCollisionDataEntryNameSelected)
 		]
 		+SVerticalBox::Slot()
 		.Padding(10)
@@ -95,12 +95,12 @@ void SChaosVDCollisionDataInspector::Construct(const FArguments& InArgs, const T
 			+SScrollBox::Slot()
 			.Padding(15.0f,0.0f,15.0f,0.0f)
 			[
-				ContactDataDetailsView->GetWidget().ToSharedRef()
+				MainCollisionDataDetailsView->GetWidget().ToSharedRef()
 			]
 			+SScrollBox::Slot()
 			.Padding(15.0f,10.0f,15.0f,5.0f)
 			[
-				ConstraintDataDetailsView->GetWidget().ToSharedRef()
+				SecondaryCollisionDataDetailsPanel->GetWidget().ToSharedRef()
 			]
 		]
 	];
@@ -171,8 +171,8 @@ void SChaosVDCollisionDataInspector::HandleSceneUpdated()
 
 void SChaosVDCollisionDataInspector::ClearInspector()
 {
-	ContactDataDetailsView->SetStructureData(nullptr);
-	ConstraintDataDetailsView->SetStructureData(nullptr);
+	MainCollisionDataDetailsView->SetStructureData(nullptr);
+	SecondaryCollisionDataDetailsPanel->SetStructureData(nullptr);
 
 	CollisionDataByNameMap.Reset();
 	CollisionDataAvailableList->UpdateNameList({});
@@ -239,7 +239,7 @@ EVisibility SChaosVDCollisionDataInspector::GetOutOfDateWarningVisibility() cons
 	return bIsUpToDate ? EVisibility::Collapsed : EVisibility::Visible;
 }
 
-void SChaosVDCollisionDataInspector::HandleSessionNameSelected(TSharedPtr<FName> SelectedName)
+void SChaosVDCollisionDataInspector::HandleCollisionDataEntryNameSelected(TSharedPtr<FName> SelectedName)
 {
 	CurrentSelectedName = SelectedName;
 
@@ -254,33 +254,58 @@ void SChaosVDCollisionDataInspector::HandleSessionNameSelected(TSharedPtr<FName>
 		{
 			FChaosVDConstraint* MutableConstraint = const_cast<FChaosVDConstraint*>(DataFinderPtr->OwningConstraint);
 			TSharedPtr<FStructOnScope> ConstraintView = MakeShared<FStructOnScope>(FChaosVDConstraint::StaticStruct(), reinterpret_cast<uint8*>(MutableConstraint));
-			ConstraintDataDetailsView->SetStructureData(ConstraintView); 
+			SecondaryCollisionDataDetailsPanel->SetStructureData(ConstraintView); 
 
 			if (DataFinderPtr->OwningConstraint->ManifoldPoints.IsValidIndex(DataFinderPtr->ContactIndex))
 			{
 				const FChaosVDManifoldPoint* ContactPointData = &DataFinderPtr->OwningConstraint->ManifoldPoints[DataFinderPtr->ContactIndex];
 				TSharedPtr<FStructOnScope> ContactViewView = MakeShared<FStructOnScope>(FChaosVDManifoldPoint::StaticStruct(), reinterpret_cast<uint8*>(const_cast<FChaosVDManifoldPoint*>(ContactPointData)));
-				ContactDataDetailsView->SetStructureData(ContactViewView);
+				MainCollisionDataDetailsView->SetStructureData(ContactViewView);
 			}
 			else
 			{
-				ContactDataDetailsView->SetStructureData(nullptr); 
+				MainCollisionDataDetailsView->SetStructureData(nullptr); 
 			}
 		}
+		else if (TSharedPtr<FChaosVDParticlePairMidPhase> MidPhaseData = DataFinderPtr->OwningMidPhase.Pin())
+		{
+			// If we have a recorded midphase with not contact data, show that instead
+			const TSharedPtr<FStructOnScope> MidPhaseView = MakeShared<FStructOnScope>(FChaosVDParticlePairMidPhase::StaticStruct(), reinterpret_cast<uint8*>(MidPhaseData.Get()));
+			MainCollisionDataDetailsView->SetStructureData(MidPhaseView); 
+		}
 	}
 }
 
 TSharedPtr<FName> SChaosVDCollisionDataInspector::GenerateNameForCollisionDataItem(const FChaosVDCollisionDataFinder& InContactFinderData)
 {
-	if (InContactFinderData.OwningMidPhase.Pin())
+	int32 Particle0Index = INDEX_NONE;
+	int32 Particle1Index = INDEX_NONE;
+
+	bool bHasConstraintData = false;
+
+	if (const TSharedPtr<FChaosVDParticlePairMidPhase> MidPhaseData = InContactFinderData.OwningMidPhase.Pin())
 	{
 		if (InContactFinderData.OwningConstraint)
 		{
-			const FText GeneratedName = FText::Format(LOCTEXT("CollisionIemDataTitle", "Particle Pair | Index0 [ID {0}] <-> Index1 [{1}]"), InContactFinderData.OwningConstraint->Particle0Index, InContactFinderData.OwningConstraint->Particle1Index);
-			return MakeShared<FName>(GeneratedName.ToString());
+			Particle0Index = InContactFinderData.OwningConstraint->Particle0Index;
+			Particle1Index = InContactFinderData.OwningConstraint->Particle1Index;
+			bHasConstraintData = true;
+		}
+		else
+		{
+			Particle0Index = MidPhaseData->Particle0Idx;
+			Particle1Index = MidPhaseData->Particle1Idx;	
 		}
 	}
 
+	if (Particle0Index != INDEX_NONE || Particle1Index != INDEX_NONE)
+	{
+		static FText MidPhaseOnlyText = LOCTEXT("CollisionItemDataTitleMidPhaseOnly", "[MidPhase Only]");
+		static FText ConstraintText = LOCTEXT("CollisionItemDataTitleConstaint", "[Constraint]");
+		const FText GeneratedName = FText::Format(LOCTEXT("CollisionItemDataTitle", "Particle Pair | Index0 [ID {0}] <-> Index1 [{1}] | Type {2}"), Particle0Index, Particle1Index, bHasConstraintData ? ConstraintText: MidPhaseOnlyText);
+		return MakeShared<FName>(GeneratedName.ToString());
+	}
+
 	return nullptr;
 }
 
@@ -311,17 +336,23 @@ FReply SChaosVDCollisionDataInspector::SelectParticleForCurrentCollisionData(ECh
 		return FReply::Handled();
 	}
 	
-	if (TSharedPtr<FChaosVDCollisionDataFinder> DataFinderPtr = GetCurrentDataBeingInspected())
+	if (const TSharedPtr<FChaosVDCollisionDataFinder> DataFinderPtr = GetCurrentDataBeingInspected())
 	{
-		if (FChaosVDMidPhasePtr MidPhasePtr = DataFinderPtr->OwningMidPhase.Pin())
+		if (const FChaosVDMidPhasePtr MidPhasePtr = DataFinderPtr->OwningMidPhase.Pin())
 		{
+			int32 ParticleIndex = INDEX_NONE;
 			if (DataFinderPtr->OwningConstraint)
 			{
-				const int32 ParticleIndex = ParticleSelector == EChaosVDCollisionParticleSelector::Index_0 ? DataFinderPtr->OwningConstraint->Particle0Index : DataFinderPtr->OwningConstraint->Particle1Index;
-				if (AChaosVDParticleActor* ParticleActor = ScenePtr->GetParticleActor(MidPhasePtr->SolverID, ParticleIndex))
-				{
-					ScenePtr->SetSelectedObject(ParticleActor);
-				}
+				ParticleIndex = ParticleSelector == EChaosVDCollisionParticleSelector::Index_0 ? DataFinderPtr->OwningConstraint->Particle0Index : DataFinderPtr->OwningConstraint->Particle1Index;
+			}
+			else
+			{
+				ParticleIndex = ParticleSelector == EChaosVDCollisionParticleSelector::Index_0 ? MidPhasePtr->Particle0Idx : MidPhasePtr->Particle1Idx;
+			}
+
+			if (AChaosVDParticleActor* ParticleActor = ScenePtr->GetParticleActor(MidPhasePtr->SolverID, ParticleIndex))
+			{
+				ScenePtr->SetSelectedObject(ParticleActor);
 			}
 		}
 	}
diff --git a/Engine/Plugins/Experimental/ChaosVD/Source/ChaosVD/Public/ChaosVDCollisionDataDetailsTab.h b/Engine/Plugins/Experimental/ChaosVD/Source/ChaosVD/Public/ChaosVDCollisionDataDetailsTab.h
--- a/Engine/Plugins/Experimental/ChaosVD/Source/ChaosVD/Public/ChaosVDCollisionDataDetailsTab.h
+++ b/Engine/Plugins/Experimental/ChaosVD/Source/ChaosVD/Public/ChaosVDCollisionDataDetailsTab.h
@@ -18,7 +18,8 @@ public:
 	}
 
 	virtual ~FChaosVDCollisionDataDetailsTab() override;
-	virtual TSharedRef<SDockTab> HandleTabSpawned(const FSpawnTabArgs& Args) override;
+	virtual TSharedRef<SDockTab> HandleTabSpawnRequest(const FSpawnTabArgs& Args) override;
+	virtual void HandleTabClosed(TSharedRef<SDockTab> InTabClosed) override;
 
 	TWeakPtr<SChaosVDCollisionDataInspector> GetCollisionInspectorInstance() const { return CollisionDataInspector; };
 
diff --git a/Engine/Plugins/Experimental/ChaosVD/Source/ChaosVD/Public/ChaosVDEditorVisualizationSettingsTab.h b/Engine/Plugins/Experimental/ChaosVD/Source/ChaosVD/Public/ChaosVDEditorVisualizationSettingsTab.h
--- a/Engine/Plugins/Experimental/ChaosVD/Source/ChaosVD/Public/ChaosVDEditorVisualizationSettingsTab.h
+++ b/Engine/Plugins/Experimental/ChaosVD/Source/ChaosVD/Public/ChaosVDEditorVisualizationSettingsTab.h
@@ -14,5 +14,5 @@ public:
 	{
 	}
 
-	virtual TSharedRef<SDockTab> HandleTabSpawned(const FSpawnTabArgs& Args) override;
+	virtual TSharedRef<SDockTab> HandleTabSpawnRequest(const FSpawnTabArgs& Args) override;
 };
diff --git a/Engine/Plugins/Experimental/ChaosVD/Source/ChaosVD/Public/ChaosVDObjectDetailsTab.h b/Engine/Plugins/Experimental/ChaosVD/Source/ChaosVD/Public/ChaosVDObjectDetailsTab.h
--- a/Engine/Plugins/Experimental/ChaosVD/Source/ChaosVD/Public/ChaosVDObjectDetailsTab.h
+++ b/Engine/Plugins/Experimental/ChaosVD/Source/ChaosVD/Public/ChaosVDObjectDetailsTab.h
@@ -27,7 +27,8 @@ public:
 
 protected:
 
-	virtual TSharedRef<SDockTab> HandleTabSpawned(const FSpawnTabArgs& Args) override;
+	virtual TSharedRef<SDockTab> HandleTabSpawnRequest(const FSpawnTabArgs& Args) override;
+	virtual void HandleTabClosed(TSharedRef<SDockTab> InTabClosed) override;
 
 	virtual void HandlePostSelectionChange(const UTypedElementSelectionSet* ChangedSelectionSet) override;
 
diff --git a/Engine/Plugins/Experimental/ChaosVD/Source/ChaosVD/Public/ChaosVDOutputLogTab.h b/Engine/Plugins/Experimental/ChaosVD/Source/ChaosVD/Public/ChaosVDOutputLogTab.h
--- a/Engine/Plugins/Experimental/ChaosVD/Source/ChaosVD/Public/ChaosVDOutputLogTab.h
+++ b/Engine/Plugins/Experimental/ChaosVD/Source/ChaosVD/Public/ChaosVDOutputLogTab.h
@@ -18,7 +18,7 @@ public:
 	{
 	}
 
-	virtual TSharedRef<SDockTab> HandleTabSpawned(const FSpawnTabArgs& Args) override;
+	virtual TSharedRef<SDockTab> HandleTabSpawnRequest(const FSpawnTabArgs& Args) override;
 
 private:
 	TSharedPtr<FOutputLogHistory> OutputLogHistory;
diff --git a/Engine/Plugins/Experimental/ChaosVD/Source/ChaosVD/Public/ChaosVDPlaybackViewportTab.h b/Engine/Plugins/Experimental/ChaosVD/Source/ChaosVD/Public/ChaosVDPlaybackViewportTab.h
--- a/Engine/Plugins/Experimental/ChaosVD/Source/ChaosVD/Public/ChaosVDPlaybackViewportTab.h
+++ b/Engine/Plugins/Experimental/ChaosVD/Source/ChaosVD/Public/ChaosVDPlaybackViewportTab.h
@@ -18,7 +18,9 @@ public:
 	{	
 	}
 	
-	virtual TSharedRef<SDockTab> HandleTabSpawned(const FSpawnTabArgs& Args) override;
+	virtual TSharedRef<SDockTab> HandleTabSpawnRequest(const FSpawnTabArgs& Args) override;
+
+	virtual void HandleTabClosed(TSharedRef<SDockTab> InTabClosed) override;
 
 	TWeakPtr<SChaosVDPlaybackViewport> GetPlaybackViewportWidget() { return PlaybackViewportWidget; }
 
diff --git a/Engine/Plugins/Experimental/ChaosVD/Source/ChaosVD/Public/ChaosVDSceneQueryDataInspectorTab.h b/Engine/Plugins/Experimental/ChaosVD/Source/ChaosVD/Public/ChaosVDSceneQueryDataInspectorTab.h
--- a/Engine/Plugins/Experimental/ChaosVD/Source/ChaosVD/Public/ChaosVDSceneQueryDataInspectorTab.h
+++ b/Engine/Plugins/Experimental/ChaosVD/Source/ChaosVD/Public/ChaosVDSceneQueryDataInspectorTab.h
@@ -17,7 +17,8 @@ public:
 	}
 
 	virtual ~FChaosVDSceneQueryDataInspectorTab() override;
-	virtual TSharedRef<SDockTab> HandleTabSpawned(const FSpawnTabArgs& Args) override;
+	virtual TSharedRef<SDockTab> HandleTabSpawnRequest(const FSpawnTabArgs& Args) override;
+	virtual void HandleTabClosed(TSharedRef<SDockTab> InTabClosed) override;
 
 	TWeakPtr<SChaosVDSceneQueryDataInspector> GetSceneQueryDataInspectorInstance() const { return SceneQueryDataInspector; }
 
diff --git a/Engine/Plugins/Experimental/ChaosVD/Source/ChaosVD/Public/ChaosVDSolversTracksTab.h b/Engine/Plugins/Experimental/ChaosVD/Source/ChaosVD/Public/ChaosVDSolversTracksTab.h
--- a/Engine/Plugins/Experimental/ChaosVD/Source/ChaosVD/Public/ChaosVDSolversTracksTab.h
+++ b/Engine/Plugins/Experimental/ChaosVD/Source/ChaosVD/Public/ChaosVDSolversTracksTab.h
@@ -15,7 +15,9 @@ public:
 	{	
 	}
 	
-	virtual TSharedRef<SDockTab> HandleTabSpawned(const FSpawnTabArgs& Args) override;
+	virtual TSharedRef<SDockTab> HandleTabSpawnRequest(const FSpawnTabArgs& Args) override;
+
+	virtual void HandleTabClosed(TSharedRef<SDockTab> InTabClosed) override;
 
 protected:
 	TSharedPtr<SChaosVDSolverTracks> SolverTracksWidget;
diff --git a/Engine/Plugins/Experimental/ChaosVD/Source/ChaosVD/Public/ChaosVDTabSpawnerBase.h b/Engine/Plugins/Experimental/ChaosVD/Source/ChaosVD/Public/ChaosVDTabSpawnerBase.h
--- a/Engine/Plugins/Experimental/ChaosVD/Source/ChaosVD/Public/ChaosVDTabSpawnerBase.h
+++ b/Engine/Plugins/Experimental/ChaosVD/Source/ChaosVD/Public/ChaosVDTabSpawnerBase.h
@@ -25,8 +25,9 @@ public:
 
 	FChaosVDTabSpawnerBase(const FName& InTabID, TSharedPtr<FTabManager> InTabManager, TWeakPtr<SChaosVDMainTab> InOwningTabWidget);
 
-	virtual TSharedRef<SDockTab> HandleTabSpawned(const FSpawnTabArgs& Args) = 0;
-	virtual void HandleTabClosed(const TSharedRef<SDockTab>& InTabClosed);
+	virtual TSharedRef<SDockTab> HandleTabSpawnRequest(const FSpawnTabArgs& Args) = 0;
+	virtual void HandleTabClosed(TSharedRef<SDockTab> InTabClosed);
+	virtual void HandleTabSpawned(TSharedRef<SDockTab> InTabClosed);
 
 	FChaosVDTabSpawned& OnTabSpawned() { return TabSpawnedDelegate; }
 	FChaosVDTabDestroyed& OnTabDestroyed() { return TabDestroyedDelegate; }
diff --git a/Engine/Plugins/Experimental/ChaosVD/Source/ChaosVD/Public/ChaosVDWorldOutlinerTab.h b/Engine/Plugins/Experimental/ChaosVD/Source/ChaosVD/Public/ChaosVDWorldOutlinerTab.h
--- a/Engine/Plugins/Experimental/ChaosVD/Source/ChaosVD/Public/ChaosVDWorldOutlinerTab.h
+++ b/Engine/Plugins/Experimental/ChaosVD/Source/ChaosVD/Public/ChaosVDWorldOutlinerTab.h
@@ -30,7 +30,8 @@ public:
 	{
 	}
 
-	virtual TSharedRef<SDockTab> HandleTabSpawned(const FSpawnTabArgs& Args) override;
+	virtual TSharedRef<SDockTab> HandleTabSpawnRequest(const FSpawnTabArgs& Args) override;
+	virtual void HandleTabClosed(TSharedRef<SDockTab> InTabClosed) override;
 
 private:
 	void CreateWorldOutlinerWidget();
diff --git a/Engine/Plugins/Experimental/ChaosVD/Source/ChaosVD/Public/Widgets/SChaosVDCollisionDataInspector.h b/Engine/Plugins/Experimental/ChaosVD/Source/ChaosVD/Public/Widgets/SChaosVDCollisionDataInspector.h
--- a/Engine/Plugins/Experimental/ChaosVD/Source/ChaosVD/Public/Widgets/SChaosVDCollisionDataInspector.h
+++ b/Engine/Plugins/Experimental/ChaosVD/Source/ChaosVD/Public/Widgets/SChaosVDCollisionDataInspector.h
@@ -52,7 +52,7 @@ protected:
 
 	EVisibility GetOutOfDateWarningVisibility() const;
 
-	void HandleSessionNameSelected(TSharedPtr<FName> SelectedName);
+	void HandleCollisionDataEntryNameSelected(TSharedPtr<FName> SelectedName);
 
 	TSharedPtr<FName> GenerateNameForCollisionDataItem(const FChaosVDCollisionDataFinder& InContactFinderData);
 
@@ -70,9 +70,9 @@ protected:
 
 	TSharedPtr<FName> CurrentSelectedName;
 
-	TSharedPtr<IStructureDetailsView> ContactDataDetailsView;
+	TSharedPtr<IStructureDetailsView> MainCollisionDataDetailsView;
 	
-	TSharedPtr<IStructureDetailsView> ConstraintDataDetailsView;
+	TSharedPtr<IStructureDetailsView> SecondaryCollisionDataDetailsPanel;
 
 	FName CurrentObjectBeingInspectedName;
 
