[ChaosVD] Added a way to start recording an auto connect as live session directly from CVD's UI if you are in the editor and recording a PIE session

Also added pop ups informing the user if an error has happened instead of silently fail (or just trigger an ensure). And lastly, added a commandline argument that can be used to start a CVD recording automatically on startup

diff --git a/Engine/Plugins/Experimental/ChaosVD/Source/ChaosVD/Private/ChaosVDEngine.cpp b/Engine/Plugins/Experimental/ChaosVD/Source/ChaosVD/Private/ChaosVDEngine.cpp
--- a/Engine/Plugins/Experimental/ChaosVD/Source/ChaosVD/Private/ChaosVDEngine.cpp
+++ b/Engine/Plugins/Experimental/ChaosVD/Source/ChaosVD/Private/ChaosVDEngine.cpp
@@ -4,6 +4,7 @@
 
 #include "ChaosVDModule.h"
 #include "ChaosVDPlaybackController.h"
+#include "ChaosVDRuntimeModule.h"
 #include "ChaosVDScene.h"
 #include "Trace/ChaosVDTraceManager.h"
 
@@ -21,6 +22,17 @@ void FChaosVDEngine::Initialize()
 
 	PlaybackController = MakeShared<FChaosVDPlaybackController>(CurrentScene);
 
+	// Listen for the recording stop event to clear the live session flag
+	// TODO: We do something similar for the live flag on the CVD recording instance. We should unify both and have a single place where to check the live state
+	// of a session, so we also have one single place to clear the flag
+	LiveSessionStoppedDelegateHandle = FChaosVDRuntimeModule::Get().RegisterRecordingStopCallback(FChaosVDRecordingStateChangedDelegate::FDelegate::CreateLambda([WeakThis = AsWeak()]()
+	{
+		if (const TSharedPtr<FChaosVDEngine> CVDEngine = WeakThis.Pin())
+		{
+			CVDEngine->CurrentSessionDescriptor.bIsLiveSession = false;
+		}
+	}));
+
 	bIsInitialized = true;
 }
 
@@ -40,6 +52,13 @@ void FChaosVDEngine::DeInitialize()
 		CVDTraceManager->CloseSession(CurrentSessionDescriptor.SessionName);
 	}
 
+	if (FChaosVDRuntimeModule::IsLoaded())
+	{
+		FChaosVDRuntimeModule::Get().RemoveRecordingStopCallback(LiveSessionStoppedDelegateHandle);
+	}
+
+	LiveSessionStoppedDelegateHandle = FDelegateHandle();
+
 	CollectGarbage(GARBAGE_COLLECTION_KEEPFLAGS);
 
 	bIsInitialized = false;
diff --git a/Engine/Plugins/Experimental/ChaosVD/Source/ChaosVD/Private/ChaosVDStyle.cpp b/Engine/Plugins/Experimental/ChaosVD/Source/ChaosVD/Private/ChaosVDStyle.cpp
--- a/Engine/Plugins/Experimental/ChaosVD/Source/ChaosVD/Private/ChaosVDStyle.cpp
+++ b/Engine/Plugins/Experimental/ChaosVD/Source/ChaosVD/Private/ChaosVDStyle.cpp
@@ -59,6 +59,8 @@ TSharedRef< FSlateStyleSet > FChaosVDStyle::Create()
 	Style->Set("LockIcon", new IMAGE_BRUSH_SVG(TEXT("lock"), Icon16x16));
 	Style->Set("UnlockedIcon", new IMAGE_BRUSH_SVG(TEXT("lock-unlocked"), Icon16x16));
 	Style->Set("RecordIcon", new IMAGE_BRUSH(TEXT("CircleBox2"), Icon16x16));
+	Style->Set("RecordToFileIcon", new IMAGE_BRUSH(TEXT("RecordToFile"), Icon16x16));
+	Style->Set("RecordToLiveIcon", new IMAGE_BRUSH(TEXT("RecordLiveSession"), Icon16x16));
 	Style->Set("ConnectionIcon", new IMAGE_BRUSH_SVG(TEXT("Connection"), Icon16x16));
 
 	return Style;
diff --git a/Engine/Plugins/Experimental/ChaosVD/Source/ChaosVD/Private/Widgets/SChaosVDMainTab.cpp b/Engine/Plugins/Experimental/ChaosVD/Source/ChaosVD/Private/Widgets/SChaosVDMainTab.cpp
--- a/Engine/Plugins/Experimental/ChaosVD/Source/ChaosVD/Private/Widgets/SChaosVDMainTab.cpp
+++ b/Engine/Plugins/Experimental/ChaosVD/Source/ChaosVD/Private/Widgets/SChaosVDMainTab.cpp
@@ -192,7 +192,7 @@ void SChaosVDMainTab::Construct(const FArguments& InArgs, TSharedPtr<FChaosVDEng
 				+SHorizontalBox::Slot()
 				[
 					
-					SNew(SChaosVDRecordingControls, StaticCastWeakPtr<SChaosVDMainTab>(AsWeak()))
+					SNew(SChaosVDRecordingControls, StaticCastSharedRef<SChaosVDMainTab>(AsShared()))
 				]
 				
 				+SHorizontalBox::Slot()
@@ -438,8 +438,15 @@ void SChaosVDMainTab::BrowseAndOpenChaosVDRecording()
 			}
 		case EChaosVDBrowseFileModalResponse::OpenTraceStore:
 			{
-				//TODO: Support remote Trace Stores 
-				BrowseChaosVDRecordingFromFolder(FChaosVDModule::Get().GetTraceManager()->GetLocalTraceStoreDirPath());
+				//TODO: Support remote Trace Stores
+				const FString TraceStorePath = FChaosVDModule::Get().GetTraceManager()->GetLocalTraceStoreDirPath();
+				if (TraceStorePath.IsEmpty())
+				{
+					UE_LOG(LogChaosVDEditor, Error, TEXT("[%s] Failed to access Trace Store..."), ANSI_TO_TCHAR(__FUNCTION__));
+					FMessageDialog::Open(EAppMsgType::Ok, LOCTEXT("OpenTraceStoreFailedMessage", "Failed to access the Trace Store, The default profiling folder will be open. \n Please see the logs for mor details... "));
+				}
+
+				BrowseChaosVDRecordingFromFolder(TraceStorePath);
 				break;
 			}
 		case EChaosVDBrowseFileModalResponse::Cancel:
@@ -464,7 +471,7 @@ void SChaosVDMainTab::BrowseChaosVDRecordingFromFolder(FStringView FolderPath)
 			FSlateApplication::Get().FindBestParentWindowHandleForDialogs(AsShared()),
 			LOCTEXT("OpenDialogTitle", "Open Chaos Visual Debug File").ToString(),
 			FolderPath.GetData(),
-			TEXT(""),
+			*FPaths::ProfilingDir(),
 			*ExtensionStr,
 			EFileDialogFlags::None,
 			OutOpenFilenames
@@ -480,6 +487,28 @@ void SChaosVDMainTab::BrowseChaosVDRecordingFromFolder(FStringView FolderPath)
 	}
 }
 
+bool SChaosVDMainTab::ConnectToLiveSession(int32 SessionID, const FString SessionAddress) const
+{
+	FChaosVDTraceSessionDescriptor NewSessionFromFileDescriptor;
+	NewSessionFromFileDescriptor.SessionName = FChaosVDModule::Get().GetTraceManager()->ConnectToLiveSession(SessionAddress, SessionID);
+	NewSessionFromFileDescriptor.bIsLiveSession = true;
+
+	bool bSuccess = false;
+
+	if (NewSessionFromFileDescriptor.SessionName.IsEmpty())
+	{
+		// If it failed we want to clean the current session name, so it is ok calling it either way
+		GetChaosVDEngineInstance()->SetCurrentSession(FChaosVDTraceSessionDescriptor());
+	}
+	else
+	{
+		GetChaosVDEngineInstance()->SetCurrentSession(NewSessionFromFileDescriptor);
+		bSuccess = true;
+	}
+
+	return bSuccess;
+}
+
 void SChaosVDMainTab::BrowseLiveSessionsFromTraceStore() const
 {
 	const TSharedRef<SChaosVDBrowseSessionsModal> SessionBrowserModal = SNew(SChaosVDBrowseSessionsModal);
@@ -490,14 +519,8 @@ void SChaosVDMainTab::BrowseLiveSessionsFromTraceStore() const
 		const FChaosVDTraceSessionInfo SessionInfo = SessionBrowserModal->GetSelectedTraceInfo();
 		if (SessionInfo.bIsValid)
 		{
-			FChaosVDTraceSessionDescriptor NewSessionFromFileDescriptor;
-			NewSessionFromFileDescriptor.SessionName =  FChaosVDModule::Get().GetTraceManager()->ConnectToLiveSession(SessionBrowserModal->GetSelectedTraceStoreAddress(), SessionInfo.TraceID);
-			NewSessionFromFileDescriptor.bIsLiveSession = true;
-
-			bSuccess = !NewSessionFromFileDescriptor.SessionName.IsEmpty();
-
-			// If it failed we want to clean the current session name, so it is ok calling it either way
-			GetChaosVDEngineInstance()->SetCurrentSession(NewSessionFromFileDescriptor);
+			const FString SessionAddress = SessionBrowserModal->GetSelectedTraceStoreAddress();
+			bSuccess = ConnectToLiveSession(SessionInfo.TraceID, SessionAddress);
 		}
 
 		if (!bSuccess)
diff --git a/Engine/Plugins/Experimental/ChaosVD/Source/ChaosVD/Private/Widgets/SChaosVDRecordingControls.cpp b/Engine/Plugins/Experimental/ChaosVD/Source/ChaosVD/Private/Widgets/SChaosVDRecordingControls.cpp
--- a/Engine/Plugins/Experimental/ChaosVD/Source/ChaosVD/Private/Widgets/SChaosVDRecordingControls.cpp
+++ b/Engine/Plugins/Experimental/ChaosVD/Source/ChaosVD/Private/Widgets/SChaosVDRecordingControls.cpp
@@ -3,6 +3,7 @@
 #include "Widgets/SChaosVDRecordingControls.h"
 
 #include "ChaosVDEngine.h"
+#include "ChaosVDModule.h"
 #include "ChaosVDStyle.h"
 #include "ChaosVDRuntimeModule.h"
 #include "ChaosVisualDebugger/ChaosVisualDebuggerTrace.h"
@@ -18,18 +19,11 @@
 
 #define LOCTEXT_NAMESPACE "ChaosVisualDebugger"
 
-void SChaosVDRecordingControls::Construct(const FArguments& InArgs, const TWeakPtr<SChaosVDMainTab>& InMainTabWeakPtr)
+void SChaosVDRecordingControls::Construct(const FArguments& InArgs, const TSharedRef<SChaosVDMainTab>& InMainTabSharedRef)
 {
-	if (TSharedPtr<SChaosVDMainTab> MainTabSharedPtr = InMainTabWeakPtr.Pin())
-	{
-		MainTabWeakPtr = MainTabSharedPtr;
-		StatusBarID = MainTabSharedPtr->GetStatusBarName();
-	}
-	else
-	{
-		ensureMsgf(false, TEXT("Constructed with an invalid Main Tab"));
-	}
-	
+	MainTabWeakPtr = InMainTabSharedRef;
+	StatusBarID = InMainTabSharedRef->GetStatusBarName();
+
 	ChildSlot
 	[
 		SNew(SHorizontalBox)
@@ -41,29 +35,19 @@ void SChaosVDRecordingControls::Construct(const FArguments& InArgs, const TWeakP
 			+SHorizontalBox::Slot()
 			.HAlign(HAlign_Left)
 			[
-				SNew(SButton)
-					.OnClicked(FOnClicked::CreateRaw(this, &SChaosVDRecordingControls::ToggleRecordingState))
-					.ForegroundColor(FSlateColor::UseForeground() )
-					.IsFocusable(false)
-					.VAlign(VAlign_Center)
-					.HAlign(HAlign_Center)
-					[
-						SNew( SImage )
-						.ToolTipText_Lambda([this]()
-						{
-							return IsRecording() ? LOCTEXT("StopRecordButtonDesc", "Stop the current recording ")
-													: LOCTEXT("RecordButtonDesc", "Starts a recording for the current sessions");
-						})
-						.Image_Raw(this, &SChaosVDRecordingControls::GetRecordOrStopButton)
-						.ColorAndOpacity(FColor::Red)
-					]
-					]
-
+				GenerateToggleRecordingStateButton(EChaosVDRecordingMode::File, LOCTEXT("RecordButtonDesc", "Starts a recording for the current session, saving it directly to file"))
+			]
+			+SHorizontalBox::Slot()
+			.Padding(5.0f,  0.0f, 0.0f, 0.0f)
+			.HAlign(HAlign_Left)
+			[
+				GenerateToggleRecordingStateButton(EChaosVDRecordingMode::Live, LOCTEXT("RecordButtonDesc", "Starts a recording and automatically connects to it playing it back in real time"))
+			]
 			+SHorizontalBox::Slot()
 			.HAlign(HAlign_Left)
 			.VAlign(VAlign_Center)
 			.AutoWidth()
-			.Padding(FMargin(12, 0, 0, 0))
+			.Padding(12.0f, 0.0f, 0.0f, 0.0f)
 			[
 				SNew( STextBlock )
 					.TextStyle(FAppStyle::Get(), "SmallButtonText")
@@ -77,6 +61,28 @@ void SChaosVDRecordingControls::Construct(const FArguments& InArgs, const TWeakP
 	RecordingStoppedHandle = FChaosVDRuntimeModule::Get().RegisterRecordingStopCallback(FChaosVDRecordingStateChangedDelegate::FDelegate::CreateRaw(this, &SChaosVDRecordingControls::HandleRecordingStop));
 }
 
+TSharedRef<SButton> SChaosVDRecordingControls::GenerateToggleRecordingStateButton(EChaosVDRecordingMode RecordingMode, const FText& StartRecordingTooltip)
+{
+	return SNew(SButton)
+		.OnClicked(FOnClicked::CreateRaw(this, &SChaosVDRecordingControls::ToggleRecordingState, RecordingMode))
+		.ForegroundColor(FSlateColor::UseForeground())
+		.IsFocusable(false)
+		.IsEnabled_Raw(this, &SChaosVDRecordingControls::IsRecordingToggleButtonEnabled, RecordingMode)
+		.Visibility_Raw(this, &SChaosVDRecordingControls::IsRecordingToggleButtonVisible, RecordingMode)
+		.VAlign(VAlign_Center)
+		.HAlign(HAlign_Center)
+		.ToolTipText_Lambda([this, StartRecordingTooltip]()
+		{
+			return IsRecording() ? LOCTEXT("StopRecordButtonDesc", "Stop the current recording ") : StartRecordingTooltip;
+		})
+		[
+			SNew(SImage)
+			.Image_Raw(this, &SChaosVDRecordingControls::GetRecordOrStopButton, RecordingMode)
+			.ColorAndOpacity_Lambda([this](){ return IsRecording() ? FColor::Red : FColor::White; })
+		];
+}
+
+
 SChaosVDRecordingControls::~SChaosVDRecordingControls()
 {
 	if (FChaosVDRuntimeModule::IsLoaded())
@@ -86,9 +92,10 @@ SChaosVDRecordingControls::~SChaosVDRecordingControls()
 	}
 }
 
-const FSlateBrush* SChaosVDRecordingControls::GetRecordOrStopButton() const
+const FSlateBrush* SChaosVDRecordingControls::GetRecordOrStopButton(EChaosVDRecordingMode RecordingMode) const
 {
-	return IsRecording() ? FChaosVDStyle::Get().GetBrush("StopIcon") : FChaosVDStyle::Get().GetBrush("RecordIcon");
+	const FSlateBrush* RecordIconBrush = RecordingMode == EChaosVDRecordingMode::File ? FChaosVDStyle::Get().GetBrush("RecordToFileIcon") : FChaosVDStyle::Get().GetBrush("RecordToLiveIcon");
+	return IsRecording() ? FChaosVDStyle::Get().GetBrush("StopIcon") : RecordIconBrush;
 }
 
 void SChaosVDRecordingControls::HandleRecordingStop()
@@ -149,11 +156,72 @@ void SChaosVDRecordingControls::HandleRecordingStart()
 	RecordingMessageHandle = StatusBarSubsystem->PushStatusBarMessage(StatusBarID, LOCTEXT("RecordingMessgae", "Recording..."));
 }
 
-FReply SChaosVDRecordingControls::ToggleRecordingState()
+void SChaosVDRecordingControls::AttemptToConnectToLiveSession()
+{
+	bAutoConnectionAttemptInProgress = true;
+	// We need to wait at least one tick before attempting to connect
+	FTSTicker::GetCoreTicker().AddTicker(FTickerDelegate::CreateLambda([WeakThis = AsWeak()](float DeltaTime)
+	{
+		if (const TSharedPtr<SChaosVDRecordingControls> RecordingControlsPtr = StaticCastSharedPtr<SChaosVDRecordingControls>(WeakThis.Pin()))
+		{
+			if (const TSharedPtr<SChaosVDMainTab> MainTabSharedPtr = RecordingControlsPtr->MainTabWeakPtr.Pin())
+			{
+				static FString SessionAddress(TEXT("127.0.0.1"));
+			
+				int32 SessionID = 0;
+
+				FChaosVDTraceManager::EnumerateActiveSessions(SessionAddress, [&SessionID](const UE::Trace::FStoreClient::FSessionInfo& InSessionInfo)
+				{
+					SessionID = InSessionInfo.GetTraceId();
+
+					// CVD stops all active sessions before staring a recording, so we know our session ID will be first;
+					return false;
+				});
+
+				// CVD needs the trace session name to be able to load a live session. Although the session exist, the session name might not be written right away
+				// Trace files don't really have metadata, it is all part of the same stream, so we need to wait until it is written which might take a few ticks.
+				// Therefore if it is not ready, try again a few times.
+				if (!MainTabSharedPtr->ConnectToLiveSession(SessionID, SessionAddress))
+				{
+					if (RecordingControlsPtr->CurrentConnectionAttempts <= RecordingControlsPtr->MaxAutoplayConnectionAttempts)
+					{
+						UE_LOG(LogChaosVDEditor, Verbose, TEXT("[%s] Failed to connect to live session | Attempting again in [%f]..."), ANSI_TO_TCHAR(__FUNCTION__), RecordingControlsPtr->IntervalBetweenAutoplayConnectionAttemptsSeconds);
+						RecordingControlsPtr->AttemptToConnectToLiveSession();
+					}
+					else
+					{
+						RecordingControlsPtr->bAutoConnectionAttemptInProgress = false;
+						UE_LOG(LogChaosVDEditor, Error, TEXT("[%s] Failed to connect to live session | [%d] attempts exhausted..."), ANSI_TO_TCHAR(__FUNCTION__), RecordingControlsPtr->MaxAutoplayConnectionAttempts);	
+					}
+				}
+				else
+				{
+					RecordingControlsPtr->bAutoConnectionAttemptInProgress = false;
+				}
+			}
+		}
+		return false;
+	}), IntervalBetweenAutoplayConnectionAttemptsSeconds);
+}
+
+FReply SChaosVDRecordingControls::ToggleRecordingState(EChaosVDRecordingMode RecordingMode)
 {
 	if (!IsRecording())
 	{
-		FChaosVDRuntimeModule::Get().StartRecording({});
+		TArray<FString, TInlineAllocator<1>> RecordingArgs;
+
+		if (RecordingMode == EChaosVDRecordingMode::Live)
+		{
+			RecordingArgs.Emplace(TEXT("Server"));
+
+			FChaosVDRuntimeModule::Get().StartRecording(RecordingArgs);
+
+			AttemptToConnectToLiveSession();
+		}
+		else
+		{
+			FChaosVDRuntimeModule::Get().StartRecording(RecordingArgs);
+		}
 	}
 	else
 	{
@@ -163,9 +231,46 @@ FReply SChaosVDRecordingControls::ToggleRecordingState()
 	return FReply::Handled();
 }
 
+bool SChaosVDRecordingControls::IsRecordingToggleButtonEnabled(EChaosVDRecordingMode RecordingMode) const
+{
+	if (bAutoConnectionAttemptInProgress)
+	{
+		return false;
+	}
+
+	const TSharedPtr<SChaosVDMainTab> MainTabSharedPtr = MainTabWeakPtr.Pin();
+	if (!MainTabSharedPtr.IsValid())
+	{
+		return false;
+
+	}
+	const bool bIsLiveSession = MainTabSharedPtr->GetChaosVDEngineInstance()->GetCurrentSessionDescriptor().bIsLiveSession;
+
+	const bool bIsRecording = IsRecording();
+
+	if (RecordingMode == EChaosVDRecordingMode::File)
+	{
+		return (bIsRecording && !bIsLiveSession) || !bIsRecording;
+	}
+	else if (RecordingMode == EChaosVDRecordingMode::Live && GEditor)
+	{
+		return (bIsRecording && bIsLiveSession) || (!bIsRecording && GEditor->IsPlayingSessionInEditor());
+	}
+
+	return false;
+}
+
+EVisibility SChaosVDRecordingControls::IsRecordingToggleButtonVisible(EChaosVDRecordingMode RecordingMode) const
+{
+	// If we are recording, don't show the stop button for the mode that is disabled
+	const bool bIsRecording = IsRecording();
+	const bool bShouldButtonBeVisible = bIsRecording ? bIsRecording && IsRecordingToggleButtonEnabled(RecordingMode) : true;
+	return bShouldButtonBeVisible ? EVisibility::Visible : EVisibility::Collapsed;
+}
+
 bool SChaosVDRecordingControls::IsRecording() const
 {
-#if WITH_CHAOS_VISUAL_DEBUGGER && UE_TRACE_ENABLED
+#if WITH_CHAOS_VISUAL_DEBUGGER
 	return FChaosVisualDebuggerTrace::IsTracing();
 #else
 	return false;
diff --git a/Engine/Plugins/Experimental/ChaosVD/Source/ChaosVD/Public/ChaosVDEngine.h b/Engine/Plugins/Experimental/ChaosVD/Source/ChaosVD/Public/ChaosVDEngine.h
--- a/Engine/Plugins/Experimental/ChaosVD/Source/ChaosVD/Public/ChaosVDEngine.h
+++ b/Engine/Plugins/Experimental/ChaosVD/Source/ChaosVD/Public/ChaosVDEngine.h
@@ -47,4 +47,6 @@ private:
 	TSharedPtr<FChaosVDPlaybackController> PlaybackController;
 	
 	bool bIsInitialized = false;
+
+	FDelegateHandle LiveSessionStoppedDelegateHandle;
 };
diff --git a/Engine/Plugins/Experimental/ChaosVD/Source/ChaosVD/Public/Trace/ChaosVDTraceManager.h b/Engine/Plugins/Experimental/ChaosVD/Source/ChaosVD/Public/Trace/ChaosVDTraceManager.h
--- a/Engine/Plugins/Experimental/ChaosVD/Source/ChaosVD/Public/Trace/ChaosVDTraceManager.h
+++ b/Engine/Plugins/Experimental/ChaosVD/Source/ChaosVD/Public/Trace/ChaosVDTraceManager.h
@@ -2,8 +2,10 @@
 
 #pragma once
 
+#include "ChaosVDModule.h"
 #include "Containers/StringFwd.h"
 #include "Templates/SharedPointer.h"
+#include "Trace/StoreClient.h"
 
 class FChaosVDEngine;
 class FChaosVDTraceModule;
@@ -48,6 +50,9 @@ public:
 	/** Stops and de-registers a trace session registered with the provided session name */
 	void CloseSession(const FString& InSessionName);
 
+	template<typename TVisitor>
+	static void EnumerateActiveSessions(FStringView InSessionHost, TVisitor Callback);
+
 private:
 
 	/** The trace analysis session. */
@@ -55,3 +60,35 @@ private:
 
 	TSharedPtr<FChaosVDTraceModule> ChaosVDTraceModule;
 };
+
+template <typename TCallback>
+void FChaosVDTraceManager::EnumerateActiveSessions(FStringView InSessionHost, TCallback Callback)
+{
+	if (InSessionHost.IsEmpty())
+	{
+		UE_LOG(LogChaosVDEditor, Error, TEXT("[%s] Failed to connect to trace store. Provided session host is empty"), ANSI_TO_TCHAR(__FUNCTION__));
+		return;
+	}
+
+	using namespace UE::Trace;
+	const FStoreClient* StoreClient = FStoreClient::Connect(InSessionHost.GetData());
+
+	if (!StoreClient)
+	{
+		UE_LOG(LogChaosVDEditor, Error, TEXT("[%s] Failed to connect to trace store at [%s]"), ANSI_TO_TCHAR(__FUNCTION__), InSessionHost.GetData())
+		return;
+	}
+
+	const uint32 SessionCount = StoreClient->GetSessionCount();
+
+	for (uint32 SessionIndex = 0; SessionIndex < SessionCount; SessionIndex++)
+	{
+		if (const FStoreClient::FSessionInfo* SessionInfo = StoreClient->GetSessionInfo(SessionIndex))
+		{
+			if (!Callback(*SessionInfo))
+			{
+				return;
+			}
+		}	
+	}
+}
diff --git a/Engine/Plugins/Experimental/ChaosVD/Source/ChaosVD/Public/Widgets/SChaosVDMainTab.h b/Engine/Plugins/Experimental/ChaosVD/Source/ChaosVD/Public/Widgets/SChaosVDMainTab.h
--- a/Engine/Plugins/Experimental/ChaosVD/Source/ChaosVD/Public/Widgets/SChaosVDMainTab.h
+++ b/Engine/Plugins/Experimental/ChaosVD/Source/ChaosVD/Public/Widgets/SChaosVDMainTab.h
@@ -53,6 +53,8 @@ public:
 	TSharedPtr<FComponentVisualizer> FindComponentVisualizer(FName ClassName);
 
 	TConstArrayView<TSharedPtr<FComponentVisualizer>> GetAllComponentVisualizers() { return ComponentVisualizers; }
+	
+	bool ConnectToLiveSession(int32 SessionID, FString SessionAddress) const;
 
 private:
 
diff --git a/Engine/Plugins/Experimental/ChaosVD/Source/ChaosVD/Public/Widgets/SChaosVDRecordingControls.h b/Engine/Plugins/Experimental/ChaosVD/Source/ChaosVD/Public/Widgets/SChaosVDRecordingControls.h
--- a/Engine/Plugins/Experimental/ChaosVD/Source/ChaosVD/Public/Widgets/SChaosVDRecordingControls.h
+++ b/Engine/Plugins/Experimental/ChaosVD/Source/ChaosVD/Public/Widgets/SChaosVDRecordingControls.h
@@ -6,28 +6,42 @@
 #include "Templates/SharedPointer.h"
 #include "Widgets/SCompoundWidget.h"
 
+class SButton;
 class SChaosVDMainTab;
 class FReply;
 struct FSlateBrush;
 
+UENUM()
+enum class EChaosVDRecordingMode
+{
+	File,
+	Live
+};
+
 class SChaosVDRecordingControls : public SCompoundWidget
 {
 public:
 	SLATE_BEGIN_ARGS( SChaosVDRecordingControls ){}
 	SLATE_END_ARGS()
 	
-	void Construct(const FArguments& InArgs, const TWeakPtr<SChaosVDMainTab>& InMainTabWeakPtr);
+	void Construct(const FArguments& InArgs, const TSharedRef<SChaosVDMainTab>& InMainTabSharedRef);
 
 	virtual ~SChaosVDRecordingControls() override;
 
 protected:
-	
-	const FSlateBrush* GetRecordOrStopButton() const;
+
+	TSharedRef<SButton>GenerateToggleRecordingStateButton(EChaosVDRecordingMode RecordingMode, const FText& StartRecordingTooltip);
+
+	const FSlateBrush* GetRecordOrStopButton(EChaosVDRecordingMode RecordingMode) const;
 	
 	void HandleRecordingStop();
 	void HandleRecordingStart();
+	void AttemptToConnectToLiveSession();
 
-	FReply ToggleRecordingState();
+	FReply ToggleRecordingState(EChaosVDRecordingMode RecordingMode);
+
+	bool IsRecordingToggleButtonEnabled(EChaosVDRecordingMode RecordingMode) const;
+	EVisibility IsRecordingToggleButtonVisible(EChaosVDRecordingMode RecordingMode) const;
 
 	bool IsRecording() const;
 
@@ -42,4 +56,9 @@ protected:
 	FDelegateHandle RecordingStoppedHandle;
 
 	TWeakPtr<SChaosVDMainTab> MainTabWeakPtr;
+
+	int32 MaxAutoplayConnectionAttempts = 20;
+	float IntervalBetweenAutoplayConnectionAttemptsSeconds = 0.1f;
+	bool bAutoConnectionAttemptInProgress = false;
+	int32 CurrentConnectionAttempts = 0;
 };
diff --git a/Engine/Source/Runtime/Experimental/ChaosVisualDebugger/Private/ChaosVDRuntimeModule.cpp b/Engine/Source/Runtime/Experimental/ChaosVisualDebugger/Private/ChaosVDRuntimeModule.cpp
--- a/Engine/Source/Runtime/Experimental/ChaosVisualDebugger/Private/ChaosVDRuntimeModule.cpp
+++ b/Engine/Source/Runtime/Experimental/ChaosVisualDebugger/Private/ChaosVDRuntimeModule.cpp
@@ -4,7 +4,10 @@
 
 #include "Containers/Array.h"
 #include "HAL/IConsoleManager.h"
+#include "Internationalization/Internationalization.h"
+#include "Misc/CommandLine.h"
 #include "Misc/FileHelper.h"
+#include "Misc/MessageDialog.h"
 #include "Misc/Paths.h"
 #include "Modules/ModuleManager.h"
 #include "ProfilingDebugging/TraceAuxiliary.h"
@@ -13,6 +16,8 @@ IMPLEMENT_MODULE(FChaosVDRuntimeModule, ChaosVDRuntime);
 
 DEFINE_LOG_CATEGORY_STATIC( LogChaosVDRuntime, Log, All );
 
+#define LOCTEXT_NAMESPACE "ChaosVisualDebugger"
+
 FAutoConsoleCommand ChaosVDStartRecordingCommand(
 	TEXT("p.Chaos.StartVDRecording"),
 	TEXT("Turn on the recording of debugging data"),
@@ -53,14 +58,37 @@ bool FChaosVDRuntimeModule::IsLoaded()
 
 void FChaosVDRuntimeModule::StartupModule()
 {
+	if (FParse::Param(FCommandLine::Get(), TEXT("StartCVDRecording")))
+	{
+		TArray<FString, TInlineAllocator<1>> CVDOptions;
+
+		{
+			FString CVDHostAddress;
+			if (FParse::Value(FCommandLine::Get(), TEXT("CVDHost="), CVDHostAddress))
+			{
+				CVDOptions.Emplace(MoveTemp(CVDHostAddress));
+			}
+		}
+        
+        StartRecording(CVDOptions);
+	}
+	else
+	{
+		
 #if UE_TRACE_ENABLED
-	// Make sure it is off until we support auto trace
-	UE::Trace::ToggleChannel(TEXT("ChaosVDChannel"), false);
+		UE::Trace::ToggleChannel(TEXT("ChaosVDChannel"), false);
 #endif
+
+	}
 }
 
 void FChaosVDRuntimeModule::ShutdownModule()
 {
+	if (bIsRecording)
+	{
+		StopRecording();
+	}
+
 	FTraceAuxiliary::OnTraceStopped.RemoveAll(this);
 }
 
@@ -105,7 +133,7 @@ bool FChaosVDRuntimeModule::RecordingTimerTick(float DeltaTime)
 	return true;
 }
 
-void FChaosVDRuntimeModule::StartRecording(const TArray<FString>& Args)
+void FChaosVDRuntimeModule::StartRecording(TConstArrayView<FString> Args)
 {
 	if (bIsRecording)
 	{
@@ -115,23 +143,6 @@ void FChaosVDRuntimeModule::StartRecording(const TArray<FString>& Args)
 	// Start Listening for Trace Stopped events, in case Trace is stopped outside our control so we can gracefully stop CVD recording and log a warning 
 	FTraceAuxiliary::OnTraceStopped.AddRaw(this, &FChaosVDRuntimeModule::HandleTraceStopRequest);
 
-	{
-		FReadScopeLock ReadLock(DelegatesRWLock);
-		RecordingStartedDelegate.Broadcast();
-	}
-
-	constexpr int32 MinAllowedTimeInSecondsBetweenCaptures = 1;
-	int32 ConfiguredTimeBetweenCaptures = CVarChaosVDGTimeBetweenFullCaptures->GetInt();
-
-	ensureAlwaysMsgf(ConfiguredTimeBetweenCaptures > MinAllowedTimeInSecondsBetweenCaptures,
-		TEXT("The minimum allowed time interval between full captures is [%d] seconds, but [%d] seconds were configured. Clamping to [%d] seconds"),
-		MinAllowedTimeInSecondsBetweenCaptures, ConfiguredTimeBetweenCaptures, MinAllowedTimeInSecondsBetweenCaptures);
-
-	FullCaptureRequesterHandle = FTSTicker::GetCoreTicker().AddTicker(FTickerDelegate::CreateRaw(this, &FChaosVDRuntimeModule::RequestFullCapture),
-		FMath::Clamp(ConfiguredTimeBetweenCaptures, MinAllowedTimeInSecondsBetweenCaptures, TNumericLimits<int32>::Max()));
-
-	RecordingTimerHandle = FTSTicker::GetCoreTicker().AddTicker(FTickerDelegate::CreateRaw(this, &FChaosVDRuntimeModule::RecordingTimerTick));
-
 #if UE_TRACE_ENABLED
 
 	// Other tools could bee using trace
@@ -141,6 +152,9 @@ void FChaosVDRuntimeModule::StartRecording(const TArray<FString>& Args)
 		StopTrace();
 	}
 
+	// Until we support allowing other channels, indicate in the logs that we are disabling everything else
+	UE_LOG(LogChaosVDRuntime, Log, TEXT("[%s] Disabling additional trace channels..."), ANSI_TO_TCHAR(__FUNCTION__));
+
 	// Disable any enabled additional channel
 	UE::Trace::EnumerateChannels([](const ANSICHAR* ChannelName, bool bEnabled, void*)
 		{
@@ -180,21 +194,50 @@ void FChaosVDRuntimeModule::StartRecording(const TArray<FString>& Args)
 #endif
 	
 	AccumulatedRecordingTime = 0.0f;
-	
-	ensure(bIsRecording);
+
+	if (ensure(bIsRecording))
+	{
+		{
+			FReadScopeLock ReadLock(DelegatesRWLock);
+			RecordingStartedDelegate.Broadcast();
+		}
+		
+		constexpr int32 MinAllowedTimeInSecondsBetweenCaptures = 1;
+		int32 ConfiguredTimeBetweenCaptures = CVarChaosVDGTimeBetweenFullCaptures->GetInt();
+
+		ensureAlwaysMsgf(ConfiguredTimeBetweenCaptures > MinAllowedTimeInSecondsBetweenCaptures,
+			TEXT("The minimum allowed time interval between full captures is [%d] seconds, but [%d] seconds were configured. Clamping to [%d] seconds"),
+			MinAllowedTimeInSecondsBetweenCaptures, ConfiguredTimeBetweenCaptures, MinAllowedTimeInSecondsBetweenCaptures);
+
+		FullCaptureRequesterHandle = FTSTicker::GetCoreTicker().AddTicker(FTickerDelegate::CreateRaw(this, &FChaosVDRuntimeModule::RequestFullCapture),
+			FMath::Clamp(ConfiguredTimeBetweenCaptures, MinAllowedTimeInSecondsBetweenCaptures, TNumericLimits<int32>::Max()));
+
+		RecordingTimerHandle = FTSTicker::GetCoreTicker().AddTicker(FTickerDelegate::CreateRaw(this, &FChaosVDRuntimeModule::RecordingTimerTick));
+
+	}
+	else
+	{
+		UE_LOG(LogChaosVDRuntime, Error, TEXT("[%s] Failed to start CVD recording..."), ANSI_TO_TCHAR(__FUNCTION__));
+
+#if WITH_EDITOR
+		FMessageDialog::Open(EAppMsgType::Ok, LOCTEXT("StartRecordingFailedMessage", "Failed to start CVD recording. Please see the logs for more details... "));
+#endif
+
+	}
+
 }
 
 void FChaosVDRuntimeModule::StopRecording()
 {
-	if (!ensure(bIsRecording))
+	if (!bIsRecording)
 	{
+		UE_LOG(LogChaosVDRuntime, Warning, TEXT("[%s] Attempted to stop recorded when there is no CVD recording active."), ANSI_TO_TCHAR(__FUNCTION__));
 		return;
 	}
 	
 	FTraceAuxiliary::OnTraceStopped.RemoveAll(this);
 
 #if UE_TRACE_ENABLED
-	
 
 	UE::Trace::ToggleChannel(TEXT("ChaosVDChannel"), false);
 	UE::Trace::ToggleChannel(TEXT("Frame"), false); 
@@ -223,6 +266,10 @@ void FChaosVDRuntimeModule::HandleTraceStopRequest(FTraceAuxiliary::EConnectionT
 		if (!ensure(bRequestedStop))
 		{
 			UE_LOG(LogChaosVDRuntime, Warning, TEXT("Trace Recording has been stopped unexpectedly"));
+
+#if WITH_EDITOR
+			FMessageDialog::Open(EAppMsgType::Ok, LOCTEXT("UnexpectedStopMessage", "Trace recording has been stopped unexpectedly. CVD cannot continue with the recording session... "));
+#endif
 		}
 
 		StopRecording();
@@ -230,3 +277,5 @@ void FChaosVDRuntimeModule::HandleTraceStopRequest(FTraceAuxiliary::EConnectionT
 
 	bRequestedStop = false;
 }
+
+#undef LOCTEXT_NAMESPACE 
diff --git a/Engine/Source/Runtime/Experimental/ChaosVisualDebugger/Public/ChaosVDRuntimeModule.h b/Engine/Source/Runtime/Experimental/ChaosVisualDebugger/Public/ChaosVDRuntimeModule.h
--- a/Engine/Source/Runtime/Experimental/ChaosVisualDebugger/Public/ChaosVDRuntimeModule.h
+++ b/Engine/Source/Runtime/Experimental/ChaosVisualDebugger/Public/ChaosVDRuntimeModule.h
@@ -34,7 +34,7 @@ public:
 	/** Starts a CVD recording by starting a Trace session. It will stop any existing trace session
 	 * @param Args : Arguments array provided by the commandline. Used to determine if we want to record to file or a local trace server
 	 */
-	void StartRecording(const TArray<FString>& Args);
+	void StartRecording(TConstArrayView<FString> Args);
 	
 	/* Stops an active recording */
 	void StopRecording();
