[ChaosVD] Fix LocText confict in SChaosVDRecordingControls

#prefight skip

diff --git a/Engine/Plugins/Experimental/ChaosVD/Source/ChaosVD/Private/Widgets/SChaosVDRecordingControls.cpp b/Engine/Plugins/Experimental/ChaosVD/Source/ChaosVD/Private/Widgets/SChaosVDRecordingControls.cpp
--- a/Engine/Plugins/Experimental/ChaosVD/Source/ChaosVD/Private/Widgets/SChaosVDRecordingControls.cpp
+++ b/Engine/Plugins/Experimental/ChaosVD/Source/ChaosVD/Private/Widgets/SChaosVDRecordingControls.cpp
@@ -35,13 +35,13 @@ void SChaosVDRecordingControls::Construct(const FArguments& InArgs, const TShare
 			+SHorizontalBox::Slot()
 			.HAlign(HAlign_Left)
 			[
-				GenerateToggleRecordingStateButton(EChaosVDRecordingMode::File, LOCTEXT("RecordButtonDesc", "Starts a recording for the current session, saving it directly to file"))
+				GenerateToggleRecordingStateButton(EChaosVDRecordingMode::File, LOCTEXT("RecordToFileButtonDesc", "Starts a recording for the current session, saving it directly to file"))
 			]
 			+SHorizontalBox::Slot()
 			.Padding(5.0f,  0.0f, 0.0f, 0.0f)
 			.HAlign(HAlign_Left)
 			[
-				GenerateToggleRecordingStateButton(EChaosVDRecordingMode::Live, LOCTEXT("RecordButtonDesc", "Starts a recording and automatically connects to it playing it back in real time"))
+				GenerateToggleRecordingStateButton(EChaosVDRecordingMode::Live, LOCTEXT("RecordLiveButtonDesc", "Starts a recording and automatically connects to it playing it back in real time"))
 			]
 			+SHorizontalBox::Slot()
 			.HAlign(HAlign_Left)
