[ChaosVD] Fix CIS issue on shipping Editors.

The new Data Channels feature had some sections gated behind the WITH_CHAOS_VISUAL_DEBUGGER_SUPPORT which currently is compiled out in shipping configurations.

The ChaosVD Editor plugin had code referincing that compiled out code, so I also gated it behind WITH_CHAOS_VISUAL_DEBUGGER_SUPPORT for now.I will look for a better solution in a follow up CL.

I also added a fix on a setter implementation  in the new system that got lost in one of the last changes before submitting.

#changelist validated
#virtualized

diff --git a/Engine/Plugins/Experimental/ChaosVD/Source/ChaosVD/Private/Widgets/SChaosVDRecordingControls.cpp b/Engine/Plugins/Experimental/ChaosVD/Source/ChaosVD/Private/Widgets/SChaosVDRecordingControls.cpp
--- a/Engine/Plugins/Experimental/ChaosVD/Source/ChaosVD/Private/Widgets/SChaosVDRecordingControls.cpp
+++ b/Engine/Plugins/Experimental/ChaosVD/Source/ChaosVD/Private/Widgets/SChaosVDRecordingControls.cpp
@@ -65,6 +65,7 @@ void SChaosVDRecordingControls::Construct(const FArguments& InArgs, const TShare
 			[
 				SNew(SComboButton)
 				.ContentPadding(FMargin(6.0f, 0.0f))
+				.IsEnabled_Raw(this, &SChaosVDRecordingControls::HasDataChannelsSupport)
 				.MenuPlacement(MenuPlacement_AboveAnchor)
 				.ComboButtonStyle(&FAppStyle::Get().GetWidgetStyle<FComboButtonStyle>("SimpleComboButton"))
 				.OnGetMenuContent(this, &SChaosVDRecordingControls::GenerateDataChannelsMenu)
@@ -116,6 +117,7 @@ TSharedRef<SWidget> SChaosVDRecordingControls::GenerateDataChannelsMenu()
 
 	FMenuBuilder MenuBuilder(true, nullptr);
 
+#if WITH_CHAOS_VISUAL_DEBUGGER
 	MenuBuilder.BeginSection("CVDRecordingWidget", LOCTEXT("CVDRecordingMenuChannels", "Data Channels"));
 	{
 		FChaosVDDataChannelsManager::Get().EnumerateChannels([this, &MenuBuilder](const TSharedRef<FCVDDataChannel>& Channel)
@@ -130,12 +132,15 @@ TSharedRef<SWidget> SChaosVDRecordingControls::GenerateDataChannelsMenu()
 			return true;
 		});
 	}
+#endif //WITH_CHAOS_VISUAL_DEBUGGER
 
 	MenuBuilder.EndSection();
 
 	return MenuBuilder.MakeWidget();
 }
 
+#if WITH_CHAOS_VISUAL_DEBUGGER
+
 void SChaosVDRecordingControls::ToggleChannelEnabledState(TWeakPtr<FCVDDataChannel> Channel)
 {
 	if (const TSharedPtr<FCVDDataChannel>& ChannelPtr = Channel.Pin())
@@ -165,6 +170,17 @@ bool SChaosVDRecordingControls::CanChangeChannelEnabledState(TWeakPtr<FCVDDataCh
 	return false;
 }
 
+#endif // WITH_CHAOS_VISUAL_DEBUGGER
+
+bool SChaosVDRecordingControls::HasDataChannelsSupport() const
+{
+#if WITH_CHAOS_VISUAL_DEBUGGER
+	return true;
+#else
+	return false;
+#endif
+}
+
 SChaosVDRecordingControls::~SChaosVDRecordingControls()
 {
 	if (FChaosVDRuntimeModule::IsLoaded())
diff --git a/Engine/Plugins/Experimental/ChaosVD/Source/ChaosVD/Public/Widgets/SChaosVDRecordingControls.h b/Engine/Plugins/Experimental/ChaosVD/Source/ChaosVD/Public/Widgets/SChaosVDRecordingControls.h
--- a/Engine/Plugins/Experimental/ChaosVD/Source/ChaosVD/Public/Widgets/SChaosVDRecordingControls.h
+++ b/Engine/Plugins/Experimental/ChaosVD/Source/ChaosVD/Public/Widgets/SChaosVDRecordingControls.h
@@ -40,10 +40,14 @@ protected:
 	TSharedRef<SButton> GenerateToggleRecordingStateButton(EChaosVDRecordingMode RecordingMode, const FText& StartRecordingTooltip);
 	TSharedRef<SWidget> GenerateDataChannelsMenu();
 
+#if WITH_CHAOS_VISUAL_DEBUGGER
 	void ToggleChannelEnabledState(TWeakPtr<FCVDDataChannel> Channel);
 	bool IsChannelEnabled(TWeakPtr<FCVDDataChannel> Channel);
 	bool CanChangeChannelEnabledState(TWeakPtr<FCVDDataChannel> Channel);
+#endif
 
+	bool HasDataChannelsSupport() const;
+	
 	const FSlateBrush* GetRecordOrStopButton(EChaosVDRecordingMode RecordingMode) const;
 	
 	void HandleRecordingStop();
diff --git a/Engine/Source/Runtime/Experimental/Chaos/Private/ChaosVisualDebugger/ChaosVDOptionalDataChannel.cpp b/Engine/Source/Runtime/Experimental/Chaos/Private/ChaosVisualDebugger/ChaosVDOptionalDataChannel.cpp
--- a/Engine/Source/Runtime/Experimental/Chaos/Private/ChaosVisualDebugger/ChaosVDOptionalDataChannel.cpp
+++ b/Engine/Source/Runtime/Experimental/Chaos/Private/ChaosVisualDebugger/ChaosVDOptionalDataChannel.cpp
@@ -4,7 +4,6 @@
 
 #if WITH_CHAOS_VISUAL_DEBUGGER
 
-#include "ChaosVDRuntimeModule.h"
 #include "HAL/IConsoleManager.h"
 #include "Internationalization/Text.h"
 #include "Internationalization/Internationalization.h"
@@ -50,7 +49,7 @@ namespace Chaos::VisualDebugger
 
 	void FChaosVDOptionalDataChannel::SetChannelEnabled(bool bNewEnabled)
 	{
-		
+		bIsEnabled = bNewEnabled;
 	}
 
 	FChaosVDDataChannelsManager& FChaosVDDataChannelsManager::Get()
diff --git a/Engine/Source/Runtime/Experimental/Chaos/Public/ChaosVisualDebugger/ChaosVDOptionalDataChannel.h b/Engine/Source/Runtime/Experimental/Chaos/Public/ChaosVisualDebugger/ChaosVDOptionalDataChannel.h
--- a/Engine/Source/Runtime/Experimental/Chaos/Public/ChaosVisualDebugger/ChaosVDOptionalDataChannel.h
+++ b/Engine/Source/Runtime/Experimental/Chaos/Public/ChaosVisualDebugger/ChaosVDOptionalDataChannel.h
@@ -108,6 +108,8 @@ namespace Chaos::VisualDebugger
 	TSharedRef<FChaosVDOptionalDataChannel> CreateDataChannel(FName InChannelID, const TSharedRef<FText>& InDisplayName, EChaosVDDataChannelInitializationFlags InitializationFlags);
 }
 
+#endif //WITH_CHAOS_VISUAL_DEBUGGER
+
 #ifndef CVD_CONCAT_NX
 	#define CVD_CONCAT_NX(A, B) A ## B
 #endif
@@ -124,6 +126,8 @@ namespace Chaos::VisualDebugger
 	#define CVD_STRINGIZE(A) CVD_STRINGIZE_NX(A)
 #endif
 
+#if WITH_CHAOS_VISUAL_DEBUGGER
+
 /** Declares an Optional CVD Data channel to be available globally. The data channel can be accessed by using CVDDC_TheNameOfTheChannelUsedWithThisMacro */
 #ifndef CVD_DECLARE_OPTIONAL_DATA_CHANNEL
 	#define CVD_DECLARE_OPTIONAL_DATA_CHANNEL(DataChannelName) \
@@ -147,4 +151,16 @@ CVD_DECLARE_OPTIONAL_DATA_CHANNEL(EndOfEvolutionCollisionConstraints);
 CVD_DECLARE_OPTIONAL_DATA_CHANNEL(EvolutionEnd);
 CVD_DECLARE_OPTIONAL_DATA_CHANNEL(SceneQueries);
 
+#else  //WITH_CHAOS_VISUAL_DEBUGGER
+
+/** Declares an Optional CVD Data channel to be available globally. The data channel can be accessed by using CVDDC_TheNameOfTheChannelUsedWithThisMacro */
+#ifndef CVD_DECLARE_OPTIONAL_DATA_CHANNEL
+	#define CVD_DECLARE_OPTIONAL_DATA_CHANNEL(DataChannelName)
 #endif
+
+/** Defines and initializes an Optional CVD Data Channel */
+#ifndef CVD_DEFINE_OPTIONAL_DATA_CHANNEL
+	#define CVD_DEFINE_OPTIONAL_DATA_CHANNEL(DataChannelName, InitializationFlags)
+#endif
+
+#endif  //WITH_CHAOS_VISUAL_DEBUGGER
