[ChaosVD] Particle Data and other Visualization improvements.

- Moved Particle data visualization to the new system based on component visualizers used by [REDACTED_PERSON] visualization and Scene query visualization.
- Exposed debug draw settings as option sin the show menu (you can change the scales and colors)
- Removed the visualizations flags from the particle actor details panel. This was part of the first implementation and no really compatible with the new system intent.
- Improved Debug text generation for particle data visualization (now the text also shows values and units when enabled).
- Improved selection feedback. When a particle is selected, its debug draw elements will also be highlighted.
- Added a simple check to early out of the visualization will end out outside of the view

#changelist validated
#virtualized

diff --git a/Engine/Plugins/Experimental/ChaosVD/Source/ChaosVD/Private/Actors/ChaosVDSolverInfoActor.cpp b/Engine/Plugins/Experimental/ChaosVD/Source/ChaosVD/Private/Actors/ChaosVDSolverInfoActor.cpp
--- a/Engine/Plugins/Experimental/ChaosVD/Source/ChaosVD/Private/Actors/ChaosVDSolverInfoActor.cpp
+++ b/Engine/Plugins/Experimental/ChaosVD/Source/ChaosVD/Private/Actors/ChaosVDSolverInfoActor.cpp
@@ -2,9 +2,11 @@
 
 #include "Actors/ChaosVDSolverInfoActor.h"
 
+#include "ChaosVDModule.h"
 #include "ChaosVDParticleActor.h"
 #include "ChaosVDScene.h"
 #include "EditorActorFolders.h"
+#include "Components/ChaosVDParticleDataComponent.h"
 #include "Components/ChaosVDSolverCollisionDataComponent.h"
 #include "Elements/Framework/TypedElementSelectionSet.h"
 #include "Engine/World.h"
@@ -15,6 +17,7 @@
 AChaosVDSolverInfoActor::AChaosVDSolverInfoActor(const FObjectInitializer& ObjectInitializer) : Super(ObjectInitializer)
 {
 	CollisionDataComponent = CreateDefaultSubobject<UChaosVDSolverCollisionDataComponent>(TEXT("SolverCollisionDataComponent"));
+	ParticleDataComponent = CreateDefaultSubobject<UChaosVDParticleDataComponent>(TEXT("ParticleCollisionDataComponent"));
 	bIsServer = false;
 }
 
@@ -51,6 +54,34 @@ AChaosVDParticleActor* AChaosVDSolverInfoActor::GetParticleActor(int32 ParticleI
 	return FoundParticleActor ? *FoundParticleActor : nullptr;
 }
 
+bool AChaosVDSolverInfoActor::IsParticleSelectedByID(int32 ParticleID)
+{
+	// Currently CVD does not support multi selection, so this should not be slow.
+	// But we might want to find another container for faster search after multi selection support is added
+	return SelectedParticlesID.Contains(ParticleID);
+}
+
+bool AChaosVDSolverInfoActor::SelectParticleByID(int32 ParticleIDToSelect)
+{
+	const TSharedPtr<FChaosVDScene> CVDScene = SceneWeakPtr.Pin();
+	if (!CVDScene)
+	{
+		UE_LOG(LogChaosVDEditor, Error, TEXT("[%s] Tried to select a particle without a valid CVD Scene."), ANSI_TO_TCHAR(__FUNCTION__))
+		return false;
+	}
+
+	AChaosVDParticleActor* ParticleToSelect = GetParticleActor(ParticleIDToSelect);
+	if (!ParticleToSelect)
+	{
+		UE_LOG(LogChaosVDEditor, Error, TEXT("[%s] Particle ID [%d] not found in Solver [%s]"), ANSI_TO_TCHAR(__FUNCTION__), ParticleIDToSelect, *GetSolverName());
+		return false;
+	}
+
+	CVDScene->SetSelectedObject(ParticleToSelect);
+
+	return true;
+}
+
 void AChaosVDSolverInfoActor::HandleVisibilitySettingsUpdated()
 {
 	for (const TPair<int32, AChaosVDParticleActor*>& ParticleWithIDPair : SolverParticlesByID)
diff --git a/Engine/Plugins/Experimental/ChaosVD/Source/ChaosVD/Private/ChaosVDEditorSettings.cpp b/Engine/Plugins/Experimental/ChaosVD/Source/ChaosVD/Private/ChaosVDEditorSettings.cpp
--- a/Engine/Plugins/Experimental/ChaosVD/Source/ChaosVD/Private/ChaosVDEditorSettings.cpp
+++ b/Engine/Plugins/Experimental/ChaosVD/Source/ChaosVD/Private/ChaosVDEditorSettings.cpp
@@ -3,6 +3,7 @@
 #include "ChaosVDEditorSettings.h"
 #include "DataWrappers/ChaosVDParticleDataWrapper.h"
 #include "Chaos/ImplicitObjectType.h"
+#include "Visualizers/ChaosVDParticleDataComponentVisualizer.h"
 
 FColor FChaosDebugDrawColorsByState::GetColorFromState(EChaosVDObjectStateType State) const
 {
@@ -21,6 +22,69 @@ FColor FChaosDebugDrawColorsByState::GetColorFromState(EChaosVDObjectStateType S
 	}
 }
 
+FColor FChaosParticleDataDebugDrawColors::GetColorForDataID(EChaosVDParticleDataVisualizationFlags DataID, bool bIsSelected) const
+{
+	constexpr float DefaultIntensityFactor = 0.6f;
+	constexpr float SelectedIntensityFactor = 1.0f;
+	const float IntensityFactor = bIsSelected ? SelectedIntensityFactor : DefaultIntensityFactor;
+
+	return (GetLinearColorForDataID(DataID) * IntensityFactor).ToFColorSRGB();
+}
+
+const FLinearColor& FChaosParticleDataDebugDrawColors::GetLinearColorForDataID(EChaosVDParticleDataVisualizationFlags DataID) const
+{
+	static FLinearColor InvalidColor = FColor::Purple;
+	switch (DataID)
+	{
+	case EChaosVDParticleDataVisualizationFlags::Acceleration:
+		return AccelerationColor;
+	case EChaosVDParticleDataVisualizationFlags::Velocity:
+		return VelocityColor;
+	case EChaosVDParticleDataVisualizationFlags::AngularVelocity:
+		return AngularVelocityColor;
+	case EChaosVDParticleDataVisualizationFlags::AngularAcceleration:
+		return AngularAccelerationColor;
+	case EChaosVDParticleDataVisualizationFlags::LinearImpulse:
+		return LinearImpulseColor;
+	case EChaosVDParticleDataVisualizationFlags::AngularImpulse:
+		return AngularImpulseColor;
+	case EChaosVDParticleDataVisualizationFlags::ClusterConnectivityEdge:
+		return ConnectivityDataColor;
+	case EChaosVDParticleDataVisualizationFlags::CenterOfMass:
+		return CenterOfMassColor;
+	case EChaosVDParticleDataVisualizationFlags::None:
+	case EChaosVDParticleDataVisualizationFlags::DrawDataOnlyForSelectedParticle:
+	default:
+		return InvalidColor;
+	}
+}
+
+
+float FChaosParticleDataDebugDrawSettings::GetScaleFortDataID(EChaosVDParticleDataVisualizationFlags DataID) const
+{
+	switch (DataID)
+    {
+    	case EChaosVDParticleDataVisualizationFlags::Acceleration:
+    		return AccelerationScale;
+    	case EChaosVDParticleDataVisualizationFlags::Velocity:
+    		return VelocityScale;
+    	case EChaosVDParticleDataVisualizationFlags::AngularVelocity:
+    		return AngularVelocityScale;
+    	case EChaosVDParticleDataVisualizationFlags::AngularAcceleration:
+    		return AngularAccelerationScale;
+    	case EChaosVDParticleDataVisualizationFlags::LinearImpulse:
+    		return LinearImpulseScale;
+    	case EChaosVDParticleDataVisualizationFlags::AngularImpulse:
+    		return AngularImpulseScale;
+    	case EChaosVDParticleDataVisualizationFlags::ClusterConnectivityEdge:
+    	case EChaosVDParticleDataVisualizationFlags::CenterOfMass:
+    	case EChaosVDParticleDataVisualizationFlags::None:
+    	case EChaosVDParticleDataVisualizationFlags::DrawDataOnlyForSelectedParticle:
+    	default:
+    		return 1.0f;
+    }
+}
+
 FColor FChaosDebugDrawColorsByShapeType::GetColorFromShapeType(Chaos::EImplicitObjectType ShapeType) const
 {
 	switch(ShapeType)
diff --git a/Engine/Plugins/Experimental/ChaosVD/Source/ChaosVD/Private/ChaosVDParticleActor.cpp b/Engine/Plugins/Experimental/ChaosVD/Source/ChaosVD/Private/ChaosVDParticleActor.cpp
--- a/Engine/Plugins/Experimental/ChaosVD/Source/ChaosVD/Private/ChaosVDParticleActor.cpp
+++ b/Engine/Plugins/Experimental/ChaosVD/Source/ChaosVD/Private/ChaosVDParticleActor.cpp
@@ -13,10 +13,7 @@
 #include "Components/SceneComponent.h"
 #include "Components/StaticMeshComponent.h"
 #include "DataWrappers/ChaosVDParticleDataWrapper.h"
-#include "Engine/CollisionProfile.h"
 #include "Engine/StaticMesh.h"
-#include "Visualizers/ChaosVDParticleDataVisualizer.h"
-#include "Visualizers/ChaosVDSolverCollisionDataComponentVisualizer.h"
 
 namespace Chaos::VisualDebugger::Cvars
 {
@@ -32,8 +29,6 @@ AChaosVDParticleActor::AChaosVDParticleActor(const FObjectInitializer& ObjectIni
 	RootComponent = CreateDefaultSubobject<USceneComponent>(TEXT("SceneComponent0"));
 	RootComponent->SetCanEverAffectNavigation(false);
 	RootComponent->bNavigationRelevant = false;
-
-	CreateVisualizers();
 }
 
 void AChaosVDParticleActor::UpdateFromRecordedParticleData(const TSharedPtr<FChaosVDParticleDataWrapper>& InRecordedData, const Chaos::FRigidTransform3& SimulationTransform)
@@ -323,18 +318,6 @@ void AChaosVDParticleActor::Destroyed()
 	Super::Destroyed();
 }
 
-void AChaosVDParticleActor::GetVisualizationContext(FChaosVDVisualizationContext& OutVisualizationContext)
-{
-	OutVisualizationContext.SpaceTransform = CachedSimulationTransform;
-	OutVisualizationContext.CVDScene = SceneWeakPtr;
-	OutVisualizationContext.SolverID = ParticleDataPtr ? ParticleDataPtr->SolverID : INDEX_NONE;
-}
-
-void AChaosVDParticleActor::CreateVisualizers()
-{
-	CVDVisualizers.Add(FChaosVDParticleDataVisualizer::VisualizerID,MakeUnique<FChaosVDParticleDataVisualizer>(*this));
-}
-
 #if WITH_EDITOR
 
 bool AChaosVDParticleActor::IsSelectedInEditor() const
@@ -349,19 +332,6 @@ bool AChaosVDParticleActor::IsSelectedInEditor() const
 	return false;
 }
 
-void AChaosVDParticleActor::PostEditChangeProperty(FPropertyChangedEvent& PropertyChangedEvent)
-{
-	Super::PostEditChangeProperty(PropertyChangedEvent);
-
-	if (PropertyChangedEvent.GetMemberPropertyName() == GET_MEMBER_NAME_STRING_CHECKED(AChaosVDParticleActor, LocalParticleDataVisualizationFlags))
-	{
-		if (TUniquePtr<FChaosVDDataVisualizerBase>* ParticleDataVisualizer = CVDVisualizers.Find(FChaosVDParticleDataVisualizer::VisualizerID))
-		{
-			ParticleDataVisualizer->Get()->UpdateVisualizationFlags(LocalParticleDataVisualizationFlags);
-		}
-	}
-}
-
 #if WITH_EDITOR
 void AChaosVDParticleActor::SetIsTemporarilyHiddenInEditor(bool bIsHidden)
 {
@@ -526,15 +496,4 @@ void AChaosVDParticleActor::SetIsActive(bool bNewActive)
 	}
 }
 
-void AChaosVDParticleActor::DrawVisualization(const FSceneView* View, FPrimitiveDrawInterface* PDI)
-{
-	for (const TPair<FStringView, TUniquePtr<FChaosVDDataVisualizerBase>>& VisualizerWithID : CVDVisualizers)
-	{
-		if (VisualizerWithID.Value)
-		{
-			VisualizerWithID.Value->DrawVisualization(View, PDI);
-		}
-	}
-}
-
 #endif //WITH_EDITOR
diff --git a/Engine/Plugins/Experimental/ChaosVD/Source/ChaosVD/Private/ChaosVDPlaybackViewportClient.cpp b/Engine/Plugins/Experimental/ChaosVD/Source/ChaosVD/Private/ChaosVDPlaybackViewportClient.cpp
--- a/Engine/Plugins/Experimental/ChaosVD/Source/ChaosVD/Private/ChaosVDPlaybackViewportClient.cpp
+++ b/Engine/Plugins/Experimental/ChaosVD/Source/ChaosVD/Private/ChaosVDPlaybackViewportClient.cpp
@@ -18,6 +18,7 @@
 #include "Selection.h"
 #include "UnrealWidget.h"
 #include "Actors/ChaosVDSolverInfoActor.h"
+#include "Components/ChaosVDParticleDataComponent.h"
 #include "Components/ChaosVDSolverCollisionDataComponent.h"
 #include "Components/InstancedStaticMeshComponent.h"
 #include "Visualizers/ChaosVDDebugDrawUtils.h"
@@ -344,6 +345,14 @@ void FChaosVDPlaybackViewportClient::Draw(const FSceneView* View, FPrimitiveDraw
 	{
 		for (const TPair<int32, AChaosVDSolverInfoActor*>& SolverInfoWithID : ScenePtr->GetSolverInfoActorsMap())
 		{
+			if (const UChaosVDParticleDataComponent* ParticleDataComponent = SolverInfoWithID.Value ? SolverInfoWithID.Value->GetParticleDataComponent() : nullptr)
+			{
+				if (const TSharedPtr<FComponentVisualizer> Visualizer = MainTabToolkitHost->FindComponentVisualizer(ParticleDataComponent->StaticClass()))
+				{
+					Visualizer->DrawVisualization(ParticleDataComponent, View, PDI);
+				}
+			}
+	
 			if (const UChaosVDSolverCollisionDataComponent* CollisionDataComponent = SolverInfoWithID.Value ? SolverInfoWithID.Value->GetCollisionDataComponent() : nullptr)
 			{
 				if (const TSharedPtr<FComponentVisualizer> Visualizer = MainTabToolkitHost->FindComponentVisualizer(CollisionDataComponent->StaticClass()))
@@ -361,15 +370,6 @@ void FChaosVDPlaybackViewportClient::Draw(const FSceneView* View, FPrimitiveDraw
 			}
 		}
 		
-		TArray<AActor*> SelectedActors = ScenePtr->GetElementSelectionSet()->GetSelectedObjects<AActor>();
-
-		for (AActor* SelectedActor : SelectedActors)
-		{
-			if (IChaosVDVisualizerContainerInterface* VisualizerContainer = Cast<IChaosVDVisualizerContainerInterface>(SelectedActor))
-			{
-				VisualizerContainer->DrawVisualization(View, PDI);
-			}
-		}
 	}
 
 	FEditorViewportClient::Draw(View, PDI);
diff --git a/Engine/Plugins/Experimental/ChaosVD/Source/ChaosVD/Private/ChaosVDScene.cpp b/Engine/Plugins/Experimental/ChaosVD/Source/ChaosVD/Private/ChaosVDScene.cpp
--- a/Engine/Plugins/Experimental/ChaosVD/Source/ChaosVD/Private/ChaosVDScene.cpp
+++ b/Engine/Plugins/Experimental/ChaosVD/Source/ChaosVD/Private/ChaosVDScene.cpp
@@ -138,6 +138,8 @@ void FChaosVDScene::UpdateFromRecordedStepData(const int32 SolverID, const FChao
 	{
 		return;
 	}
+
+	SolverSceneData->SetSimulationTransform(InFrameData.SimulationTransform);
 	
 	TSet<int32> ParticlesIDsInRecordedStepData;
 	ParticlesIDsInRecordedStepData.Reserve(InRecordedStepData.RecordedParticlesData.Num());
diff --git a/Engine/Plugins/Experimental/ChaosVD/Source/ChaosVD/Private/Visualizers/ChaosVDDataVisualizerBase.cpp b/Engine/Plugins/Experimental/ChaosVD/Source/ChaosVD/Private/Visualizers/ChaosVDDataVisualizerBase.cpp
deleted file mode 100644
diff --git a/Engine/Plugins/Experimental/ChaosVD/Source/ChaosVD/Private/Visualizers/ChaosVDDebugDrawUtils.cpp b/Engine/Plugins/Experimental/ChaosVD/Source/ChaosVD/Private/Visualizers/ChaosVDDebugDrawUtils.cpp
--- a/Engine/Plugins/Experimental/ChaosVD/Source/ChaosVD/Private/Visualizers/ChaosVDDebugDrawUtils.cpp
+++ b/Engine/Plugins/Experimental/ChaosVD/Source/ChaosVD/Private/Visualizers/ChaosVDDebugDrawUtils.cpp
@@ -11,7 +11,7 @@
 
 TQueue<FChaosVDDebugDrawUtils::FChaosVDQueuedTextToDraw> FChaosVDDebugDrawUtils::TexToDrawQueue = TQueue<FChaosVDQueuedTextToDraw>();
 
-void FChaosVDDebugDrawUtils::DrawArrowVector(FPrimitiveDrawInterface* PDI, const FVector& StartLocation, const FVector& EndLocation, FStringView DebugText, const FColor& Color, ESceneDepthPriorityGroup DepthPriority)
+void FChaosVDDebugDrawUtils::DrawArrowVector(FPrimitiveDrawInterface* PDI, const FVector& StartLocation, const FVector& EndLocation, FStringView DebugText, const FColor& Color, ESceneDepthPriorityGroup DepthPriority, float Thickness)
 {
 	if (!PDI)
 	{
@@ -35,7 +35,7 @@ void FChaosVDDebugDrawUtils::DrawArrowVector(FPrimitiveDrawInterface* PDI, const
 	const float ProportionalArrowSize = MaxTipOfArrowSize * (ArrowLength / MaxVectorSizeForArrow);
 	const float ArrowSize = FMath::Clamp(ProportionalArrowSize, MinTipOfArrowSize, MaxTipOfArrowSize);
 	
-	DrawDirectionalArrow(PDI, ArrowTransformMatrix, Color, ArrowLength, ArrowSize, DepthPriority);
+	DrawDirectionalArrow(PDI, ArrowTransformMatrix, Color, ArrowLength, ArrowSize, DepthPriority, Thickness);
 
 	if (!DebugText.IsEmpty())
 	{
@@ -83,7 +83,7 @@ void FChaosVDDebugDrawUtils::DrawCircle(FPrimitiveDrawInterface* PDI, const FVec
 		return;
 	}
 
-	constexpr float DepthBias = 0;
+	constexpr float DepthBias = 0.0f;
 	const bool bScreenSpace = Thickness > 0;
 
 	// Need at least 2 sides
@@ -106,7 +106,7 @@ void FChaosVDDebugDrawUtils::DrawCircle(FPrimitiveDrawInterface* PDI, const FVec
 	}
 }
 
-void FChaosVDDebugDrawUtils::DrawBox(FPrimitiveDrawInterface* PDI, const FVector& InExtents, const FColor& InColor, const FTransform& InTransform, FStringView DebugText, ESceneDepthPriorityGroup DepthPriority)
+void FChaosVDDebugDrawUtils::DrawBox(FPrimitiveDrawInterface* PDI, const FVector& InExtents, const FColor& InColor, const FTransform& InTransform, FStringView DebugText, ESceneDepthPriorityGroup DepthPriority, float Thickness)
 {
 	if (!PDI)
 	{
@@ -132,9 +132,8 @@ void FChaosVDDebugDrawUtils::DrawBox(FPrimitiveDrawInterface* PDI, const FVector
 		{FVector(-1, 1,1), FVector(-1,1,-1)},
 	};
 
-	constexpr float Thickness = 2.0f;
-	constexpr float DepthBias = 0;
-	constexpr bool bScreenSpace = Thickness > 0;
+	constexpr float DepthBias = 0.0f;
+	const bool bScreenSpace = Thickness > 0;
 
 	for (int32 BoxLineIndex = 0; BoxLineIndex < MaxBoxLines; BoxLineIndex++)
 	{
@@ -150,16 +149,15 @@ void FChaosVDDebugDrawUtils::DrawBox(FPrimitiveDrawInterface* PDI, const FVector
 	}
 }
 
-void FChaosVDDebugDrawUtils::DrawLine(FPrimitiveDrawInterface* PDI, const FVector& InStartPosition, const FVector& InEndPosition, const FColor& InColor, FStringView DebugText, ESceneDepthPriorityGroup DepthPriority)
+void FChaosVDDebugDrawUtils::DrawLine(FPrimitiveDrawInterface* PDI, const FVector& InStartPosition, const FVector& InEndPosition, const FColor& InColor, FStringView DebugText, ESceneDepthPriorityGroup DepthPriority, float Thickness)
 {
 	if (!PDI)
 	{
 		return;
 	}
 
-	constexpr float Thickness = 2.0f;
-	constexpr float DepthBias = 0;
-	constexpr bool bScreenSpace = Thickness > 0;
+	constexpr float DepthBias = 0.0f;
+	const bool bScreenSpace = Thickness > 0;
 
 	PDI->DrawLine(InStartPosition, InEndPosition, InColor, DepthPriority, Thickness, DepthBias, bScreenSpace);
 
@@ -171,7 +169,7 @@ void FChaosVDDebugDrawUtils::DrawLine(FPrimitiveDrawInterface* PDI, const FVecto
 	}
 }
 
-void FChaosVDDebugDrawUtils::DrawImplicitObject(FPrimitiveDrawInterface* PDI, const TSharedPtr<FChaosVDGeometryBuilder>& GeometryGenerator, const Chaos::FConstImplicitObjectPtr& ImplicitObject, const FTransform& InWorldTransform, const FColor& InColor, FStringView DebugText, ESceneDepthPriorityGroup DepthPriority)
+void FChaosVDDebugDrawUtils::DrawImplicitObject(FPrimitiveDrawInterface* PDI, const TSharedPtr<FChaosVDGeometryBuilder>& GeometryGenerator, const Chaos::FConstImplicitObjectPtr& ImplicitObject, const FTransform& InWorldTransform, const FColor& InColor, FStringView DebugText, ESceneDepthPriorityGroup DepthPriority, float Thickness)
 {
 	if (!PDI)
 	{
@@ -196,9 +194,9 @@ void FChaosVDDebugDrawUtils::DrawImplicitObject(FPrimitiveDrawInterface* PDI, co
 			FVector VertexB = AdjustedTransform.TransformPosition(MeshGenerator->Vertices[Triangle.B]);
 			FVector VertexC = AdjustedTransform.TransformPosition(MeshGenerator->Vertices[Triangle.C]);
 	
-			DrawLine(PDI, VertexA, VertexB, InColor, nullptr, DepthPriority);
-			DrawLine(PDI, VertexB, VertexC, InColor, nullptr, DepthPriority);
-			DrawLine(PDI, VertexC, VertexA, InColor, nullptr, DepthPriority);
+			DrawLine(PDI, VertexA, VertexB, InColor, nullptr, DepthPriority, Thickness);
+			DrawLine(PDI, VertexB, VertexC, InColor, nullptr, DepthPriority, Thickness);
+			DrawLine(PDI, VertexC, VertexA, InColor, nullptr, DepthPriority, Thickness);
 		}	
 	}
 
@@ -208,7 +206,6 @@ void FChaosVDDebugDrawUtils::DrawImplicitObject(FPrimitiveDrawInterface* PDI, co
 	}
 }
 
-
 void FChaosVDDebugDrawUtils::DrawCanvas(FViewport& InViewport, FSceneView& View, FCanvas& Canvas)
 {
 	if (!GEngine)
@@ -221,6 +218,11 @@ void FChaosVDDebugDrawUtils::DrawCanvas(FViewport& InViewport, FSceneView& View,
 		FChaosVDQueuedTextToDraw TextToDraw;
 		if (TexToDrawQueue.Dequeue(TextToDraw))
 		{
+			if (!View.ViewFrustum.IntersectPoint(TextToDraw.WorldPosition))
+			{
+				continue;
+			}
+
 			FVector2D PixelLocation;
 			if (View.WorldToPixel(TextToDraw.WorldPosition, PixelLocation))
 			{
diff --git a/Engine/Plugins/Experimental/ChaosVD/Source/ChaosVD/Private/Visualizers/ChaosVDParticleDataComponentVisualizer.cpp b/Engine/Plugins/Experimental/ChaosVD/Source/ChaosVD/Private/Visualizers/ChaosVDParticleDataComponentVisualizer.cpp
new file mode 100644
--- /dev/null
+++ b/Engine/Plugins/Experimental/ChaosVD/Source/ChaosVD/Private/Visualizers/ChaosVDParticleDataComponentVisualizer.cpp
@@ -0,0 +1,244 @@
+﻿// Copyright Epic Games, Inc. All Rights Reserved.
+
+#include "Visualizers/ChaosVDParticleDataComponentVisualizer.h"
+
+#include "Actors/ChaosVDSolverInfoActor.h"
+#include "ChaosVDEditorSettings.h"
+#include "ChaosVDParticleActor.h"
+#include "ChaosVDScene.h"
+#include "SceneView.h"
+#include "Components/ChaosVDParticleDataComponent.h"
+#include "DataWrappers/ChaosVDParticleDataWrapper.h"
+#include "Visualizers/ChaosVDDebugDrawUtils.h"
+
+IMPLEMENT_HIT_PROXY(HChaosVDParticleDataProxy, HComponentVisProxy)
+
+/** Sets a Hit proxy which will be cleared out as soon this struct goes out of scope*/
+struct FChaosVDScopedParticleHitProxy
+{
+	FChaosVDScopedParticleHitProxy(FPrimitiveDrawInterface* PDI, HHitProxy* HitProxy)
+	{
+		PDIPtr = PDI;
+		if (PDIPtr)
+		{
+			PDIPtr->SetHitProxy(HitProxy);
+		}
+	}
+
+	~FChaosVDScopedParticleHitProxy()
+	{
+		if (PDIPtr)
+		{
+			PDIPtr->SetHitProxy(nullptr);
+		}
+	}
+	
+	FPrimitiveDrawInterface* PDIPtr = nullptr;
+};
+
+void FChaosVDParticleDataComponentVisualizer::DrawVisualization(const UActorComponent* Component, const FSceneView* View, FPrimitiveDrawInterface* PDI)
+{
+	const UChaosVDEditorSettings* EditorSettings = GetDefault<UChaosVDEditorSettings>();
+	if (!EditorSettings)
+	{
+		return;
+	}
+
+	if (EditorSettings->GlobalParticleDataVisualizationFlags == 0)
+	{
+		// Nothing to visualize
+		return;
+	}
+
+	const UChaosVDParticleDataComponent* ParticleDataComponent = Cast<UChaosVDParticleDataComponent>(Component);
+	if (!ParticleDataComponent)
+	{
+		return;
+	}
+
+	AChaosVDSolverInfoActor* SolverDataActor = Cast<AChaosVDSolverInfoActor>(Component->GetOwner());
+	if (!SolverDataActor)
+	{
+		return;
+	}
+
+	const TSharedPtr<FChaosVDScene> CVDScene = SolverDataActor->GetScene().Pin();
+	if (!CVDScene)
+	{
+		return;
+	}
+
+	FChaosVDParticleDataVisualizationContext VisualizationContext;
+	VisualizationContext.VisualizationFlags = EditorSettings->GlobalParticleDataVisualizationFlags;
+	VisualizationContext.SpaceTransform = SolverDataActor->GetSimulationTransform();
+	VisualizationContext.CVDScene = CVDScene;
+	VisualizationContext.GeometryGenerator = CVDScene->GetGeometryGenerator();
+	VisualizationContext.bShowDebugText = EditorSettings->bShowDebugText;
+	VisualizationContext.DebugDrawSettings = &EditorSettings->ParticleDataDebugDrawSettings;
+
+	if (!VisualizationContext.IsVisualizationFlagEnabled(EChaosVDParticleDataVisualizationFlags::EnableDraw))
+	{
+		return;
+	}
+
+	if (VisualizationContext.IsVisualizationFlagEnabled(EChaosVDParticleDataVisualizationFlags::DrawDataOnlyForSelectedParticle))
+	{
+		VisualizationContext.bIsSelectedData = true;
+		SolverDataActor->VisitSelectedParticleData([this, PDI, View, &VisualizationContext, Component](const FChaosVDParticleDataWrapper& InParticleDataViewer)
+		{
+			DrawVisualizationForParticleData(Component, PDI, View, VisualizationContext, InParticleDataViewer);
+
+			return true;
+		});
+	}
+	else
+	{
+		SolverDataActor->VisitAllParticleData([this, PDI, View, &VisualizationContext, Component, SolverDataActor](const FChaosVDParticleDataWrapper& InParticleDataViewer)
+		{
+			VisualizationContext.bIsSelectedData = SolverDataActor->IsParticleSelectedByID(InParticleDataViewer.ParticleIndex);
+			DrawVisualizationForParticleData(Component, PDI, View, VisualizationContext, InParticleDataViewer);
+
+			return true;
+		});
+	}
+}
+
+bool FChaosVDParticleDataComponentVisualizer::VisProxyHandleClick(FEditorViewportClient* InViewportClient, HComponentVisProxy* VisProxy, const FViewportClick& Click)
+{
+	bool bHandled = false;
+
+	const HChaosVDParticleDataProxy* ParticleDataProxy = HitProxyCast<HChaosVDParticleDataProxy>(VisProxy);
+	if (!ParticleDataProxy)
+	{
+		return bHandled;
+	}
+	
+	const UChaosVDParticleDataComponent* ParticleDataComponent = Cast<UChaosVDParticleDataComponent>(VisProxy->Component.Get());
+	if (!ParticleDataComponent)
+	{
+		return bHandled;
+	}
+
+	AChaosVDSolverInfoActor* SolverDataActor = Cast<AChaosVDSolverInfoActor>(ParticleDataComponent->GetOwner());
+	if (!SolverDataActor)
+	{
+		return bHandled;
+	}
+
+	bHandled = SolverDataActor->SelectParticleByID(ParticleDataProxy->DataSelectionHandle.ParticleIndex);
+
+	return bHandled;	
+}
+
+FString FChaosVDParticleDataComponentVisualizer::GenerateDebugTextForVector(const FVector& InVector, const FString& VectorName, const FString& InVectorUnits)
+{
+	return FString::Format(TEXT("{5} : {0} {4} \n	|-- X : {1} {4} \n	|-- Y : {2} {4} \n	|-- Z : {3} {4}"), {InVector.Size(), InVector.X, InVector.Y, InVector.Z, InVectorUnits, VectorName });
+}
+
+void FChaosVDParticleDataComponentVisualizer::DrawParticleVector(FPrimitiveDrawInterface* PDI, const FVector& StartLocation, const FVector& InVector, EChaosVDParticleDataVisualizationFlags VectorID, const FChaosVDParticleDataVisualizationContext& InVisualizationContext, float LineThickness)
+{
+	if (!InVisualizationContext.IsVisualizationFlagEnabled(VectorID))
+	{
+		return;
+	}
+
+	if (!ensure(InVisualizationContext.DebugDrawSettings))
+	{
+		return;
+	}
+
+	const FString DebugText = InVisualizationContext.bShowDebugText ? GenerateDebugTextForVector(InVector, UEnum::GetDisplayValueAsText(VectorID).ToString(), Chaos::VisualDebugger::ParticleDataUnitsStrings::GetUnitByID(VectorID)) : TEXT("");
+	FChaosVDDebugDrawUtils::DrawArrowVector(PDI, StartLocation, StartLocation +  InVisualizationContext.DebugDrawSettings->GetScaleFortDataID(VectorID) * InVector, DebugText, InVisualizationContext.DebugDrawSettings->ColorSettings.GetColorForDataID(VectorID, InVisualizationContext.bIsSelectedData),  InVisualizationContext.DebugDrawSettings->DepthPriority, LineThickness);
+}
+
+void FChaosVDParticleDataComponentVisualizer::DrawVisualizationForParticleData(const UActorComponent* Component, FPrimitiveDrawInterface* PDI, const FSceneView* View, const FChaosVDParticleDataVisualizationContext& InVisualizationContext, const FChaosVDParticleDataWrapper& InParticleDataViewer)
+{
+	using namespace Chaos::VisualDebugger::ParticleDataUnitsStrings;
+
+	if (!View)
+	{
+		return;
+	}
+
+	if (!ensure(InVisualizationContext.DebugDrawSettings))
+	{
+		return;
+	}
+
+	const FVector& OwnerLocation = InVisualizationContext.SpaceTransform.TransformPosition(InParticleDataViewer.ParticlePositionRotation.MX);
+
+	// TODO: See how expensive is get the bounds. It is not something we have recorded
+	constexpr float VisibleRadius = 50.0f;
+	if (!View->ViewFrustum.IntersectSphere(OwnerLocation, VisibleRadius))
+	{
+		// If this particle location is not even visible, just ignore it.
+		return;
+	}
+
+	const FQuat& OwnerRotation =  InVisualizationContext.SpaceTransform.TransformRotation(InParticleDataViewer.ParticlePositionRotation.MR);
+	const FVector OwnerCoMLocation = OwnerLocation + OwnerRotation *  InVisualizationContext.SpaceTransform.TransformPosition(InParticleDataViewer.ParticleMassProps.MCenterOfMass);
+	
+	FChaosVDScopedParticleHitProxy ScopedHitProxy(PDI, new HChaosVDParticleDataProxy(Component, {InParticleDataViewer.ParticleIndex, InParticleDataViewer.SolverID} ));
+
+	constexpr float DefaultLineThickness = 1.5f;
+	constexpr float SelectedLineThickness = 3.5f;
+	const float LineThickness = InVisualizationContext.bIsSelectedData ? SelectedLineThickness : DefaultLineThickness;
+	
+
+	if (InParticleDataViewer.ParticleVelocities.HasValidData())
+	{
+		DrawParticleVector(PDI, OwnerCoMLocation, InParticleDataViewer.ParticleVelocities.MV,EChaosVDParticleDataVisualizationFlags::Velocity, InVisualizationContext, LineThickness); 
+		DrawParticleVector(PDI, OwnerCoMLocation, InParticleDataViewer.ParticleVelocities.MW,EChaosVDParticleDataVisualizationFlags::AngularVelocity, InVisualizationContext, LineThickness); 
+	}
+
+	if (InParticleDataViewer.ParticleDynamics.HasValidData())
+	{
+		DrawParticleVector(PDI, OwnerCoMLocation, InParticleDataViewer.ParticleDynamics.MAcceleration,EChaosVDParticleDataVisualizationFlags::Acceleration, InVisualizationContext, LineThickness); 
+		DrawParticleVector(PDI, OwnerCoMLocation, InParticleDataViewer.ParticleDynamics.MAngularAcceleration,EChaosVDParticleDataVisualizationFlags::AngularAcceleration, InVisualizationContext, LineThickness); 
+		DrawParticleVector(PDI, OwnerCoMLocation, InParticleDataViewer.ParticleDynamics.MLinearImpulseVelocity,EChaosVDParticleDataVisualizationFlags::LinearImpulse, InVisualizationContext, LineThickness); 
+		DrawParticleVector(PDI, OwnerCoMLocation, InParticleDataViewer.ParticleDynamics.MAngularImpulseVelocity,EChaosVDParticleDataVisualizationFlags::AngularImpulse, InVisualizationContext, LineThickness);
+	}
+
+	if (InParticleDataViewer.ParticleMassProps.HasValidData())
+	{
+		if (InVisualizationContext.IsVisualizationFlagEnabled(EChaosVDParticleDataVisualizationFlags::CenterOfMass))
+		{
+			if (const TSharedPtr<FChaosVDGeometryBuilder> GeometryGenerator = InVisualizationContext.GeometryGenerator.Pin())
+			{
+				FCollisionShape Sphere;
+				Sphere.SetSphere(InVisualizationContext.DebugDrawSettings->CenterOfMassRadius);
+				const FPhysicsShapeAdapter SphereShapeAdapter(FQuat::Identity, Sphere);
+
+				const FString DebugText = UEnum::GetDisplayValueAsText(EChaosVDParticleDataVisualizationFlags::CenterOfMass).ToString();
+				FChaosVDDebugDrawUtils::DrawImplicitObject(PDI, GeometryGenerator, &SphereShapeAdapter.GetGeometry(), FTransform(OwnerCoMLocation), InVisualizationContext.DebugDrawSettings->ColorSettings.GetColorForDataID(EChaosVDParticleDataVisualizationFlags::CenterOfMass, InVisualizationContext.bIsSelectedData), DebugText, InVisualizationContext.DebugDrawSettings->DepthPriority, LineThickness);
+			}
+		}
+	}
+
+	// TODO: This is a Proof of concept to test how debug draw connectivity data will look
+	if (InParticleDataViewer.ParticleCluster.HasValidData())
+	{
+		if (InVisualizationContext.IsVisualizationFlagEnabled(EChaosVDParticleDataVisualizationFlags::ClusterConnectivityEdge))
+		{
+			for (const FChaosVDConnectivityEdge& ConnectivityEdge : InParticleDataViewer.ParticleCluster.ConnectivityEdges)
+			{
+				if (const TSharedPtr<FChaosVDScene> ScenePtr = InVisualizationContext.CVDScene.Pin())
+				{
+					if (AChaosVDParticleActor* SiblingParticle = ScenePtr->GetParticleActor(InVisualizationContext.SolverID, ConnectivityEdge.SiblingParticleID))
+					{
+						if (const FChaosVDParticleDataWrapper* SiblingParticleData = SiblingParticle->GetParticleData())
+						{
+							FColor DebugDrawColor = InVisualizationContext.DebugDrawSettings->ColorSettings.GetColorForDataID(EChaosVDParticleDataVisualizationFlags::ClusterConnectivityEdge, InVisualizationContext.bIsSelectedData);
+							FVector BoxExtents(2,2,2);
+							FTransform BoxTransform(OwnerRotation, OwnerLocation);
+							FChaosVDDebugDrawUtils::DrawBox(PDI, BoxExtents, DebugDrawColor, BoxTransform, nullptr, InVisualizationContext.DebugDrawSettings->DepthPriority, LineThickness);
+
+							FVector SiblingParticleLocation = InVisualizationContext.SpaceTransform.TransformPosition(SiblingParticleData->ParticlePositionRotation.MX);
+							FChaosVDDebugDrawUtils::DrawLine(PDI, OwnerLocation, SiblingParticleLocation, DebugDrawColor, FString::Printf(TEXT("Strain [%f]"), ConnectivityEdge.Strain), InVisualizationContext.DebugDrawSettings->DepthPriority, LineThickness);
+						}
+					}	
+				}
+			}
+		}
+	}
+}
diff --git a/Engine/Plugins/Experimental/ChaosVD/Source/ChaosVD/Private/Visualizers/ChaosVDParticleDataVisualizer.cpp b/Engine/Plugins/Experimental/ChaosVD/Source/ChaosVD/Private/Visualizers/ChaosVDParticleDataVisualizer.cpp
deleted file mode 100644
diff --git a/Engine/Plugins/Experimental/ChaosVD/Source/ChaosVD/Private/Visualizers/ChaosVDSceneQueryDataComponentVisualizer.cpp b/Engine/Plugins/Experimental/ChaosVD/Source/ChaosVD/Private/Visualizers/ChaosVDSceneQueryDataComponentVisualizer.cpp
--- a/Engine/Plugins/Experimental/ChaosVD/Source/ChaosVD/Private/Visualizers/ChaosVDSceneQueryDataComponentVisualizer.cpp
+++ b/Engine/Plugins/Experimental/ChaosVD/Source/ChaosVD/Private/Visualizers/ChaosVDSceneQueryDataComponentVisualizer.cpp
@@ -12,6 +12,7 @@
 #include "ChaosVDTabsIDs.h"
 #include "EditorModeManager.h"
 #include "EditorViewportClient.h"
+#include "SceneView.h"
 #include "Visualizers/ChaosVDDebugDrawUtils.h"
 #include "Visualizers/IChaosVDParticleVisualizationDataProvider.h"
 #include "Widgets/SChaosVDMainTab.h"
@@ -268,6 +269,16 @@ void FChaosVDSceneQueryDataComponentVisualizer::DrawSceneQuery(const UActorCompo
 		return;
 	}
 
+	//TODO: Should we try to calculate actual bounds?
+	constexpr float MinVisibleRadius = 100.0f;
+	const float QueryHalfDistance = (Query->EndLocation - Query->StartLocation).Size() * 0.5;
+	const float VisibleRadius = FMath::Max(QueryHalfDistance, MinVisibleRadius);
+	if (!View->ViewFrustum.IntersectSphere(Query->StartLocation, VisibleRadius))
+	{
+		// If this query location is not even visible, just ignore it.
+		return;
+	}
+
 	if (const Chaos::FConstImplicitObjectPtr* InputShapePtrPtr = CVDRecording->GetGeometryMap().Find(Query->InputGeometryKey))
 	{
 		VisualizationContext.InputGeometry = *InputShapePtrPtr;
diff --git a/Engine/Plugins/Experimental/ChaosVD/Source/ChaosVD/Private/Widgets/SChaosVDMainTab.cpp b/Engine/Plugins/Experimental/ChaosVD/Source/ChaosVD/Private/Widgets/SChaosVDMainTab.cpp
--- a/Engine/Plugins/Experimental/ChaosVD/Source/ChaosVD/Private/Widgets/SChaosVDMainTab.cpp
+++ b/Engine/Plugins/Experimental/ChaosVD/Source/ChaosVD/Private/Widgets/SChaosVDMainTab.cpp
@@ -27,10 +27,12 @@
 #include "IDesktopPlatform.h"
 #include "Misc/MessageDialog.h"
 #include "StatusBarSubsystem.h"
+#include "Components/ChaosVDParticleDataComponent.h"
 #include "Components/ChaosVDSolverCollisionDataComponent.h"
 #include "Styling/StyleColors.h"
 #include "Styling/ToolBarStyle.h"
 #include "Trace/ChaosVDTraceManager.h"
+#include "Visualizers/ChaosVDParticleDataComponentVisualizer.h"
 #include "Visualizers/ChaosVDSceneQueryDataComponentVisualizer.h"
 #include "Visualizers/ChaosVDSolverCollisionDataComponentVisualizer.h"
 #include "Widgets/SChaosBrowseTraceFileSourceModal.h"
@@ -53,6 +55,7 @@ void SChaosVDMainTab::Construct(const FArguments& InArgs, TSharedPtr<FChaosVDEng
 
 	RegisterComponentVisualizer(UChaosVDSolverCollisionDataComponent::StaticClass()->GetFName(), MakeShared<FChaosVDSolverCollisionDataComponentVisualizer>());
 	RegisterComponentVisualizer(UChaosVDSceneQueryDataComponent::StaticClass()->GetFName(), MakeShared<FChaosVDSceneQueryDataComponentVisualizer>());
+	RegisterComponentVisualizer(UChaosVDParticleDataComponent::StaticClass()->GetFName(), MakeShared<FChaosVDParticleDataComponentVisualizer>());
 
 	TabManager = FGlobalTabmanager::Get()->NewTabManager(InArgs._OwnerTab.ToSharedRef()).ToSharedPtr();
 
diff --git a/Engine/Plugins/Experimental/ChaosVD/Source/ChaosVD/Public/Actors/ChaosVDSolverInfoActor.h b/Engine/Plugins/Experimental/ChaosVD/Source/ChaosVD/Public/Actors/ChaosVDSolverInfoActor.h
--- a/Engine/Plugins/Experimental/ChaosVD/Source/ChaosVD/Public/Actors/ChaosVDSolverInfoActor.h
+++ b/Engine/Plugins/Experimental/ChaosVD/Source/ChaosVD/Public/Actors/ChaosVDSolverInfoActor.h
@@ -2,13 +2,16 @@
 
 #pragma once
 
+#include "ChaosVDParticleActor.h"
 #include "ChaosVDSceneObjectBase.h"
 #include "ChaosVDSceneSelectionObserver.h"
 #include "GameFramework/Actor.h"
 #include "ChaosVDSolverInfoActor.generated.h"
 
+struct FChaosVDParticleDataWrapper;
 class AChaosVDParticleActor;
 class UChaosVDSolverCollisionDataComponent;
+class UChaosVDParticleDataComponent;
 
 enum class EChaosVDParticleType : uint8;
 
@@ -36,6 +39,7 @@ public:
 	const FTransform& GetSimulationTransform() const { return SimulationTransform; }
 
 	UChaosVDSolverCollisionDataComponent* GetCollisionDataComponent() { return CollisionDataComponent; }
+	UChaosVDParticleDataComponent* GetParticleDataComponent() { return ParticleDataComponent; }
 
 	void RegisterParticleActor(int32 ParticleID, AChaosVDParticleActor* ParticleActor);
 
@@ -44,6 +48,16 @@ public:
 
 	const TArray<int32>& GetSelectedParticlesIDs() const { return SelectedParticlesID; }
 
+	bool IsParticleSelectedByID(int32 ParticleID);
+
+	bool SelectParticleByID(int32 ParticleIDToSelect);
+
+	template <typename TCallback>
+	void VisitSelectedParticleData(TCallback VisitCallback);
+
+	template <typename TCallback>
+	void VisitAllParticleData(TCallback VisitCallback);
+
 	void HandleVisibilitySettingsUpdated();
 	void HandleColorsSettingsUpdated();
 	void RemoveSolverFolders(UWorld* World);
@@ -78,4 +92,45 @@ protected:
 	TSet<FFolder> CreatedFolders;
 
 	bool bIsServer;
+
+	UPROPERTY()
+	TObjectPtr<UChaosVDParticleDataComponent> ParticleDataComponent;
 };
+
+template <typename TCallback>
+void AChaosVDSolverInfoActor::VisitSelectedParticleData(TCallback VisitCallback)
+{
+	for (const int32 SelectedParticleID : SelectedParticlesID)
+	{
+		AChaosVDParticleActor* ParticleActor = GetParticleActor(SelectedParticleID);
+		const FChaosVDParticleDataWrapper* ParticleDataViewer = ParticleActor ? ParticleActor->GetParticleData() : nullptr;
+		if (!ensure(ParticleDataViewer))
+		{
+			continue;
+		}
+
+		if (!VisitCallback(*ParticleDataViewer))
+		{
+			return;
+		}
+	}
+}
+
+template <typename TCallback>
+void AChaosVDSolverInfoActor::VisitAllParticleData(TCallback VisitCallback)
+{
+	for (const TPair<int32, AChaosVDParticleActor*>& ParticleWithIDPair : SolverParticlesByID)
+	{
+		AChaosVDParticleActor* ParticleActor = ParticleWithIDPair.Value;
+		const FChaosVDParticleDataWrapper* ParticleDataViewer = ParticleActor ? ParticleActor->GetParticleData() : nullptr;
+		if (!ensure(ParticleDataViewer))
+		{
+			continue;
+		}
+
+		if (!VisitCallback(*ParticleDataViewer))
+		{
+			return;
+		}
+	}
+}
diff --git a/Engine/Plugins/Experimental/ChaosVD/Source/ChaosVD/Public/ChaosVDEditorSettings.h b/Engine/Plugins/Experimental/ChaosVD/Source/ChaosVD/Public/ChaosVDEditorSettings.h
--- a/Engine/Plugins/Experimental/ChaosVD/Source/ChaosVD/Public/ChaosVDEditorSettings.h
+++ b/Engine/Plugins/Experimental/ChaosVD/Source/ChaosVD/Public/ChaosVDEditorSettings.h
@@ -7,7 +7,7 @@
 #include "Engine/StaticMesh.h"
 #include "UObject/Object.h"
 #include "UObject/SoftObjectPtr.h"
-#include "Visualizers/ChaosVDParticleDataVisualizer.h"
+#include "Visualizers/ChaosVDParticleDataComponentVisualizer.h"
 #include "Visualizers/ChaosVDSceneQueryDataComponentVisualizer.h"
 #include "Visualizers/ChaosVDSolverCollisionDataComponentVisualizer.h"
 
@@ -106,6 +106,90 @@ struct FChaosDebugDrawColorsByState
 	FColor GetColorFromState(EChaosVDObjectStateType State) const;
 };
 
+USTRUCT()
+struct FChaosParticleDataDebugDrawColors
+{
+	GENERATED_BODY()
+
+	/** Color to apply to the Velocity vector when drawing it */
+	UPROPERTY(EditAnywhere, Category=DebugDraw)
+	FLinearColor VelocityColor = FColor::Green;
+
+	/** Color to apply to the Angular Velocity vector when drawing it */
+	UPROPERTY(EditAnywhere, Category=DebugDraw)
+	FLinearColor AngularVelocityColor = FColor::Blue;
+
+	/** Color to apply to the Acceleration vector when drawing it */
+	UPROPERTY(EditAnywhere, Category=DebugDraw)
+	FLinearColor AccelerationColor = FColor::Orange;
+
+	/** Color to apply to the Angular Acceleration vector when drawing it */
+	UPROPERTY(EditAnywhere, Category=DebugDraw)
+	FLinearColor AngularAccelerationColor = FColor::Silver;
+
+	/** Color to apply to the Linear Impulse when drawing it */
+	UPROPERTY(EditAnywhere, Category=DebugDraw)
+	FLinearColor LinearImpulseColor = FColor::Turquoise;
+
+	/** Color to apply to the Angular Impulse vector when drawing it */
+	UPROPERTY(EditAnywhere, Category=DebugDraw)
+	FLinearColor AngularImpulseColor = FColor::Emerald;
+
+	/** Color to apply the debug drawn sphere representing the center of mass location */
+	UPROPERTY(EditAnywhere, Category=DebugDraw)
+	FLinearColor CenterOfMassColor = FColor::Red;
+
+	/** Color to apply to when drawing the connectivity data */
+	UPROPERTY(EditAnywhere, Category=DebugDraw)
+	FLinearColor ConnectivityDataColor = FColor::Yellow;
+	
+	FColor GetColorForDataID(EChaosVDParticleDataVisualizationFlags DataID, bool bIsSelected = false) const;
+	const FLinearColor& GetLinearColorForDataID(EChaosVDParticleDataVisualizationFlags DataID) const;
+};
+
+USTRUCT()
+struct FChaosParticleDataDebugDrawSettings
+{
+	GENERATED_BODY()
+
+	/** The depth priority used for while drawing contact data. Can be World or Foreground (with this one the shapes will be drawn on top of the geometry and be always visible) */
+	UPROPERTY(EditAnywhere, Category=DebugDraw)
+	TEnumAsByte<ESceneDepthPriorityGroup> DepthPriority = ESceneDepthPriorityGroup::SDPG_World;
+
+	/** Scale to apply to the Velocity vector before draw it. Unit is cm/s */
+	UPROPERTY(EditAnywhere, Category=DebugDraw)
+	float VelocityScale = 0.5f;
+
+	/** Scale to apply to the Angular Velocity vector before draw it. Unit is rad/s */
+	UPROPERTY(EditAnywhere, Category=DebugDraw)
+	float AngularVelocityScale = 50.0f;
+
+	/** Scale to apply to the Acceleration vector before draw it. Unit is cm/s2 */
+	UPROPERTY(EditAnywhere, Category=DebugDraw)
+	float AccelerationScale = 0.005f;
+
+	/** Scale to apply to the Angular Acceleration vector before draw it. Unit is rad/s2 */
+	UPROPERTY(EditAnywhere, Category=DebugDraw)
+	float AngularAccelerationScale = 0.5f;
+
+	/** Scale to apply to the Linear Impulse vector before draw it. Unit is g.m/s */
+	UPROPERTY(EditAnywhere, Category=DebugDraw)
+	float LinearImpulseScale = 0.001;
+
+	/** Scale to apply to the Angular Impulse vector before draw it. Unit is g.m2/s */
+	UPROPERTY(EditAnywhere, Category=DebugDraw)
+	float AngularImpulseScale = 0.1f;
+
+	/** Radius to use when creating the sphere that will represent the center of mass location */
+	UPROPERTY(EditAnywhere, Category=DebugDraw)
+	float CenterOfMassRadius = 10.0f;
+
+	float GetScaleFortDataID(EChaosVDParticleDataVisualizationFlags DataID) const;
+	
+	UPROPERTY(EditAnywhere, Category=DebugDraw)
+	FChaosParticleDataDebugDrawColors ColorSettings;
+};
+
 /** Structure holding the settings using to debug draw Particles shape based on their shape type on the Chaos Visual Debugger */
 USTRUCT()
 struct FChaosDebugDrawColorsByShapeType
@@ -171,10 +255,14 @@ struct FChaosDebugDrawColorsByClientServer
 UENUM()
 enum class EChaosVDParticleDebugColorMode
 {
-	None, // Draw particles with the default gray color
-	State, // Draw particles with a specific color based on the recorded particle state
-	ShapeType, // Draw particles with a specific color based on their shape type
-	ClientServer, // Draw particles with a specific color based on if they are a Server Particle or Client particle
+	/** Draw particles with the default gray color */
+	None,
+	/** Draw particles with a specific color based on the recorded particle state */
+	State,
+	/** Draw particles with a specific color based on their shape type */
+	ShapeType,
+	/** Draw particles with a specific color based on if they are a Server Particle or Client particle */
+	ClientServer,
 };
 
 UCLASS(config = Engine)
@@ -197,7 +285,7 @@ public:
 
 	/** Set of flags to enable/disable visualization of specific particle data as debug draw */
 	UPROPERTY(EditAnywhere, Category = "Viewport Visualization Flags", meta = (Bitmask, BitmaskEnum = "/Script/ChaosVD.EChaosVDParticleDataVisualizationFlags"))
-	uint32 GlobalParticleDataVisualizationFlags = 0;
+	uint32 GlobalParticleDataVisualizationFlags = static_cast<uint32>(EChaosVDParticleDataVisualizationFlags::Velocity | EChaosVDParticleDataVisualizationFlags::AngularVelocity | EChaosVDParticleDataVisualizationFlags::DrawDataOnlyForSelectedParticle);
 
 	/** Set of flags to enable/disable visualization of specific collision data as debug draw */
 	UPROPERTY(EditAnywhere, Category = "Viewport Visualization Flags", meta = (Bitmask, BitmaskEnum = "/Script/ChaosVD.EChaosVDCollisionVisualizationFlags"))
@@ -205,7 +293,7 @@ public:
 
 	/** Set of flags to enable/disable visualization of specific scene queries data as debug draw */
 	UPROPERTY(EditAnywhere, Category = "Viewport Visualization Flags", meta = (Bitmask, BitmaskEnum = "/Script/ChaosVD.EChaosVDSceneQueryVisualizationFlags"))
-	uint32 GlobalSceneQueriesVisualizationFlags = static_cast<uint8>(EChaosVDSceneQueryVisualizationFlags::DrawClientQueries | EChaosVDSceneQueryVisualizationFlags::DrawServerQueries |  EChaosVDSceneQueryVisualizationFlags::DrawHits | EChaosVDSceneQueryVisualizationFlags::DrawLineTraceQueries);
+	uint32 GlobalSceneQueriesVisualizationFlags = static_cast<uint32>(EChaosVDSceneQueryVisualizationFlags::DrawClientQueries | EChaosVDSceneQueryVisualizationFlags::DrawServerQueries |  EChaosVDSceneQueryVisualizationFlags::DrawHits | EChaosVDSceneQueryVisualizationFlags::DrawLineTraceQueries);
 
 	/** If true, text information (if available) will be drawn alongside any other debug draw shape */
 	UPROPERTY(EditAnywhere, Category = "Viewport Visualization")
@@ -214,6 +302,9 @@ public:
 	UPROPERTY(EditAnywhere, Category = "Viewport Visualization",  meta=(EditCondition = "GlobalCollisionDataVisualizationFlags != 0", EditConditionHides))
 	FChaosVDContactDebugDrawSettings ContactDebugDrawSettings;
 
+	UPROPERTY(EditAnywhere, Category = "Viewport Visualization",  meta=(EditCondition = "GlobalParticleDataVisualizationFlags != 0", EditConditionHides))
+	FChaosParticleDataDebugDrawSettings ParticleDataDebugDrawSettings;
+
 	UPROPERTY(EditAnywhere, Category = "Viewport Visualization")
 	EChaosVDParticleDebugColorMode ParticleColorMode;
 	
diff --git a/Engine/Plugins/Experimental/ChaosVD/Source/ChaosVD/Public/ChaosVDParticleActor.h b/Engine/Plugins/Experimental/ChaosVD/Source/ChaosVD/Public/ChaosVDParticleActor.h
--- a/Engine/Plugins/Experimental/ChaosVD/Source/ChaosVD/Public/ChaosVDParticleActor.h
+++ b/Engine/Plugins/Experimental/ChaosVD/Source/ChaosVD/Public/ChaosVDParticleActor.h
@@ -8,7 +8,6 @@
 #include "DataWrappers/ChaosVDParticleDataWrapper.h"
 #include "DataWrappers/ChaosVDCollisionDataWrappers.h"
 #include "GameFramework/Actor.h"
-#include "Visualizers/ChaosVDDataVisualizerBase.h"
 #include "Visualizers/IChaosVDParticleVisualizationDataProvider.h"
 
 #include "ChaosVDParticleActor.generated.h"
@@ -40,7 +39,7 @@ DECLARE_DELEGATE(FChaosVDParticleDataUpdatedDelegate)
 
 /** Actor used to represent a Chaos Particle in the Visual Debugger's world */
 UCLASS(HideCategories=(Transform))
-class AChaosVDParticleActor : public AActor, public IChaosVDParticleVisualizationDataProvider, public IChaosVDVisualizerContainerInterface,
+class AChaosVDParticleActor : public AActor, public IChaosVDParticleVisualizationDataProvider,
 								public FChaosVDSceneObjectBase, public IChaosVDCollisionDataProviderInterface
 {
 	GENERATED_BODY()
@@ -59,15 +58,9 @@ public:
 	virtual void Destroyed() override;
 
 	virtual const FChaosVDParticleDataWrapper* GetParticleData() override { return ParticleDataPtr.Get(); }
-	virtual void GetVisualizationContext(FChaosVDVisualizationContext& OutVisualizationContext) override;
-
-	void CreateVisualizers();
-
-	virtual void DrawVisualization(const FSceneView* View, FPrimitiveDrawInterface* PDI) override;
 	
 #if WITH_EDITOR
 	virtual bool IsSelectedInEditor() const override;
-	virtual void PostEditChangeProperty(FPropertyChangedEvent& PropertyChangedEvent) override;
 	virtual void SetIsTemporarilyHiddenInEditor(bool bIsHidden) override;
 #endif
 
@@ -109,9 +102,6 @@ protected:
 
 	template<typename TTaskCallback>
 	void VisitGeometryInstances(const TTaskCallback& VisitorCallback);
-	
-	UPROPERTY(EditAnywhere, Category = "Viewport Visualization Flags", meta = (Bitmask, BitmaskEnum = "/Script/ChaosVD.EChaosVDParticleDataVisualizationFlags"))
-	uint8 LocalParticleDataVisualizationFlags;
 
 	UPROPERTY(EditAnywhere, Category = "Viewport Visualization Flags")
 	bool bShowDebugText = false;
@@ -124,8 +114,6 @@ protected:
 
 	FDelegateHandle GeometryUpdatedDelegate;
 
-	TMap<FStringView, TUniquePtr<FChaosVDDataVisualizerBase>> CVDVisualizers;
-
 	TArray<TSharedPtr<FChaosVDMeshDataInstanceHandle>> MeshDataHandles;
 
 	FChaosVDParticleDataUpdatedDelegate ParticleDataUpdatedDelegate;
diff --git a/Engine/Plugins/Experimental/ChaosVD/Source/ChaosVD/Public/Components/ChaosVDParticleDataComponent.h b/Engine/Plugins/Experimental/ChaosVD/Source/ChaosVD/Public/Components/ChaosVDParticleDataComponent.h
new file mode 100644
--- /dev/null
+++ b/Engine/Plugins/Experimental/ChaosVD/Source/ChaosVD/Public/Components/ChaosVDParticleDataComponent.h
@@ -0,0 +1,18 @@
+﻿// Copyright Epic Games, Inc. All Rights Reserved.
+
+#pragma once
+
+#include "Components/ActorComponent.h"
+#include "ChaosVDParticleDataComponent.generated.h"
+
+/**
+ * Component that will reference all particle data for a specific solver for the current frame.
+ * Currently it is empty and only used to interface with the component visualizer system for particle data,
+ * but when we eventually remove the dependency on actors to represent particles as planned, the particle data will be referenced in this component.\
+ * (similarly to how the Scene Query Data and Collision Data is referenced in their respective components)
+ */
+UCLASS()
+class UChaosVDParticleDataComponent : public UActorComponent
+{
+	GENERATED_BODY()
+};
diff --git a/Engine/Plugins/Experimental/ChaosVD/Source/ChaosVD/Public/Visualizers/ChaosVDDataVisualizerBase.h b/Engine/Plugins/Experimental/ChaosVD/Source/ChaosVD/Public/Visualizers/ChaosVDDataVisualizerBase.h
deleted file mode 100644
diff --git a/Engine/Plugins/Experimental/ChaosVD/Source/ChaosVD/Public/Visualizers/ChaosVDDebugDrawUtils.h b/Engine/Plugins/Experimental/ChaosVD/Source/ChaosVD/Public/Visualizers/ChaosVDDebugDrawUtils.h
--- a/Engine/Plugins/Experimental/ChaosVD/Source/ChaosVD/Public/Visualizers/ChaosVDDebugDrawUtils.h
+++ b/Engine/Plugins/Experimental/ChaosVD/Source/ChaosVD/Public/Visualizers/ChaosVDDebugDrawUtils.h
@@ -14,13 +14,13 @@ class FChaosVDDebugDrawUtils
 public:
 	//TODO: Create a common args struct that can be passed as parameter to all debug draw methods
 
-	static void DrawArrowVector(FPrimitiveDrawInterface* PDI, const FVector& StartLocation, const FVector& EndLocation, FStringView DebugText, const FColor& Color, ESceneDepthPriorityGroup DepthPriority = SDPG_World);
+	static void DrawArrowVector(FPrimitiveDrawInterface* PDI, const FVector& StartLocation, const FVector& EndLocation, FStringView DebugText, const FColor& Color, ESceneDepthPriorityGroup DepthPriority = SDPG_World, float Thickness = 0.0f);
 	static void DrawPoint(FPrimitiveDrawInterface* PDI, const FVector& Location, FStringView DebugText, const FColor& Color, float Size = 10.0f, ESceneDepthPriorityGroup DepthPriority = SDPG_World);
 	static void DrawText(FStringView StringToDraw, const FVector& Location, const FColor& Color);
 	static void DrawCircle(FPrimitiveDrawInterface* PDI, const FVector& Origin, float Radius, int32 Segments, const FColor& Color, float Thickness, const FVector& XAxis, const FVector& YAxis, FStringView DebugText, ESceneDepthPriorityGroup DepthPriority = SDPG_World);
-	static void DrawBox(FPrimitiveDrawInterface* PDI, const FVector& InExtents, const FColor& InColor, const FTransform& InTransform, FStringView DebugText, ESceneDepthPriorityGroup DepthPriority = SDPG_World);
-	static void DrawLine(FPrimitiveDrawInterface* PDI, const FVector& InStartPosition, const FVector& InEndPosition, const FColor& InColor, FStringView DebugText, ESceneDepthPriorityGroup DepthPriority = SDPG_World);
-	static void DrawImplicitObject(FPrimitiveDrawInterface* PDI, const TSharedPtr<FChaosVDGeometryBuilder>& GeometryGenerator, const Chaos::FConstImplicitObjectPtr& ImplicitObject, const FTransform& InWorldTransform, const FColor& InColor, FStringView DebugText, ESceneDepthPriorityGroup DepthPriority = SDPG_World);
+	static void DrawBox(FPrimitiveDrawInterface* PDI, const FVector& InExtents, const FColor& InColor, const FTransform& InTransform, FStringView DebugText, ESceneDepthPriorityGroup DepthPriority = SDPG_World, float Thickness = 2.0f);
+	static void DrawLine(FPrimitiveDrawInterface* PDI, const FVector& InStartPosition, const FVector& InEndPosition, const FColor& InColor, FStringView DebugText, ESceneDepthPriorityGroup DepthPriority = SDPG_World, float Thickness = 2.0f);
+	static void DrawImplicitObject(FPrimitiveDrawInterface* PDI, const TSharedPtr<FChaosVDGeometryBuilder>& GeometryGenerator, const Chaos::FConstImplicitObjectPtr& ImplicitObject, const FTransform& InWorldTransform, const FColor& InColor, FStringView DebugText, ESceneDepthPriorityGroup DepthPriority = SDPG_World, float Thickness = 2.0f);
 
 	static void DrawCanvas(FViewport& InViewport, FSceneView& View, FCanvas& Canvas);
 
diff --git a/Engine/Plugins/Experimental/ChaosVD/Source/ChaosVD/Public/Visualizers/ChaosVDParticleDataComponentVisualizer.h b/Engine/Plugins/Experimental/ChaosVD/Source/ChaosVD/Public/Visualizers/ChaosVDParticleDataComponentVisualizer.h
new file mode 100644
--- /dev/null
+++ b/Engine/Plugins/Experimental/ChaosVD/Source/ChaosVD/Public/Visualizers/ChaosVDParticleDataComponentVisualizer.h
@@ -0,0 +1,120 @@
+﻿// Copyright Epic Games, Inc. All Rights Reserved.
+
+#pragma once
+
+#include "ComponentVisualizer.h"
+#include "IChaosVDParticleVisualizationDataProvider.h"
+#include "DataWrappers/ChaosVDParticleDataWrapper.h"
+#include "HAL/Platform.h"
+#include "Templates/SharedPointer.h"
+
+class FChaosVDGeometryBuilder;
+struct FChaosParticleDataDebugDrawSettings;
+
+UENUM(meta = (Bitflags, UseEnumValuesAsMaskValuesInEditor = "true"))
+enum class EChaosVDParticleDataVisualizationFlags : uint32
+{
+	None				= 0 UMETA(Hidden),
+	Velocity			= 1 << 0,
+	AngularVelocity		= 1 << 1,
+	Acceleration		= 1 << 2,
+	AngularAcceleration = 1 << 3,
+	LinearImpulse		= 1 << 4,
+	AngularImpulse		= 1 << 5,
+	ClusterConnectivityEdge	= 1 << 6,
+	CenterOfMass	= 1 << 7,
+	DrawDataOnlyForSelectedParticle	= 1 << 8,
+	EnableDraw	= 1 << 9,
+};
+ENUM_CLASS_FLAGS(EChaosVDParticleDataVisualizationFlags);
+
+struct FChaosVDVisualizedParticleDataSelectionHandle
+{
+	int32 ParticleIndex = INDEX_NONE;
+	int32 SolverID = INDEX_NONE;
+};
+
+struct FChaosParticleDataDebugDrawSettings;
+
+struct FChaosVDParticleDataVisualizationContext : public FChaosVDVisualizationContext
+{
+	TWeakPtr<FChaosVDGeometryBuilder> GeometryGenerator = nullptr;
+	bool bIsSelectedData = false;
+	bool bShowDebugText = false;
+
+	const FChaosParticleDataDebugDrawSettings* DebugDrawSettings = nullptr;
+
+	bool IsVisualizationFlagEnabled(EChaosVDParticleDataVisualizationFlags Flag) const
+	{
+		const EChaosVDParticleDataVisualizationFlags FlagsAsParticleFlags = static_cast<EChaosVDParticleDataVisualizationFlags>(VisualizationFlags);
+		return EnumHasAnyFlags(FlagsAsParticleFlags, Flag);
+	}
+};
+
+namespace Chaos::VisualDebugger::ParticleDataUnitsStrings
+{
+	static FString Velocity = TEXT("cm/s");
+	static FString AngularVelocity = TEXT("rad/s");
+	static FString Acceleration = TEXT("cm/s2");
+	static FString AngularAcceleration = TEXT("rad/s2");
+	static FString LinearImpulse = TEXT("g.m/s");
+	static FString AngularImpulse = TEXT("g.m2/s");
+
+	static FString GetUnitByID(EChaosVDParticleDataVisualizationFlags DataID)
+	{
+		switch (DataID)
+		{
+		case EChaosVDParticleDataVisualizationFlags::Velocity:
+			return Velocity;
+		case EChaosVDParticleDataVisualizationFlags::AngularVelocity:
+			return AngularVelocity;
+		case EChaosVDParticleDataVisualizationFlags::Acceleration:
+			return Acceleration;
+		case EChaosVDParticleDataVisualizationFlags::AngularAcceleration:
+			return AngularAcceleration;
+		case EChaosVDParticleDataVisualizationFlags::LinearImpulse:
+			return LinearImpulse;
+		case EChaosVDParticleDataVisualizationFlags::AngularImpulse:
+			return AngularImpulse;
+
+		case EChaosVDParticleDataVisualizationFlags::None:
+		case EChaosVDParticleDataVisualizationFlags::ClusterConnectivityEdge:
+		case EChaosVDParticleDataVisualizationFlags::CenterOfMass:
+		case EChaosVDParticleDataVisualizationFlags::DrawDataOnlyForSelectedParticle:
+		default:
+			return TEXT("");
+		}
+	}
+}
+
+/** Custom Hit Proxy for debug drawn particle data */
+struct HChaosVDParticleDataProxy : public HComponentVisProxy
+{
+	DECLARE_HIT_PROXY()
+	
+	HChaosVDParticleDataProxy(const UActorComponent* Component, const FChaosVDVisualizedParticleDataSelectionHandle& InContactFinderData) : HComponentVisProxy(Component, HPP_UI), DataSelectionHandle(InContactFinderData)
+	{	
+	}
+
+	virtual EMouseCursor::Type GetMouseCursor() override
+	{
+		return EMouseCursor::Crosshairs;
+	}
+
+	FChaosVDVisualizedParticleDataSelectionHandle DataSelectionHandle;
+};
+
+/**
+ * Component visualizer in charge of generating debug draw visualizations for for particles
+ */
+class FChaosVDParticleDataComponentVisualizer : public FComponentVisualizer
+{
+public:
+	virtual void DrawVisualization(const UActorComponent* Component, const FSceneView* View, FPrimitiveDrawInterface* PDI) override;
+	virtual bool VisProxyHandleClick(FEditorViewportClient* InViewportClient, HComponentVisProxy* VisProxy, const FViewportClick& Click) override;
+
+protected:
+	FString GenerateDebugTextForVector(const FVector& InVector, const FString& VectorName, const FString& InVectorUnits);
+	void DrawParticleVector(FPrimitiveDrawInterface* PDI, const FVector& StartLocation, const FVector& InVector, EChaosVDParticleDataVisualizationFlags VectorID, const FChaosVDParticleDataVisualizationContext& InVisualizationContext, float LineThickness);
+	void DrawVisualizationForParticleData(const UActorComponent* Component, FPrimitiveDrawInterface* PDI, const FSceneView* View, const FChaosVDParticleDataVisualizationContext& InVisualizationContext, const FChaosVDParticleDataWrapper& InParticleDataViewer);
+};
diff --git a/Engine/Plugins/Experimental/ChaosVD/Source/ChaosVD/Public/Visualizers/ChaosVDParticleDataVisualizer.h b/Engine/Plugins/Experimental/ChaosVD/Source/ChaosVD/Public/Visualizers/ChaosVDParticleDataVisualizer.h
deleted file mode 100644
