[ChaosVD] Invalidating hit proxy cache when we update the CVD scene

This is required because otherwise the per instance hit proxy of the ISMC will be in the incorrect position.

#changelist validated
#virtualized

diff --git a/Engine/Plugins/Experimental/ChaosVD/Source/ChaosVD/Private/ChaosVDPlaybackViewportClient.cpp b/Engine/Plugins/Experimental/ChaosVD/Source/ChaosVD/Private/ChaosVDPlaybackViewportClient.cpp
--- a/Engine/Plugins/Experimental/ChaosVD/Source/ChaosVD/Private/ChaosVDPlaybackViewportClient.cpp
+++ b/Engine/Plugins/Experimental/ChaosVD/Source/ChaosVD/Private/ChaosVDPlaybackViewportClient.cpp
@@ -333,6 +333,17 @@ void FChaosVDPlaybackViewportClient::PerformSelectedTrackingForFrame(FChaosVDGam
 	}
 }
 
+bool FChaosVDPlaybackViewportClient::InputKey(const FInputKeyEventArgs& EventArgs)
+{
+	// Each time we requested a re-draw means we move something in the scene, so we need to re-create the cached hit proxy map.
+	if (bNeedsRedraw)
+	{
+		RequestInvalidateHitProxy(Viewport);
+	}
+	
+	return FEditorViewportClient::InputKey(EventArgs);
+}
+
 void FChaosVDPlaybackViewportClient::Draw(const FSceneView* View, FPrimitiveDrawInterface* PDI)
 {
 	const TSharedPtr<SChaosVDMainTab> MainTabToolkitHost = ModeTools.IsValid() ? StaticCastSharedPtr<SChaosVDMainTab>(ModeTools->GetToolkitHost()) : nullptr;
diff --git a/Engine/Plugins/Experimental/ChaosVD/Source/ChaosVD/Public/ChaosVDPlaybackViewportClient.h b/Engine/Plugins/Experimental/ChaosVD/Source/ChaosVD/Public/ChaosVDPlaybackViewportClient.h
--- a/Engine/Plugins/Experimental/ChaosVD/Source/ChaosVD/Public/ChaosVDPlaybackViewportClient.h
+++ b/Engine/Plugins/Experimental/ChaosVD/Source/ChaosVD/Public/ChaosVDPlaybackViewportClient.h
@@ -32,6 +32,8 @@ public:
 	
 	void PerformSelectedTrackingForFrame(FChaosVDGameFrameData* FrameData);
 
+	virtual bool InputKey(const FInputKeyEventArgs& EventArgs) override;
+
 private:
 	void HandleObjectFocused(UObject* FocusedObject);
 	void HandleActorMoving(AActor* MovedActor) const;
