[ChaosVD] Recording additional data to know when a solver frame being recorded is re-simulated or not.

Also, added a new badge in the solver frames timeline controls that reads that new flag and indicates when a frame being played was resimulated or not

diff --git a/Engine/Plugins/Experimental/ChaosVD/Source/ChaosVD/Private/ChaosVDPlaybackController.cpp b/Engine/Plugins/Experimental/ChaosVD/Source/ChaosVD/Private/ChaosVDPlaybackController.cpp
--- a/Engine/Plugins/Experimental/ChaosVD/Source/ChaosVD/Private/ChaosVDPlaybackController.cpp
+++ b/Engine/Plugins/Experimental/ChaosVD/Source/ChaosVD/Private/ChaosVDPlaybackController.cpp
@@ -235,6 +235,7 @@ void FChaosVDPlaybackController::GoToRecordedSolverStep_AssumesLocked(const int3
 
 				CurrentTrackInfo->CurrentFrame = FrameNumber;
 				CurrentTrackInfo->CurrentStep = Step;
+				CurrentTrackInfo->bIsReSimulated = SolverFrameData->bIsResimulated;
 			
 				EnqueueTrackInfoUpdate(*CurrentTrackInfo.Get(), InstigatorID);
 			}
diff --git a/Engine/Plugins/Experimental/ChaosVD/Source/ChaosVD/Private/Trace/ChaosVDTraceAnalyzer.cpp b/Engine/Plugins/Experimental/ChaosVD/Source/ChaosVD/Private/Trace/ChaosVDTraceAnalyzer.cpp
--- a/Engine/Plugins/Experimental/ChaosVD/Source/ChaosVD/Private/Trace/ChaosVDTraceAnalyzer.cpp
+++ b/Engine/Plugins/Experimental/ChaosVD/Source/ChaosVD/Private/Trace/ChaosVDTraceAnalyzer.cpp
@@ -82,6 +82,7 @@ bool FChaosVDTraceAnalyzer::OnEvent(uint16 RouteId, EStyle Style, const FOnEvent
 			NewFrameData.SolverID = EventData.GetValue<int32>("SolverID");
 			NewFrameData.FrameCycle = EventData.GetValue<uint64>("Cycle");
 			NewFrameData.bIsKeyFrame = EventData.GetValue<bool>("IsKeyFrame");
+			NewFrameData.bIsResimulated = EventData.GetValue<bool>("IsReSimulated");
 			NewFrameData.StartTime = Context.EventTime.AsSeconds(NewFrameData.FrameCycle);
 
 			FWideStringView DebugNameView;
diff --git a/Engine/Plugins/Experimental/ChaosVD/Source/ChaosVD/Private/Widgets/SChaosVDSolverPlaybackControls.cpp b/Engine/Plugins/Experimental/ChaosVD/Source/ChaosVD/Private/Widgets/SChaosVDSolverPlaybackControls.cpp
--- a/Engine/Plugins/Experimental/ChaosVD/Source/ChaosVD/Private/Widgets/SChaosVDSolverPlaybackControls.cpp
+++ b/Engine/Plugins/Experimental/ChaosVD/Source/ChaosVD/Private/Widgets/SChaosVDSolverPlaybackControls.cpp
@@ -7,6 +7,7 @@
 #include "Widgets/SBoxPanel.h"
 #include "Widgets/SChaosVDPlaybackViewport.h"
 #include "Widgets/SChaosVDTimelineWidget.h"
+#include "Widgets/Input/SButton.h"
 #include "Widgets/Text/STextBlock.h"
 
 #define LOCTEXT_NAMESPACE "ChaosVisualDebugger"
@@ -27,7 +28,7 @@ void SChaosVDSolverPlaybackControls::Construct(const FArguments& InArgs, int32 I
 			+SHorizontalBox::Slot()
 			.FillWidth(0.7f)
 			[
-			SNew(SVerticalBox)
+				SNew(SVerticalBox)
 				+SVerticalBox::Slot()
 				.AutoHeight()
 				.Padding(0.0f, 0.0f, 0.0f, 2.0f)
@@ -38,10 +39,32 @@ void SChaosVDSolverPlaybackControls::Construct(const FArguments& InArgs, int32 I
 				]
 				+SVerticalBox::Slot()
 				[
-					SAssignNew(FramesTimelineWidget, SChaosVDTimelineWidget)
+					SNew(SHorizontalBox)
+					+SHorizontalBox::Slot()
+					.FillWidth(0.95f)
+					[
+						SAssignNew(FramesTimelineWidget, SChaosVDTimelineWidget)
 						.ButtonVisibilityFlags(static_cast<uint16>(EChaosVDTimelineElementIDFlags::AllPlayback))
 						.OnFrameChanged_Raw(this, &SChaosVDSolverPlaybackControls::OnFrameSelectionUpdated)
 						.MaxFrames(0)
+					]
+					+SHorizontalBox::Slot()
+					.FillWidth(0.05f)
+					[
+						SNew(SBorder)
+						.BorderImage_Raw(this, &SChaosVDSolverPlaybackControls::GetFrameTypeBadgeBrush)
+						.Padding(1.5f)
+						.Content()
+						[
+							SNew(SBox)
+							.VAlign(VAlign_Center)
+							[
+								SNew(STextBlock)
+								.Justification(ETextJustify::Center)
+								.Text_Lambda([this]()->FText{ return bIsReSimFrame ? LOCTEXT("PlaybackViewportWidgetPhysicsFramesResimLabel", "ReSim" ) : LOCTEXT("PlaybackViewportWidgetPhysicsFramesNormalLabel", "Normal" );})
+							]
+						]
+					]
 				]
 			]
 			+SHorizontalBox::Slot()
@@ -54,12 +77,12 @@ void SChaosVDSolverPlaybackControls::Construct(const FArguments& InArgs, int32 I
 				[
 					SNew(STextBlock)
 					.Justification(ETextJustify::Center)
-					.Text_Lambda([this]()->FText{ return FText::Format(LOCTEXT("PlaybackViewportWidgetStepsLabel","Step {0}"), FText::AsCultureInvariant(CurrentStepName));})
+					.Text_Lambda([this]()->FText{ return FText::Format(LOCTEXT("PlaybackViewportWidgetStepsLabel","Solver Stage: {0}"), FText::AsCultureInvariant(CurrentStepName));})
 				]
 				+SVerticalBox::Slot()
 				[
 					SAssignNew(StepsTimelineWidget, SChaosVDTimelineWidget)
-					.ButtonVisibilityFlags(static_cast<uint16>(EChaosVDTimelineElementIDFlags::AllManualStepping | EChaosVDTimelineElementIDFlags::Lock))
+					.ButtonVisibilityFlags(static_cast<uint16>(EChaosVDTimelineElementIDFlags::AllManualStepping))
 					.OnFrameLockStateChanged_Raw(this, &SChaosVDSolverPlaybackControls::HandleLockStateChanged)
 					.OnFrameChanged_Raw(this, &SChaosVDSolverPlaybackControls::OnStepSelectionUpdated)
 					.MaxFrames(0)
@@ -163,6 +186,8 @@ void SChaosVDSolverPlaybackControls::HandleControllerTrackFrameUpdated(TWeakPtr<
 				UpdateStepsWidgetForFrame(*CurrentPlaybackControllerPtr.Get(), SolverTrackInfo->CurrentFrame, SolverTrackInfo->CurrentStep);
 			}
 
+			bIsReSimFrame = SolverTrackInfo->bIsReSimulated;
+
 			FramesTimelineWidget->SetTargetFrameTime(CurrentPlaybackControllerPtr->GetFrameTimeForTrack(EChaosVDTrackType::Solver, SolverID, *SolverTrackInfo));
 		}
 	}
@@ -183,6 +208,12 @@ void SChaosVDSolverPlaybackControls::HandleLockStateChanged(bool NewIsLocked)
 	}
 }
 
+const FSlateBrush* SChaosVDSolverPlaybackControls::GetFrameTypeBadgeBrush() const
+{
+	const FButtonStyle& ButtonStyle = FAppStyle::Get().GetWidgetStyle<FButtonStyle>("Menu.Button");
+	return bIsReSimFrame ? &ButtonStyle.Pressed : FCoreStyle::Get().GetBrush("Border");
+}
+
 void SChaosVDSolverPlaybackControls::OnFrameSelectionUpdated(int32 NewFrameIndex)
 {
 	if (const TSharedPtr<FChaosVDPlaybackController> PlaybackControllerPtr = PlaybackController.Pin())
diff --git a/Engine/Plugins/Experimental/ChaosVD/Source/ChaosVD/Public/ChaosVDPlaybackController.h b/Engine/Plugins/Experimental/ChaosVD/Source/ChaosVD/Public/ChaosVDPlaybackController.h
--- a/Engine/Plugins/Experimental/ChaosVD/Source/ChaosVD/Public/ChaosVDPlaybackController.h
+++ b/Engine/Plugins/Experimental/ChaosVD/Source/ChaosVD/Public/ChaosVDPlaybackController.h
@@ -39,6 +39,7 @@ struct FChaosVDTrackInfo
 	int32 LockedOnStep = INDEX_NONE;
 	int32 MaxFrames = INDEX_NONE;
 	FString TrackName;
+	bool bIsReSimulated = false;
 };
 
 struct FChaosVDQueuedTrackInfoUpdate
diff --git a/Engine/Plugins/Experimental/ChaosVD/Source/ChaosVD/Public/Widgets/SChaosVDSolverPlaybackControls.h b/Engine/Plugins/Experimental/ChaosVD/Source/ChaosVD/Public/Widgets/SChaosVDSolverPlaybackControls.h
--- a/Engine/Plugins/Experimental/ChaosVD/Source/ChaosVD/Public/Widgets/SChaosVDSolverPlaybackControls.h
+++ b/Engine/Plugins/Experimental/ChaosVD/Source/ChaosVD/Public/Widgets/SChaosVDSolverPlaybackControls.h
@@ -45,7 +45,10 @@ private:
 
 	void HandleLockStateChanged(bool NewIsLocked);
 
+	const FSlateBrush* GetFrameTypeBadgeBrush() const;
+
 	int32 SolverID = INDEX_NONE;
+	bool bIsReSimFrame = false;
 	FString CurrentStepName;
 	TSharedPtr<SChaosVDTimelineWidget> FramesTimelineWidget;
 	TSharedPtr<SChaosVDTimelineWidget> StepsTimelineWidget;
diff --git a/Engine/Source/Runtime/Experimental/Chaos/Private/ChaosVisualDebugger/ChaosVisualDebuggerTrace.cpp b/Engine/Source/Runtime/Experimental/Chaos/Private/ChaosVisualDebugger/ChaosVisualDebuggerTrace.cpp
--- a/Engine/Source/Runtime/Experimental/Chaos/Private/ChaosVisualDebugger/ChaosVisualDebuggerTrace.cpp
+++ b/Engine/Source/Runtime/Experimental/Chaos/Private/ChaosVisualDebugger/ChaosVisualDebuggerTrace.cpp
@@ -384,6 +384,8 @@ void FChaosVisualDebuggerTrace::TraceSolverFrameStart(const FChaosVDContext& Con
 
 	FChaosVDThreadContext::Get().PushContext(ContextData);
 
+	bool bIsReSimulatedFrame = EnumHasAnyFlags(static_cast<EChaosVDContextAttributes>(ContextData.Attributes), EChaosVDContextAttributes::Resimulated);
+
 	// Check if we need to do a full capture for this solver, and setup accordingly
 	bool bOutIsFullCaptureRequested;
 	SetupForFullCaptureIfNeeded(ContextData.Id, bOutIsFullCaptureRequested);
@@ -392,7 +394,8 @@ void FChaosVisualDebuggerTrace::TraceSolverFrameStart(const FChaosVDContext& Con
 		<< ChaosVDSolverFrameStart.SolverID(ContextData.Id)
 		<< ChaosVDSolverFrameStart.Cycle(FPlatformTime::Cycles64())
 		<< ChaosVDSolverFrameStart.DebugName(*InDebugName, InDebugName.Len())
-		<< ChaosVDSolverFrameStart.IsKeyFrame(bOutIsFullCaptureRequested);
+		<< ChaosVDSolverFrameStart.IsKeyFrame(bOutIsFullCaptureRequested)
+		<< ChaosVDSolverFrameStart.IsReSimulated(bIsReSimulatedFrame);
 }
 
 void FChaosVisualDebuggerTrace::TraceSolverFrameEnd(const FChaosVDContext& ContextData)
diff --git a/Engine/Source/Runtime/Experimental/Chaos/Private/PBDRigidsSolver.cpp b/Engine/Source/Runtime/Experimental/Chaos/Private/PBDRigidsSolver.cpp
--- a/Engine/Source/Runtime/Experimental/Chaos/Private/PBDRigidsSolver.cpp
+++ b/Engine/Source/Runtime/Experimental/Chaos/Private/PBDRigidsSolver.cpp
@@ -1940,7 +1940,7 @@ namespace Chaos
 					UE_LOG(LogTemp, Warning, TEXT("COMMON | PT | ConditionalApplyRewind_Internal | PERFORMING RESIMULATION | Resim From Frame = %d | Num Steps = %d | To Current Frame: %d"), ResimStep, NumResimSteps, CurrentFrame);
 #endif
 
-					GetEvolution()->SetResim(true);
+					SetIsResimming(true);
 					CurrentFrame = ResimStep;
 
 					TArray<FPushPhysicsData*> RecordedPushData = MarshallingManager.StealHistory_Internal(NumResimSteps);
@@ -2004,7 +2004,7 @@ namespace Chaos
 					}
 					GetEvolution()->GetIslandManager().ResetParticleResimFrame();
 
-					GetEvolution()->SetResim(false);
+					SetIsResimming(false);
 					GetEvolution()->SetReset(false);
 
 					ResimTimer.Stop();
@@ -2028,6 +2028,25 @@ namespace Chaos
 		}
 	}
 
+	void FPBDRigidsSolver::SetIsResimming(bool bIsResimming)
+	{
+		GetEvolution()->SetResim(bIsResimming);
+
+#if WITH_CHAOS_VISUAL_DEBUGGER
+		EChaosVDContextAttributes Attributes = static_cast<EChaosVDContextAttributes>(GetChaosVDContextData().Attributes);
+		if (bIsResimming)
+		{
+			EnumAddFlags(Attributes,  EChaosVDContextAttributes::Resimulated);
+		}
+		else
+		{
+			EnumRemoveFlags(Attributes,  EChaosVDContextAttributes::Resimulated);
+		}
+
+		GetChaosVDContextData().Attributes = static_cast<int32>(Attributes);
+#endif	
+	}
+
 	void FPBDRigidsSolver::CompleteSceneSimulation()
 	{
 		LLM_SCOPE(ELLMTag::Chaos);
diff --git a/Engine/Source/Runtime/Experimental/Chaos/Public/ChaosVisualDebugger/ChaosVDContextProvider.h b/Engine/Source/Runtime/Experimental/Chaos/Public/ChaosVisualDebugger/ChaosVDContextProvider.h
--- a/Engine/Source/Runtime/Experimental/Chaos/Public/ChaosVisualDebugger/ChaosVDContextProvider.h
+++ b/Engine/Source/Runtime/Experimental/Chaos/Public/ChaosVisualDebugger/ChaosVDContextProvider.h
@@ -20,19 +20,27 @@ enum class EChaosVDContextType : int32
 	SubTraceQuery,
 };
 
+enum class EChaosVDContextAttributes : int32
+{
+	None = 0,
+	Resimulated = 1 << 0
+};
+ENUM_CLASS_FLAGS(EChaosVDContextAttributes)
+
 /** Chaos Visual Debugger data used to context for logging or debugging purposes */
 struct FChaosVDContext
 {
 	int32 OwnerID = INDEX_NONE;
 	int32 Id = INDEX_NONE;
 	int32 Type = INDEX_NONE;
+	int32 Attributes = 0;
 
 	FORCEINLINE void SetDataChannel(const TSharedRef<Chaos::VisualDebugger::FChaosVDOptionalDataChannel>& NewDataChannel)
 	{
 		CurrentDataChannel = NewDataChannel;
 	}
 
-	FORCEINLINE bool IsDataChannelEnabled () const { return CurrentDataChannel->IsChannelEnabled(); }
+	FORCEINLINE bool IsDataChannelEnabled() const { return CurrentDataChannel->IsChannelEnabled(); }
 
 private:
 	TSharedRef<Chaos::VisualDebugger::FChaosVDOptionalDataChannel> CurrentDataChannel = CVDDC_Default;
diff --git a/Engine/Source/Runtime/Experimental/Chaos/Public/ChaosVisualDebugger/ChaosVisualDebuggerTrace.h b/Engine/Source/Runtime/Experimental/Chaos/Public/ChaosVisualDebugger/ChaosVisualDebuggerTrace.h
--- a/Engine/Source/Runtime/Experimental/Chaos/Public/ChaosVisualDebugger/ChaosVisualDebuggerTrace.h
+++ b/Engine/Source/Runtime/Experimental/Chaos/Public/ChaosVisualDebugger/ChaosVisualDebuggerTrace.h
@@ -33,6 +33,7 @@ UE_TRACE_EVENT_BEGIN_EXTERN(ChaosVDLogger, ChaosVDSolverFrameStart)
 	UE_TRACE_EVENT_FIELD(uint64, Cycle)
 	UE_TRACE_EVENT_FIELD(UE::Trace::WideString, DebugName)
 	UE_TRACE_EVENT_FIELD(bool, IsKeyFrame)
+	UE_TRACE_EVENT_FIELD(bool, IsReSimulated)
 UE_TRACE_EVENT_END()
 
 UE_TRACE_EVENT_BEGIN_EXTERN(ChaosVDLogger, ChaosVDSolverFrameEnd)
diff --git a/Engine/Source/Runtime/Experimental/Chaos/Public/PBDRigidsSolver.h b/Engine/Source/Runtime/Experimental/Chaos/Public/PBDRigidsSolver.h
--- a/Engine/Source/Runtime/Experimental/Chaos/Public/PBDRigidsSolver.h
+++ b/Engine/Source/Runtime/Experimental/Chaos/Public/PBDRigidsSolver.h
@@ -425,6 +425,9 @@ namespace Chaos
 
 		/** Check if we are resimming or not */
 		virtual bool IsResimming() const {return GetEvolution()->IsResimming();}
+
+		/** Sets if we are resimming or not */
+		void SetIsResimming(bool bIsResimming);
 	};
 
 	template<>
diff --git a/Engine/Source/Runtime/Experimental/ChaosVDData/Public/ChaosVDRecording.h b/Engine/Source/Runtime/Experimental/ChaosVDData/Public/ChaosVDRecording.h
--- a/Engine/Source/Runtime/Experimental/ChaosVDData/Public/ChaosVDRecording.h
+++ b/Engine/Source/Runtime/Experimental/ChaosVDData/Public/ChaosVDRecording.h
@@ -52,6 +52,7 @@ struct CHAOSVDDATA_API FChaosVDSolverFrameData
 	uint64 FrameCycle = 0;
 	Chaos::FRigidTransform3 SimulationTransform;
 	bool bIsKeyFrame = false;
+	bool bIsResimulated = false;
 	FChaosVDStepsContainer SolverSteps;
 	TSet<int32> ParticlesDestroyedIDs;
 	double StartTime = -1.0;
