[ChaosVD] Fix for change of settings not being properly applied in some situations

diff --git a/Engine/Plugins/Experimental/ChaosVD/Source/ChaosVD/Private/ChaosVDEditorSettings.cpp b/Engine/Plugins/Experimental/ChaosVD/Source/ChaosVD/Private/ChaosVDEditorSettings.cpp
--- a/Engine/Plugins/Experimental/ChaosVD/Source/ChaosVD/Private/ChaosVDEditorSettings.cpp
+++ b/Engine/Plugins/Experimental/ChaosVD/Source/ChaosVD/Private/ChaosVDEditorSettings.cpp
@@ -135,23 +135,23 @@ void UChaosVDEditorSettings::PostEditChangeProperty(FPropertyChangedEvent& Prope
 {
 	Super::PostEditChangeProperty(PropertyChangedEvent);
 
-	const FName PropertyName = PropertyChangedEvent.GetPropertyName();
-	if (PropertyName == GET_MEMBER_NAME_CHECKED(UChaosVDEditorSettings, GeometryVisibilityFlags))
+	const FName MemberPropertyName = PropertyChangedEvent.GetMemberPropertyName();
+	if (MemberPropertyName == GET_MEMBER_NAME_CHECKED(UChaosVDEditorSettings, GeometryVisibilityFlags))
 	{
 		VisibilitySettingsChangedDelegate.Broadcast(this);
 	}
-	else if (PropertyName == GET_MEMBER_NAME_CHECKED(UChaosVDEditorSettings, ParticleColorMode)
-			|| PropertyName == GET_MEMBER_NAME_CHECKED(UChaosVDEditorSettings, ColorsByParticleState)
-			|| PropertyName == GET_MEMBER_NAME_CHECKED(UChaosVDEditorSettings, ColorsByShapeType))
+	else if (MemberPropertyName == GET_MEMBER_NAME_CHECKED(UChaosVDEditorSettings, ParticleColorMode)
+			|| MemberPropertyName == GET_MEMBER_NAME_CHECKED(UChaosVDEditorSettings, ColorsByParticleState)
+			|| MemberPropertyName == GET_MEMBER_NAME_CHECKED(UChaosVDEditorSettings, ColorsByShapeType))
 	{
 		ColorsSettingsChangedDelegate.Broadcast(this);
 	}
-	else if (PropertyName == GET_MEMBER_NAME_CHECKED(UChaosVDEditorSettings, FarClippingOverride))
+	else if (MemberPropertyName == GET_MEMBER_NAME_CHECKED(UChaosVDEditorSettings, FarClippingOverride))
 	{
 		FarClippingOverrideChangedDelegate.Broadcast(this);
 	}
-	else if (PropertyName == GET_MEMBER_NAME_CHECKED(UChaosVDEditorSettings, bPlaybackAtRecordedFrameRate) ||
-			 PropertyName == GET_MEMBER_NAME_CHECKED(UChaosVDEditorSettings, TargetFrameRateOverride))
+	else if (MemberPropertyName == GET_MEMBER_NAME_CHECKED(UChaosVDEditorSettings, bPlaybackAtRecordedFrameRate) ||
+			 MemberPropertyName == GET_MEMBER_NAME_CHECKED(UChaosVDEditorSettings, TargetFrameRateOverride))
 	{
 		PlaybackSettingsChangedDelegate.Broadcast(this);
 	}
@@ -159,3 +159,16 @@ void UChaosVDEditorSettings::PostEditChangeProperty(FPropertyChangedEvent& Prope
 	// TODO: If we keep this object as the main setting object,
 	// we should have a single event for what changed and an enum flags that the listener could use to decide if cares about the change
 }
+
+void UChaosVDEditorSettings::PostEditUndo()
+{
+	UObject::PostEditUndo();
+
+	// This is not ideal, but we don't get what property was changed in the post edit undo callback.
+	// A proper fix to avoid calling all the settings change delegates will be done when we split this settings object into several objects
+	// and expose the options as proper UE menus instead of a details panel. Jira for tracking UE-206957 
+	VisibilitySettingsChangedDelegate.Broadcast(this);
+	ColorsSettingsChangedDelegate.Broadcast(this);
+	PlaybackSettingsChangedDelegate.Broadcast(this);
+	FarClippingOverrideChangedDelegate.Broadcast(this);
+}
diff --git a/Engine/Plugins/Experimental/ChaosVD/Source/ChaosVD/Public/ChaosVDEditorSettings.h b/Engine/Plugins/Experimental/ChaosVD/Source/ChaosVD/Public/ChaosVDEditorSettings.h
--- a/Engine/Plugins/Experimental/ChaosVD/Source/ChaosVD/Public/ChaosVDEditorSettings.h
+++ b/Engine/Plugins/Experimental/ChaosVD/Source/ChaosVD/Public/ChaosVDEditorSettings.h
@@ -364,6 +364,8 @@ public:
 
 	TSharedPtr<FName> SelectedTrackedTransformName;
 	TSharedPtr<FName> SelectedTrackedLocationName;
+	
+	virtual void PostEditUndo() override;
 
 protected:
 	FChaosVDSettingChaged VisibilitySettingsChangedDelegate;
