[ChaosVD] Marking CVD Materials as transient so they are not saved.

Also, updated the code to reduce the amount of Material Instances being created

diff --git a/Engine/Plugins/Experimental/ChaosVD/Source/ChaosVD/Private/ChaosVDGeometryDataComponent.cpp b/Engine/Plugins/Experimental/ChaosVD/Source/ChaosVD/Private/ChaosVDGeometryDataComponent.cpp
--- a/Engine/Plugins/Experimental/ChaosVD/Source/ChaosVD/Private/ChaosVDGeometryDataComponent.cpp
+++ b/Engine/Plugins/Experimental/ChaosVD/Source/ChaosVD/Private/ChaosVDGeometryDataComponent.cpp
@@ -107,11 +107,9 @@ void FChaosVDMeshDataInstanceHandle::UpdateMeshComponentForCollisionData(const F
 			{
 				if (IChaosVDGeometryComponent* CVDNewGeometryComponent = Cast<IChaosVDGeometryComponent>(GetMeshComponent()))
 				{
-					if (CurrentGeometryColor != FLinearColor(ForceInitToZero))
-					{
-						CVDNewGeometryComponent->UpdateInstanceColor(AsShared(), CurrentGeometryColor);
-					}
-						
+					// Reset the color so it is updated in the next Update color calls (which always happens after updating the shape instance data)
+					CurrentGeometryColor = FLinearColor(ForceInitToZero);
+		
 					CVDNewGeometryComponent->UpdateInstanceVisibility(AsShared(), bIsVisible);
 					CVDNewGeometryComponent->SetIsSelected(AsShared(), bIsSelected);
 				}
@@ -152,57 +150,6 @@ void FChaosVDMeshDataInstanceHandle::SetVisibility(bool bInIsVisible)
 	bIsVisible = bInIsVisible;
 }
 
-UMaterialInstanceDynamic* FChaosVDMeshDataInstanceHandle::GetCachedMaterialInstance(EChaosVDMaterialType Type)
-{
-	const UChaosVDEditorSettings* EditorSettings = GetDefault<UChaosVDEditorSettings>();
-	if (!EditorSettings)
-	{
-		return nullptr;
-	}
-
-	if (const TStrongObjectPtr<UMaterialInstanceDynamic>* MaterialInstance = MaterialInstancesByID.Find(Type))
-	{
-		return MaterialInstance->Get();
-	}
-	else
-	{
-		UMaterialInterface* MaterialToCreate = nullptr;
-		switch(Type)
-		{
-		case EChaosVDMaterialType::QueryOnlyMaterial:
-			{
-				MaterialToCreate = EditorSettings->QueryOnlyMeshesMaterial.Get();
-				break;
-			}
-		case EChaosVDMaterialType::SimOnlyMaterial:
-			{
-				MaterialToCreate = EditorSettings->SimOnlyMeshesMaterial.Get();
-				break;
-			}
-		case EChaosVDMaterialType::Instanced:
-			{
-				MaterialToCreate = EditorSettings->InstancedMeshesMaterial.Get();
-				break;
-			}
-		case EChaosVDMaterialType::InstancedQueryOnly:
-			{
-				MaterialToCreate = EditorSettings->InstancedMeshesQueryOnlyMaterial.Get();
-				break;
-			}
-		}
-		
-		if (MaterialToCreate)
-		{
-			UMaterialInstanceDynamic* DynamicMaterial = UMaterialInstanceDynamic::Create(MaterialToCreate, nullptr);
-			MaterialInstancesByID.Add(Type, TStrongObjectPtr<UMaterialInstanceDynamic>(DynamicMaterial));
-
-			return DynamicMaterial;
-		}
-	}
-
-	return nullptr;
-}
-
 void FChaosVDMeshDataInstanceHandle::HandleInstanceIndexUpdated(TArrayView<const FInstancedStaticMeshDelegates::FInstanceIndexUpdateData> InIndexUpdates)
 {
 	// When an index changes, we receive an array with all the indexes that were modified. We need to only act upon the update of for the index we are tracking in this handle
@@ -398,3 +345,44 @@ FLinearColor FChaosVDGeometryComponentUtils::GetGeometryParticleColor(const TSha
 
 	return ColorToApply;
 }
+
+UMaterialInterface* FChaosVDGeometryComponentUtils::GetBaseMaterialForType(EChaosVDMaterialType Type)
+{
+	const UChaosVDEditorSettings* EditorSettings = GetDefault<UChaosVDEditorSettings>();
+	if (!EditorSettings)
+	{
+		return nullptr;
+	}
+
+	switch(Type)
+	{
+		case EChaosVDMaterialType::QueryOnlyMaterial:
+				return EditorSettings->QueryOnlyMeshesMaterial.Get();
+		case EChaosVDMaterialType::SimOnlyMaterial:
+				return EditorSettings->SimOnlyMeshesMaterial.Get();
+		case EChaosVDMaterialType::Instanced:
+				return EditorSettings->InstancedMeshesMaterial.Get();
+		case EChaosVDMaterialType::InstancedQueryOnly:
+				return EditorSettings->InstancedMeshesQueryOnlyMaterial.Get();
+		default:
+			return nullptr;
+	}	
+}
+
+UMaterialInstanceDynamic* FChaosVDGeometryComponentUtils::CreateMaterialInstance(UMaterialInterface* BaseMaterial)
+{
+	if (!BaseMaterial)
+	{
+		return nullptr;
+	}
+
+	UMaterialInstanceDynamic* DynamicMaterial = UMaterialInstanceDynamic::Create(BaseMaterial, nullptr);
+	DynamicMaterial->SetFlags(RF_Transient);
+	
+	return DynamicMaterial;
+}
+
+UMaterialInstanceDynamic* FChaosVDGeometryComponentUtils::CreateMaterialInstance(EChaosVDMaterialType Type)
+{
+	return CreateMaterialInstance(GetBaseMaterialForType(Type));
+}
diff --git a/Engine/Plugins/Experimental/ChaosVD/Source/ChaosVD/Private/Components/ChaosVDInstancedStaticMeshComponent.cpp b/Engine/Plugins/Experimental/ChaosVD/Source/ChaosVD/Private/Components/ChaosVDInstancedStaticMeshComponent.cpp
--- a/Engine/Plugins/Experimental/ChaosVD/Source/ChaosVD/Private/Components/ChaosVDInstancedStaticMeshComponent.cpp
+++ b/Engine/Plugins/Experimental/ChaosVD/Source/ChaosVD/Private/Components/ChaosVDInstancedStaticMeshComponent.cpp
@@ -60,14 +60,25 @@ void UChaosVDInstancedStaticMeshComponent::UpdateInstanceColor(const TSharedPtr<
 	if (CurrentMaterial == nullptr)
 	{
 		CurrentMaterial = GetMaterial(0);
+	}
 
-		bool bIsSolidColor = FMath::IsNearlyEqual(NewColor.A, 1.0f);
-		UMaterialInterface* DesiredMaterial = InInstanceHandle->GetCachedMaterialInstance( bIsSolidColor ? EChaosVDMaterialType::Instanced : EChaosVDMaterialType::InstancedQueryOnly);
-
-		if (CurrentMaterial != DesiredMaterial)
+	// Check that this mesh component supports the intended visualization
+	// We can't change the material of Instanced mesh components because we might have other instances that are not intended to be translucent (or the other way around).
+	// The Mesh handle instance system should have detected we need to migrate the instance to another component before ever reaching this point
+	const bool bIsSolidColor = FMath::IsNearlyEqual(NewColor.A, 1.0f);
+	const bool bMeshComponentSupportsTranslucentInstances = EnumHasAnyFlags(static_cast<EChaosVDMeshAttributesFlags>(MeshComponentAttributeFlags), EChaosVDMeshAttributesFlags::TranslucentGeometry);
+	if (bIsSolidColor)
+	{
+		if (!ensure(!bMeshComponentSupportsTranslucentInstances))
+		{
+			UE_LOG(LogChaosVDEditor, Warning, TEXT("[%s] Desired Color is not supported in this mesh component [%s]..."), ANSI_TO_TCHAR(__FUNCTION__), *NewColor.ToString());
+		}
+	}
+	else
+	{
+		if (!ensure(bMeshComponentSupportsTranslucentInstances))
 		{
-			SetMaterial(0, DesiredMaterial);
-			CurrentMaterial = DesiredMaterial;
+			UE_LOG(LogChaosVDEditor, Warning, TEXT("[%s] Desired Color is not supported in this mesh component [%s]..."), ANSI_TO_TCHAR(__FUNCTION__), *NewColor.ToString());
 		}
 	}
 	
diff --git a/Engine/Plugins/Experimental/ChaosVD/Source/ChaosVD/Private/Components/ChaosVDStaticMeshComponent.cpp b/Engine/Plugins/Experimental/ChaosVD/Source/ChaosVD/Private/Components/ChaosVDStaticMeshComponent.cpp
--- a/Engine/Plugins/Experimental/ChaosVD/Source/ChaosVD/Private/Components/ChaosVDStaticMeshComponent.cpp
+++ b/Engine/Plugins/Experimental/ChaosVD/Source/ChaosVD/Private/Components/ChaosVDStaticMeshComponent.cpp
@@ -55,9 +55,9 @@ void UChaosVDStaticMeshComponent::UpdateInstanceColor(const TSharedPtr<FChaosVDM
 		return;
 	}
 
-	bool bIsSolidColor = FMath::IsNearlyEqual(NewColor.A, 1.0f);
+	const bool bIsSolidColor = FMath::IsNearlyEqual(NewColor.A, 1.0f);
 
-	if (UMaterialInstanceDynamic* MaterialToApply = bIsSolidColor ? InInstanceHandle->GetCachedMaterialInstance(EChaosVDMaterialType::SimOnlyMaterial) : InInstanceHandle->GetCachedMaterialInstance(EChaosVDMaterialType::QueryOnlyMaterial))
+	if (UMaterialInstanceDynamic* MaterialToApply = GetCachedMaterialInstance(bIsSolidColor ? EChaosVDMaterialType::SimOnlyMaterial : EChaosVDMaterialType::QueryOnlyMaterial))
 	{
 		MaterialToApply->SetVectorParameterValue(TEXT("BaseColor"), NewColor);
 		InInstanceHandle->GetMeshComponent()->SetMaterial(0, MaterialToApply);
@@ -169,3 +169,23 @@ TArrayView<TSharedPtr<FChaosVDMeshDataInstanceHandle>> UChaosVDStaticMeshCompone
 {
 	return TArrayView<TSharedPtr<FChaosVDMeshDataInstanceHandle>>(&CurrentMeshDataHandle, 1);
 }
+
+UMaterialInstanceDynamic* UChaosVDStaticMeshComponent::GetCachedMaterialInstance(EChaosVDMaterialType Type)
+{
+	if (const TObjectPtr<UMaterialInstanceDynamic>* MaterialInstance = CachedMaterialInstancesByID.Find(Type))
+	{
+		return MaterialInstance->Get();
+	}
+	else
+	{
+		if (UMaterialInterface* MaterialToCreate = FChaosVDGeometryComponentUtils::GetBaseMaterialForType(Type))
+		{
+			UMaterialInstanceDynamic* DynamicMaterial = UMaterialInstanceDynamic::Create(MaterialToCreate, nullptr);
+			CachedMaterialInstancesByID.Add(Type, TObjectPtr<UMaterialInstanceDynamic>(DynamicMaterial));
+
+			return DynamicMaterial;
+		}
+	}
+
+	return nullptr;
+}
diff --git a/Engine/Plugins/Experimental/ChaosVD/Source/ChaosVD/Public/ChaosVDEditorSettings.h b/Engine/Plugins/Experimental/ChaosVD/Source/ChaosVD/Public/ChaosVDEditorSettings.h
--- a/Engine/Plugins/Experimental/ChaosVD/Source/ChaosVD/Public/ChaosVDEditorSettings.h
+++ b/Engine/Plugins/Experimental/ChaosVD/Source/ChaosVD/Public/ChaosVDEditorSettings.h
@@ -337,16 +337,16 @@ public:
 	UPROPERTY(EditAnywhere, Category = "Geometry Visibility", meta = (Bitmask, BitmaskEnum = "/Script/ChaosVD.EChaosVDGeometryVisibilityFlags"))
 	uint8 GeometryVisibilityFlags = static_cast<uint8>(EChaosVDGeometryVisibilityFlags::Simulated | EChaosVDGeometryVisibilityFlags::Simple |  EChaosVDGeometryVisibilityFlags::ShowHeightfields);
 
-	UPROPERTY(Config)
+	UPROPERTY(Config, Transient)
 	TSoftObjectPtr<UMaterial> QueryOnlyMeshesMaterial;
 
-	UPROPERTY(Config)
+	UPROPERTY(Config, Transient)
 	TSoftObjectPtr<UMaterial> SimOnlyMeshesMaterial;
 
-	UPROPERTY(Config)
+	UPROPERTY(Config, Transient)
 	TSoftObjectPtr<UMaterial> InstancedMeshesMaterial;
 
-	UPROPERTY(Config)
+	UPROPERTY(Config, Transient)
 	TSoftObjectPtr<UMaterial> InstancedMeshesQueryOnlyMaterial;
 
 	UPROPERTY(Config)
@@ -368,6 +368,7 @@ public:
 	virtual void PostEditUndo() override;
 
 protected:
+	
 	FChaosVDSettingChaged VisibilitySettingsChangedDelegate;
 	FChaosVDSettingChaged ColorsSettingsChangedDelegate;
 	FChaosVDSettingChaged FarClippingOverrideChangedDelegate;
diff --git a/Engine/Plugins/Experimental/ChaosVD/Source/ChaosVD/Public/ChaosVDGeometryBuilder.h b/Engine/Plugins/Experimental/ChaosVD/Source/ChaosVD/Public/ChaosVDGeometryBuilder.h
--- a/Engine/Plugins/Experimental/ChaosVD/Source/ChaosVD/Public/ChaosVDGeometryBuilder.h
+++ b/Engine/Plugins/Experimental/ChaosVD/Source/ChaosVD/Public/ChaosVDGeometryBuilder.h
@@ -19,6 +19,7 @@
 
 #include "UDynamicMesh.h"
 #include "Components/ChaosVDInstancedStaticMeshComponent.h"
+#include "Materials/MaterialInstanceDynamic.h"
 #include "UObject/WeakObjectPtr.h"
 #include "Tasks/Task.h"
 #include "Templates/SharedPointer.h"
@@ -157,6 +158,9 @@ private:
 	template <class ComponentType>
     bool InitializeMeshComponent(AActor* Owner, ComponentType* MeshComponent);
 
+	template <class ComponentType>
+	void SetMeshComponentMaterial(EChaosVDMeshAttributesFlags MeshComponentAttributeFlags, ComponentType* MeshComponent);
+
 public:
 	/**
 	 * Finds or creates a Mesh component for the geometry data handle provided, and add a new instance of that geometry to it
@@ -485,6 +489,19 @@ bool FChaosVDGeometryBuilder::InitializeMeshComponent(AActor* Owner, ComponentTy
 	return true;
 }
 
+template <typename ComponentType>
+void FChaosVDGeometryBuilder::SetMeshComponentMaterial(EChaosVDMeshAttributesFlags MeshComponentAttributeFlags, ComponentType* MeshComponent)
+{
+	UMaterialInstanceDynamic* Material = nullptr;
+
+	Material = FChaosVDGeometryComponentUtils::CreateMaterialInstance(FChaosVDGeometryComponentUtils::GetMaterialTypeForComponent<ComponentType>(MeshComponentAttributeFlags));
+
+	ensure(Material);
+
+	MeshComponent->SetMaterial(0, Material);
+}
+
+
 template <typename ComponentType>
 ComponentType* FChaosVDGeometryBuilder::GetAvailableInstancedStaticMeshComponent(const TSharedPtr<FChaosVDExtractedGeometryDataHandle>& InExtractedGeometryDataHandle, AActor* MeshComponentsContainerActor, EChaosVDMeshAttributesFlags MeshComponentAttributeFlags, bool& bOutIsNewComponent)
 {
@@ -506,9 +523,11 @@ ComponentType* FChaosVDGeometryBuilder::GetAvailableInstancedStaticMeshComponent
 		}
 
 		// If this is a Instanced Static Mesh Component, make sure we set the reverse culling flag correctly
+		// And the material
 		if (UChaosVDInstancedStaticMeshComponent* AsInstancedMeshComponent = Component)
 		{
 			AsInstancedMeshComponent->bReverseCulling = EnumHasAnyFlags(MeshComponentAttributeFlags, EChaosVDMeshAttributesFlags::MirroredGeometry);
+			SetMeshComponentMaterial<ComponentType>(MeshComponentAttributeFlags, AsInstancedMeshComponent);
 		}
 
 		bOutIsNewComponent = true;
@@ -535,6 +554,8 @@ ComponentType* FChaosVDGeometryBuilder::GetAvailableMeshComponent(const TSharedP
 			return nullptr;
 		}
 
+		SetMeshComponentMaterial(MeshComponentAttributeFlags, MeshComponent);
+
 		bOutIsNewComponent = true;
 	}
 
diff --git a/Engine/Plugins/Experimental/ChaosVD/Source/ChaosVD/Public/ChaosVDGeometryDataComponent.h b/Engine/Plugins/Experimental/ChaosVD/Source/ChaosVD/Public/ChaosVDGeometryDataComponent.h
--- a/Engine/Plugins/Experimental/ChaosVD/Source/ChaosVD/Public/ChaosVDGeometryDataComponent.h
+++ b/Engine/Plugins/Experimental/ChaosVD/Source/ChaosVD/Public/ChaosVDGeometryDataComponent.h
@@ -114,11 +114,6 @@ public:
 	/** Returns a copy of the current shape collision data to this mesh instance */
 	FChaosVDShapeCollisionData GetGeometryCollisionData() const { return CollisionData; }
 
-	// TODO: This is legacy from the old system and is used by the Static mesh Component implementation mainly where we can change the material instance on the fly
-	// We need to refactor in a way we can move it to the CVD Static Mesh component and let the Instanced Static Mesh Component have they material instances created on component creation time.
-	/** Returns an existing material instance used by this mesh instances, or creates a new one for the provided type */
-	UMaterialInstanceDynamic* GetCachedMaterialInstance(EChaosVDMaterialType Type);
-
 	/** Handles a mesh instance index update reported by the mesh component used to render this mesh instance */
 	void HandleInstanceIndexUpdated(TArrayView<const FInstancedStaticMeshDelegates::FInstanceIndexUpdateData> InIndexUpdates);
 
@@ -158,9 +153,6 @@ private:
 
 	bool bIsSelected = false;
 
-	//TODO: We should avoid using a StrongObjtPtr here as it will make difficult to track any leaks
-	TMap<EChaosVDMaterialType, TStrongObjectPtr<UMaterialInstanceDynamic>> MaterialInstancesByID;
-
 	TWeakPtr<FChaosVDGeometryBuilder> GeometryBuilderInstance = nullptr;
 
 	friend UChaosVDInstancedStaticMeshComponent;
@@ -274,7 +266,36 @@ public:
 
 	/** Calculates the correct visibility state based on the particle state, and applies it to the mesh instance the provided handle represents */
 	static void UpdateMeshVisibility(const TSharedPtr<FChaosVDMeshDataInstanceHandle>& InInstanceHandle, const FChaosVDParticleDataWrapper& InParticleData, bool bIsActive);
+
+	/** Returns an the material to use as a base to create material instances for the provided type */
+	static UMaterialInterface* GetBaseMaterialForType(EChaosVDMaterialType Type);
+
+	/** Creates a material instance using the provided material as a base */
+	static UMaterialInstanceDynamic* CreateMaterialInstance(UMaterialInterface* BaseMaterial);
+
+	/** Creates a material instance for the provided CVD material type */
+	static UMaterialInstanceDynamic* CreateMaterialInstance(EChaosVDMaterialType Type);
+
+	/** Returns the correct material type to use based on the provided Component type and Mesh Attributes */
+	template<typename TComponent>
+	static EChaosVDMaterialType GetMaterialTypeForComponent(EChaosVDMeshAttributesFlags MeshAttributes);
+
 private:
 	/** Returns the color that needs to be used to present the provided particle data based on its state and current selected options */
 	static FLinearColor GetGeometryParticleColor(const TSharedPtr<FChaosVDExtractedGeometryDataHandle>& InGeometryHandle, const FChaosVDParticleDataWrapper& InParticleData, bool bIsServer);
 };
+
+
+template <typename TComponent>
+EChaosVDMaterialType FChaosVDGeometryComponentUtils::GetMaterialTypeForComponent(EChaosVDMeshAttributesFlags MeshAttributes)
+{
+	constexpr bool bIsInstancedMeshComponent = std::is_base_of_v<UInstancedStaticMeshComponent, TComponent>;
+	if (EnumHasAnyFlags(MeshAttributes, EChaosVDMeshAttributesFlags::TranslucentGeometry))
+	{
+		return bIsInstancedMeshComponent ? EChaosVDMaterialType::InstancedQueryOnly : EChaosVDMaterialType::QueryOnlyMaterial;
+	}
+	else
+	{
+		return bIsInstancedMeshComponent ? EChaosVDMaterialType::Instanced : EChaosVDMaterialType::SimOnlyMaterial;
+	}
+}
diff --git a/Engine/Plugins/Experimental/ChaosVD/Source/ChaosVD/Public/Components/ChaosVDInstancedStaticMeshComponent.h b/Engine/Plugins/Experimental/ChaosVD/Source/ChaosVD/Public/Components/ChaosVDInstancedStaticMeshComponent.h
--- a/Engine/Plugins/Experimental/ChaosVD/Source/ChaosVD/Public/Components/ChaosVDInstancedStaticMeshComponent.h
+++ b/Engine/Plugins/Experimental/ChaosVD/Source/ChaosVD/Public/Components/ChaosVDInstancedStaticMeshComponent.h
@@ -67,9 +67,12 @@ protected:
 
 	bool bIsMeshReady = false;
 
-	UPROPERTY()
+	UPROPERTY(Transient)
 	UMaterialInterface* CurrentMaterial = nullptr;
 
+	UPROPERTY(Transient)
+	TMap<EChaosVDMaterialType, TObjectPtr<UMaterialInstanceDynamic>> CachedMaterialInstancesByID;
+
 	FChaosVDMeshReadyDelegate MeshReadyDelegate;
 	FChaosVDMeshComponentEmptyDelegate ComponentEmptyDelegate;
 
diff --git a/Engine/Plugins/Experimental/ChaosVD/Source/ChaosVD/Public/Components/ChaosVDStaticMeshComponent.h b/Engine/Plugins/Experimental/ChaosVD/Source/ChaosVD/Public/Components/ChaosVDStaticMeshComponent.h
--- a/Engine/Plugins/Experimental/ChaosVD/Source/ChaosVD/Public/Components/ChaosVDStaticMeshComponent.h
+++ b/Engine/Plugins/Experimental/ChaosVD/Source/ChaosVD/Public/Components/ChaosVDStaticMeshComponent.h
@@ -57,6 +57,9 @@ protected:
 
 	bool UpdateGeometryKey(uint32 NewHandleGeometryKey);
 
+	/** Returns an existing material instance used by this mesh instances, or creates a new one for the provided type */
+	UMaterialInstanceDynamic* GetCachedMaterialInstance(EChaosVDMaterialType Type);
+
 	uint8 MeshComponentAttributeFlags = 0;
 	uint8 CurrentGeometryKey = 0;
 	bool bIsMeshReady = false;
@@ -67,4 +70,7 @@ protected:
 	TSharedPtr<FChaosVDMeshDataInstanceHandle> CurrentMeshDataHandle = nullptr;
 
 	TSharedPtr<FChaosVDExtractedGeometryDataHandle> CurrentGeometryHandle = nullptr;
+
+	UPROPERTY(Transient)
+	TMap<EChaosVDMaterialType, TObjectPtr<UMaterialInstanceDynamic>> CachedMaterialInstancesByID;
 };
