[ChaosVD] Fix an issue where you can click play in two timelines at the same time.

Now when you clic play in a timeline, the controls in the other ones will be disabled until you clic pause or stop.

diff --git a/Engine/Plugins/Experimental/ChaosVD/Source/ChaosVD/Private/ChaosVDPlaybackController.cpp b/Engine/Plugins/Experimental/ChaosVD/Source/ChaosVD/Private/ChaosVDPlaybackController.cpp
--- a/Engine/Plugins/Experimental/ChaosVD/Source/ChaosVD/Private/ChaosVDPlaybackController.cpp
+++ b/Engine/Plugins/Experimental/ChaosVD/Source/ChaosVD/Private/ChaosVDPlaybackController.cpp
@@ -23,6 +23,8 @@ FChaosVDPlaybackController::FChaosVDPlaybackController(const TWeakPtr<FChaosVDSc
 		Settings->OnPlaybackSettingsChanged().AddRaw(this, &FChaosVDPlaybackController::HandleFrameRateOverrideSettingsChanged);
 		HandleFrameRateOverrideSettingsChanged(Settings);
 	}
+
+	CurrentPlaybackInstigator = IChaosVDPlaybackControllerInstigator::InvalidGuid;
 }
 
 FChaosVDPlaybackController::~FChaosVDPlaybackController()
@@ -656,6 +658,39 @@ void FChaosVDPlaybackController::HandleDisconnectedFromSession()
 	bHasPendingGTUpdateBroadcast = true;
 }
 
+void FChaosVDPlaybackController::RequestStop(const IChaosVDPlaybackControllerInstigator& InPlaybackInstigator)
+{
+	constexpr int32 FrameNumber = 0;
+	constexpr int32 StepNumber = 0;
+	GoToTrackFrame(InPlaybackInstigator.GetInstigatorID(), EChaosVDTrackType::Game, GameTrackID, FrameNumber, StepNumber);
+}
+
+bool FChaosVDPlaybackController::AcquireExclusivePlaybackControls(const IChaosVDPlaybackControllerInstigator& InPlaybackInstigator)
+{
+	const FGuid& InstigatorID = InPlaybackInstigator.GetInstigatorID();
+	if (InstigatorID == CurrentPlaybackInstigator || CurrentPlaybackInstigator == IChaosVDPlaybackControllerInstigator::InvalidGuid)
+	{
+		CurrentPlaybackInstigator = InPlaybackInstigator.GetInstigatorID();
+		bHasPendingGTUpdateBroadcast = true;
+		return true;
+	}
+
+	return false;
+}
+
+bool FChaosVDPlaybackController::ReleaseExclusivePlaybackControls(const IChaosVDPlaybackControllerInstigator& InPlaybackInstigator)
+{
+	const FGuid& InstigatorID = InPlaybackInstigator.GetInstigatorID();
+	if (InstigatorID == CurrentPlaybackInstigator || CurrentPlaybackInstigator == IChaosVDPlaybackControllerInstigator::InvalidGuid)
+	{
+		CurrentPlaybackInstigator = IChaosVDPlaybackControllerInstigator::InvalidGuid;
+		bHasPendingGTUpdateBroadcast = true;
+		return true;
+	}
+
+	return false;
+}
+
 float FChaosVDPlaybackController::GetFrameTimeOverride() const
 {
 	constexpr int32 MinimumFrameRateOverride = 1;
diff --git a/Engine/Plugins/Experimental/ChaosVD/Source/ChaosVD/Private/ChaosVDScene.cpp b/Engine/Plugins/Experimental/ChaosVD/Source/ChaosVD/Private/ChaosVDScene.cpp
--- a/Engine/Plugins/Experimental/ChaosVD/Source/ChaosVD/Private/ChaosVDScene.cpp
+++ b/Engine/Plugins/Experimental/ChaosVD/Source/ChaosVD/Private/ChaosVDScene.cpp
@@ -132,7 +132,15 @@ void FChaosVDScene::AddReferencedObjects(FReferenceCollector& Collector)
 
 void FChaosVDScene::UpdateFromRecordedStepData(const int32 SolverID, const FChaosVDStepData& InRecordedStepData, const FChaosVDSolverFrameData& InFrameData)
 {
-	AChaosVDSolverInfoActor* SolverSceneData = SolverDataContainerBySolverID.FindChecked(SolverID);
+	AChaosVDSolverInfoActor* SolverSceneData = nullptr;
+	if (AChaosVDSolverInfoActor** SolverSceneDataPtrPtr = SolverDataContainerBySolverID.Find(SolverID))
+	{
+		SolverSceneData = *SolverSceneDataPtrPtr;
+	}
+	else
+	{
+		UE_LOG(LogChaosVDEditor, Warning, TEXT("[%s] Attempted to playback a solver frame from an invalid solver container"), ANSI_TO_TCHAR(__FUNCTION__));	
+	}
 
 	if (!SolverSceneData)
 	{
diff --git a/Engine/Plugins/Experimental/ChaosVD/Source/ChaosVD/Private/Widgets/ChaosVDPlaybackControlsHelper.cpp b/Engine/Plugins/Experimental/ChaosVD/Source/ChaosVD/Private/Widgets/ChaosVDPlaybackControlsHelper.cpp
new file mode 100644
--- /dev/null
+++ b/Engine/Plugins/Experimental/ChaosVD/Source/ChaosVD/Private/Widgets/ChaosVDPlaybackControlsHelper.cpp
@@ -0,0 +1,52 @@
+﻿// Copyright Epic Games, Inc. All Rights Reserved.
+
+#include "Widgets/ChaosVDPlaybackControlsHelper.h"
+
+#include "ChaosVDModule.h"
+#include "ChaosVDPlaybackController.h"
+#include "ChaosVDPlaybackControllerObserver.h"
+#include "Templates/SharedPointer.h"
+#include "Widgets/SChaosVDTimelineWidget.h"
+
+void Chaos::VisualDebugger::HandleUserPlaybackInputControl(EChaosVDPlaybackButtonsID ButtonClicked, const IChaosVDPlaybackControllerInstigator& Instigator, const TWeakPtr<FChaosVDPlaybackController>& InPlaybackController)
+{
+	if (const TSharedPtr<FChaosVDPlaybackController> PlaybackControllerPtr = InPlaybackController.Pin())
+	{
+		switch (ButtonClicked)
+		{
+			case EChaosVDPlaybackButtonsID::Play:
+				{
+					if (!ensure(PlaybackControllerPtr->AcquireExclusivePlaybackControls(Instigator)))
+					{
+						UE_LOG(LogChaosVDEditor, Warning, TEXT("[%s] Failed to aquire exclusive playback controls..."), ANSI_TO_TCHAR(__FUNCTION__));
+						return;
+					}
+					PlaybackControllerPtr->RequestUnpause();
+					break;
+				}
+			case EChaosVDPlaybackButtonsID::Pause:
+				{
+					if (!ensure(PlaybackControllerPtr->ReleaseExclusivePlaybackControls(Instigator)))
+					{
+						UE_LOG(LogChaosVDEditor, Warning, TEXT("[%s] Failed to release exclusive playback controls..."), ANSI_TO_TCHAR(__FUNCTION__));
+						return;
+					}
+					PlaybackControllerPtr->RequestPause();
+					break;
+				}
+			case EChaosVDPlaybackButtonsID::Stop:
+				{
+					if (!ensure(PlaybackControllerPtr->ReleaseExclusivePlaybackControls(Instigator)))
+					{
+						UE_LOG(LogChaosVDEditor, Warning, TEXT("[%s] Failed to release exclusive playback controls..."), ANSI_TO_TCHAR(__FUNCTION__));
+						return;
+					}
+					PlaybackControllerPtr->RequestStop(Instigator);
+
+					break;
+				}
+			default:
+				break;
+		}
+	}
+}
diff --git a/Engine/Plugins/Experimental/ChaosVD/Source/ChaosVD/Private/Widgets/SChaosVDPlaybackViewport.cpp b/Engine/Plugins/Experimental/ChaosVD/Source/ChaosVD/Private/Widgets/SChaosVDPlaybackViewport.cpp
--- a/Engine/Plugins/Experimental/ChaosVD/Source/ChaosVD/Private/Widgets/SChaosVDPlaybackViewport.cpp
+++ b/Engine/Plugins/Experimental/ChaosVD/Source/ChaosVD/Private/Widgets/SChaosVDPlaybackViewport.cpp
@@ -5,12 +5,14 @@
 #include "ChaosVDEditorMode.h"
 #include "ChaosVDEditorModeTools.h"
 #include "ChaosVDEditorSettings.h"
+#include "ChaosVDModule.h"
 #include "ChaosVDPlaybackController.h"
 #include "ChaosVDPlaybackViewportClient.h"
 #include "ChaosVDScene.h"
 #include "EditorModeManager.h"
 #include "Framework/Application/SlateApplication.h"
 #include "Elements/Framework/TypedElementSelectionSet.h"
+#include "Widgets/ChaosVDPlaybackControlsHelper.h"
 #include "Widgets/SChaosVDTimelineWidget.h"
 #include "Widgets/SChaosVDViewportToolbar.h"
 #include "Widgets/Text/STextBlock.h"
@@ -172,8 +174,12 @@ void SChaosVDPlaybackViewport::HandlePlaybackControllerDataUpdated(TWeakPtr<FCha
 			}
 			else
 			{
-				uint16& CurrentEnabledFlags = GameFramesTimelineWidget->GetMutableElementEnabledFlagsRef();
-				CurrentEnabledFlags = CurrentEnabledFlags | PlaybackElementDisabledDuringLiveSession;
+				const FGuid CurrentPlaybackInstigatorID = ControllerSharedPtr->GetPlaybackInstigatorWithExclusiveControlsID();
+				const bool bUserCanControlPlayback = CurrentPlaybackInstigatorID == InvalidGuid || CurrentPlaybackInstigatorID == GetInstigatorID();
+
+				// When it is not a live session, the Game Frames timeline follows the same rule as other timelines. The controls are locked unless we are who started a Play action
+				GameFramesTimelineWidget->SetIsLocked(!bUserCanControlPlayback);
+
 				GameFramesTimelineWidget->SetAutoStopEnabled(true);
 			}
 		}
@@ -192,7 +198,7 @@ void SChaosVDPlaybackViewport::HandleControllerTrackFrameUpdated(TWeakPtr<FChaos
 	if (const TSharedPtr<FChaosVDPlaybackController> ControllerSharedPtr = InController.Pin())
 	{
 		// The frame number we receive could be from a Solver track, so make sure it is converted to the correct game track frame number
-		if (const FChaosVDTrackInfo* GameTrackInfo = ControllerSharedPtr->GetTrackInfo(EChaosVDTrackType::Game, FChaosVDPlaybackController::GameTrackID))
+		if (FChaosVDTrackInfo* GameTrackInfo = ControllerSharedPtr->GetMutableTrackInfo(EChaosVDTrackType::Game, FChaosVDPlaybackController::GameTrackID))
 		{
 			const int32 GameTrackFrame = ControllerSharedPtr->ConvertCurrentFrameToOtherTrackFrame(UpdatedTrackInfo, GameTrackInfo);
 
@@ -200,6 +206,7 @@ void SChaosVDPlaybackViewport::HandleControllerTrackFrameUpdated(TWeakPtr<FChaos
 			if (InstigatorGuid != GetInstigatorID())
 			{
 				GameFramesTimelineWidget->SetCurrentTimelineFrame(GameTrackFrame, EChaosVDSetTimelineFrameFlags::None);
+				GameTrackInfo->CurrentFrame = GameTrackFrame;
 			}
 
 			GameFramesTimelineWidget->SetTargetFrameTime(ControllerSharedPtr->GetFrameTimeForTrack(EChaosVDTrackType::Game, FChaosVDPlaybackController::GameTrackID, *GameTrackInfo));
@@ -262,24 +269,7 @@ void SChaosVDPlaybackViewport::OnFrameSelectionUpdated(int32 NewFrameIndex) cons
 
 void SChaosVDPlaybackViewport::HandlePlaybackButtonClicked(EChaosVDPlaybackButtonsID ButtonID)
 {
-	if (const TSharedPtr<FChaosVDPlaybackController> PlaybackControllerPtr = PlaybackController.Pin())
-	{
-		switch (ButtonID)
-		{
-		case EChaosVDPlaybackButtonsID::Play:
-			{
-				PlaybackControllerPtr->RequestUnpause();
-				break;
-			}
-		case EChaosVDPlaybackButtonsID::Pause:
-			{
-				PlaybackControllerPtr->RequestPause();
-				break;
-			}
-		default:
-			break;
-		}
-	}
+	Chaos::VisualDebugger::HandleUserPlaybackInputControl(ButtonID, *this, PlaybackController);
 }
 
 #undef LOCTEXT_NAMESPACE
diff --git a/Engine/Plugins/Experimental/ChaosVD/Source/ChaosVD/Private/Widgets/SChaosVDSolverPlaybackControls.cpp b/Engine/Plugins/Experimental/ChaosVD/Source/ChaosVD/Private/Widgets/SChaosVDSolverPlaybackControls.cpp
--- a/Engine/Plugins/Experimental/ChaosVD/Source/ChaosVD/Private/Widgets/SChaosVDSolverPlaybackControls.cpp
+++ b/Engine/Plugins/Experimental/ChaosVD/Source/ChaosVD/Private/Widgets/SChaosVDSolverPlaybackControls.cpp
@@ -3,7 +3,9 @@
 #include "Widgets/SChaosVDSolverPlaybackControls.h"
 
 #include "ChaosVDEditorSettings.h"
+#include "ChaosVDModule.h"
 #include "ChaosVDPlaybackController.h"
+#include "Widgets/ChaosVDPlaybackControlsHelper.h"
 #include "Widgets/SBoxPanel.h"
 #include "Widgets/SChaosVDPlaybackViewport.h"
 #include "Widgets/SChaosVDTimelineWidget.h"
@@ -46,6 +48,7 @@ void SChaosVDSolverPlaybackControls::Construct(const FArguments& InArgs, int32 I
 						SAssignNew(FramesTimelineWidget, SChaosVDTimelineWidget)
 						.ButtonVisibilityFlags(static_cast<uint16>(EChaosVDTimelineElementIDFlags::AllPlayback))
 						.OnFrameChanged_Raw(this, &SChaosVDSolverPlaybackControls::OnFrameSelectionUpdated)
+						.OnButtonClicked_Raw(this, &SChaosVDSolverPlaybackControls::HandlePlaybackButtonClicked)
 						.MaxFrames(0)
 					]
 					+SHorizontalBox::Slot()
@@ -102,6 +105,24 @@ void SChaosVDSolverPlaybackControls::Construct(const FArguments& InArgs, int32 I
 	}
 }
 
+SChaosVDSolverPlaybackControls::~SChaosVDSolverPlaybackControls()
+{
+	if (const TSharedPtr<FChaosVDPlaybackController> PlaybackControllerPtr = PlaybackController.Pin())
+	{
+		PlaybackControllerPtr->ReleaseExclusivePlaybackControls(*this);
+	}
+}
+
+void SChaosVDSolverPlaybackControls::ConditionallyLockPlaybackControl(const TSharedRef<FChaosVDPlaybackController>& InControllerSharedRef)
+{
+	const FGuid CurrentPlaybackInstigatorID = InControllerSharedRef->GetPlaybackInstigatorWithExclusiveControlsID();
+	const bool bUserCanControlPlayback = CurrentPlaybackInstigatorID == InvalidGuid || CurrentPlaybackInstigatorID == GetInstigatorID();
+
+	// On Live Sessions, only the Game Frames timeline controls are allowed for now
+	FramesTimelineWidget->SetIsLocked(InControllerSharedRef->IsPlayingLiveSession() || !bUserCanControlPlayback);
+	StepsTimelineWidget->SetIsLocked(InControllerSharedRef->IsPlayingLiveSession() || !bUserCanControlPlayback);
+}
+
 void SChaosVDSolverPlaybackControls::HandlePlaybackControllerDataUpdated(TWeakPtr<FChaosVDPlaybackController> InController)
 {
 	if (PlaybackController != InController)
@@ -128,9 +149,7 @@ void SChaosVDSolverPlaybackControls::HandlePlaybackControllerDataUpdated(TWeakPt
 		// or just set the slider to start from 1 and handle the offset later 
 		StepsTimelineWidget->UpdateMinMaxValue(0,AvailableSteps != INDEX_NONE ? AvailableSteps -1 : 0);
 
-		// On Live Sessions, only the Game Frames timeline controls are allowed for now
-		FramesTimelineWidget->SetIsLocked(ControllerSharedPtr->IsPlayingLiveSession());
-		StepsTimelineWidget->SetIsLocked(ControllerSharedPtr->IsPlayingLiveSession());
+		ConditionallyLockPlaybackControl(ControllerSharedPtr.ToSharedRef());
 	}
 	else
 	{
@@ -208,6 +227,11 @@ void SChaosVDSolverPlaybackControls::HandleLockStateChanged(bool NewIsLocked)
 	}
 }
 
+void SChaosVDSolverPlaybackControls::HandlePlaybackButtonClicked(EChaosVDPlaybackButtonsID ButtonID)
+{
+	Chaos::VisualDebugger::HandleUserPlaybackInputControl(ButtonID, *this, PlaybackController);
+}
+
 const FSlateBrush* SChaosVDSolverPlaybackControls::GetFrameTypeBadgeBrush() const
 {
 	const FButtonStyle& ButtonStyle = FAppStyle::Get().GetWidgetStyle<FButtonStyle>("Menu.Button");
diff --git a/Engine/Plugins/Experimental/ChaosVD/Source/ChaosVD/Public/ChaosVDPlaybackController.h b/Engine/Plugins/Experimental/ChaosVD/Source/ChaosVD/Public/ChaosVDPlaybackController.h
--- a/Engine/Plugins/Experimental/ChaosVD/Source/ChaosVD/Public/ChaosVDPlaybackController.h
+++ b/Engine/Plugins/Experimental/ChaosVD/Source/ChaosVD/Public/ChaosVDPlaybackController.h
@@ -10,6 +10,7 @@
 #include "Delegates/Delegate.h"
 #include "HAL/ThreadSafeBool.h"
 
+class IChaosVDPlaybackControllerInstigator;
 class UChaosVDEditorSettings;
 struct FChaosVDTraceSessionDescriptor;
 struct FChaosVDTrackInfo;
@@ -226,6 +227,13 @@ public:
 	void RequestPause() { bPauseRequested = true; }
 	void RequestUnpause() { bPauseRequested = false; }
 	bool HasPauseRequest() const { return bPauseRequested; }
+	void RequestStop(const IChaosVDPlaybackControllerInstigator& InPlaybackInstigator);
+
+	/** Returns the ID of the current playback instigator that is controlling the playback */
+	const FGuid& GetPlaybackInstigatorWithExclusiveControlsID() const { return CurrentPlaybackInstigator; }
+
+	bool AcquireExclusivePlaybackControls(const IChaosVDPlaybackControllerInstigator& InPlaybackInstigator);
+	bool ReleaseExclusivePlaybackControls(const IChaosVDPlaybackControllerInstigator& InPlaybackInstigator);
 
 	float GetFrameTimeOverride() const;
 	float GetFrameTimeForTrack(EChaosVDTrackType TrackType, int32 TrackID, const FChaosVDTrackInfo& TrackInfo) const;
@@ -289,4 +297,6 @@ protected:
 	bool bPauseRequested = false;
 
 	FDelegateHandle RecordingStoppedHandle;
+
+	FGuid CurrentPlaybackInstigator;
 };
diff --git a/Engine/Plugins/Experimental/ChaosVD/Source/ChaosVD/Public/Widgets/ChaosVDPlaybackControlsHelper.h b/Engine/Plugins/Experimental/ChaosVD/Source/ChaosVD/Public/Widgets/ChaosVDPlaybackControlsHelper.h
new file mode 100644
--- /dev/null
+++ b/Engine/Plugins/Experimental/ChaosVD/Source/ChaosVD/Public/Widgets/ChaosVDPlaybackControlsHelper.h
@@ -0,0 +1,16 @@
+﻿// Copyright Epic Games, Inc. All Rights Reserved.
+
+#pragma once
+#include "HAL/Platform.h"
+#include "Templates/SharedPointer.h"
+
+class FChaosVDPlaybackController;
+class IChaosVDPlaybackControllerInstigator;
+enum class EChaosVDPlaybackButtonsID : uint8;
+
+namespace Chaos::VisualDebugger
+{
+	/* Helper that allows execute the same behavior from different playback control objects for common Playback UI actions **/
+	void HandleUserPlaybackInputControl(EChaosVDPlaybackButtonsID ButtonClicked, const IChaosVDPlaybackControllerInstigator& Instigator, const TWeakPtr<FChaosVDPlaybackController>& InPlaybackController);
+}
+
diff --git a/Engine/Plugins/Experimental/ChaosVD/Source/ChaosVD/Public/Widgets/SChaosVDSolverPlaybackControls.h b/Engine/Plugins/Experimental/ChaosVD/Source/ChaosVD/Public/Widgets/SChaosVDSolverPlaybackControls.h
--- a/Engine/Plugins/Experimental/ChaosVD/Source/ChaosVD/Public/Widgets/SChaosVDSolverPlaybackControls.h
+++ b/Engine/Plugins/Experimental/ChaosVD/Source/ChaosVD/Public/Widgets/SChaosVDSolverPlaybackControls.h
@@ -7,6 +7,7 @@
 #include "Templates/SharedPointer.h"
 #include "Widgets/SCompoundWidget.h"
 
+enum class EChaosVDPlaybackButtonsID : uint8;
 struct FChaosVDTrackInfo;
 class FChaosVDPlaybackController;
 class SChaosVDTimelineWidget;
@@ -34,6 +35,8 @@ public:
 
 	void Construct(const FArguments& InArgs, int32 InSolverID, const TWeakPtr<FChaosVDPlaybackController>& InPlaybackController);
 
+	~SChaosVDSolverPlaybackControls();
+
 private:
 
 	void OnFrameSelectionUpdated(int32 NewFrameIndex);
@@ -44,6 +47,9 @@ private:
 	virtual void HandleControllerTrackFrameUpdated(TWeakPtr<FChaosVDPlaybackController> InController, const FChaosVDTrackInfo* UpdatedTrackInfo, FGuid InstigatorGuid) override;
 
 	void HandleLockStateChanged(bool NewIsLocked);
+	void HandlePlaybackButtonClicked(EChaosVDPlaybackButtonsID ButtonID);
+
+	void ConditionallyLockPlaybackControl(const TSharedRef<FChaosVDPlaybackController>& InControllerSharedRef);
 
 	const FSlateBrush* GetFrameTypeBadgeBrush() const;
 
