[ChaosVD] Fix for potential crash when an invalid trace file is loaded (or attempted to be loaded)

There was a missing null check before trying to access then name of the newly loaded trace session.

CVD already handles failures on that stage, so the null check will prevent the crash and allow the failure handling happen down the line.

diff --git a/Engine/Plugins/Experimental/ChaosVD/Source/ChaosVD/Private/Trace/ChaosVDTraceManager.cpp b/Engine/Plugins/Experimental/ChaosVD/Source/ChaosVD/Private/Trace/ChaosVDTraceManager.cpp
--- a/Engine/Plugins/Experimental/ChaosVD/Source/ChaosVD/Private/Trace/ChaosVDTraceManager.cpp
+++ b/Engine/Plugins/Experimental/ChaosVD/Source/ChaosVD/Private/Trace/ChaosVDTraceManager.cpp
@@ -33,9 +33,11 @@ FString FChaosVDTraceManager::LoadTraceFile(const FString& InTraceFilename)
 	ITraceServicesModule& TraceServicesModule = FModuleManager::LoadModuleChecked<ITraceServicesModule>("TraceServices");
 	if (const TSharedPtr<TraceServices::IAnalysisService> TraceAnalysisService = TraceServicesModule.GetAnalysisService())
 	{
-		const TSharedPtr<const TraceServices::IAnalysisSession> NewSession = TraceAnalysisService->StartAnalysis(*InTraceFilename);
-		AnalysisSessionByName.Add(InTraceFilename, NewSession);
-		return NewSession->GetName();
+		if (const TSharedPtr<const TraceServices::IAnalysisSession> NewSession = TraceAnalysisService->StartAnalysis(*InTraceFilename))
+		{
+			AnalysisSessionByName.Add(InTraceFilename, NewSession);
+			return NewSession->GetName();
+		}
 	}
 
 	return FString();
