[ChaosVD] Allow selecting between three different locations when using the open file button in CVD.

This fixes an issue where the button to open folder, was opening the last folder open instead of the profiling folder. That could have been fixed by ensuring the profiling folder was open, but having the ability to just open the last opened folder is a time saver in most cases when looking to CVD trace files recorded in different editors and not what you recorded withing the editor instance.

diff --git a/Engine/Plugins/Experimental/ChaosVD/Source/ChaosVD/Private/Widgets/SChaosBrowseTraceFileSourceModal.cpp b/Engine/Plugins/Experimental/ChaosVD/Source/ChaosVD/Private/Widgets/SChaosBrowseTraceFileSourceModal.cpp
--- a/Engine/Plugins/Experimental/ChaosVD/Source/ChaosVD/Private/Widgets/SChaosBrowseTraceFileSourceModal.cpp
+++ b/Engine/Plugins/Experimental/ChaosVD/Source/ChaosVD/Private/Widgets/SChaosBrowseTraceFileSourceModal.cpp
@@ -2,14 +2,43 @@
 
 #include "Widgets/SChaosBrowseTraceFileSourceModal.h"
 
+#include "ChaosVDModule.h"
 #include "Editor.h"
 #include "GenericPlatform/GenericPlatformMisc.h"
 #include "Widgets/Layout/SUniformGridPanel.h"
+#include "Widgets/SChaosVDNameListPicker.h"
 #include "Widgets/Text/STextBlock.h"
 #include "Widgets/Input/SButton.h"
 
 #define LOCTEXT_NAMESPACE "ChaosVisualDebugger"
 
+
+void SChaosBrowseTraceFileSourceModal::PopulateLocationNamesList()
+{
+	TArray<TSharedPtr<FName>> NewNameList;
+	NewNameList.Reserve(LocationNameToResponseID.Num());
+	Algo::Transform(LocationNameToResponseID, NewNameList, [](const TPair<FName, EChaosVDBrowseFileModalResponse>& TransformData){ return MakeShared<FName>(TransformData.Key);});
+	NamePickerWidget->UpdateNameList(MoveTemp(NewNameList));
+}
+
+void SChaosBrowseTraceFileSourceModal::RegisterAvailableLocationOptions()
+{
+	// TODO: This FText to FName conversion can be avoided if we make SChaosVDNameListPicker accept FText as well
+	// But doing so it is a larger change than the one we can do now. SChaosVDNameListPicker will need to take the type as a template and there are several places currently using it.
+	// FText.ToString should return the already localized name so we should be good in that area
+
+	CVDFilesLocationNames[0] = LOCTEXT("BrowseLastFolder", "Last Opened Folder");
+	LocationNameToResponseID.Add(FName(CVDFilesLocationNames[0].ToString()), EChaosVDBrowseFileModalResponse::OpenLastFolder);
+	
+	CVDFilesLocationNames[1] = LOCTEXT("BrowseProfilingFolder", "Profiling Folder");
+	LocationNameToResponseID.Add(FName(CVDFilesLocationNames[1].ToString()), EChaosVDBrowseFileModalResponse::OpenProfilingFolder);
+
+	CVDFilesLocationNames[2] = LOCTEXT("BrowseTraceStore", "Trace Store Folder");
+	LocationNameToResponseID.Add(FName(CVDFilesLocationNames[2].ToString()), EChaosVDBrowseFileModalResponse::OpenTraceStore);
+
+	CurrentSelectedLocationName = FName(CVDFilesLocationNames[0].ToString());
+}
+
 void SChaosBrowseTraceFileSourceModal::Construct(const FArguments& InArgs)
 {
 	SWindow::Construct(SWindow::FArguments()
@@ -44,22 +73,23 @@ void SChaosBrowseTraceFileSourceModal::Construct(const FArguments& InArgs)
 				.MinDesiredSlotHeight(FAppStyle::GetFloat("StandardDialog.MinDesiredSlotHeight"))
 				+SUniformGridPanel::Slot(0, 0)
 				[
-					SNew(SButton)
-					.HAlign(HAlign_Center)
-					.ContentPadding(FAppStyle::GetMargin("StandardDialog.ContentPadding"))
-					.Text(LOCTEXT("BrowseFolder", "Folder"))
-					.OnClicked(this, &SChaosBrowseTraceFileSourceModal::OnButtonClick, EChaosVDBrowseFileModalResponse::OpenFolder)
+					SAssignNew(NamePickerWidget, SChaosVDNameListPicker)
+					.OnNameSleceted_Raw(this, &SChaosBrowseTraceFileSourceModal::HandleSessionNameSelected)
 				]
 				+SUniformGridPanel::Slot(1, 0)
 				[
 					SNew(SButton)
 					.HAlign(HAlign_Center)
 					.ContentPadding(FAppStyle::GetMargin("StandardDialog.ContentPadding"))
-					.Text(LOCTEXT("BrowseTraceStore", "Trace Store"))
-					.OnClicked(this, &SChaosBrowseTraceFileSourceModal::OnButtonClick, EChaosVDBrowseFileModalResponse::OpenTraceStore)
+					.Text(LOCTEXT("OpenSelectedLocation", "Open"))
+					.OnClicked(this, &SChaosBrowseTraceFileSourceModal::OnButtonClick)
 				]
 			]
 	]);
+	
+	RegisterAvailableLocationOptions();
+
+	PopulateLocationNamesList();
 }
 
 EChaosVDBrowseFileModalResponse SChaosBrowseTraceFileSourceModal::ShowModal()
@@ -68,9 +98,26 @@ EChaosVDBrowseFileModalResponse SChaosBrowseTraceFileSourceModal::ShowModal()
 	return UserResponse;
 }
 
-FReply SChaosBrowseTraceFileSourceModal::OnButtonClick(EChaosVDBrowseFileModalResponse Response)
+void SChaosBrowseTraceFileSourceModal::HandleSessionNameSelected(TSharedPtr<FName> SelectedName)
+{
+	if (SelectedName)
+	{
+		CurrentSelectedLocationName = *SelectedName;
+	}
+}
+
+FReply SChaosBrowseTraceFileSourceModal::OnButtonClick()
 {
-	UserResponse = Response;
+	if (const EChaosVDBrowseFileModalResponse* SelectedLocationResponse = LocationNameToResponseID.Find(CurrentSelectedLocationName))
+	{
+		UserResponse = *SelectedLocationResponse;
+	}
+	else
+	{
+		UserResponse = EChaosVDBrowseFileModalResponse::OpenLastFolder;
+		UE_LOG(LogChaosVDEditor, Warning, TEXT("[%s] Failed to get selected response [%s] | Defaulting to opening the last opened folder..."), ANSI_TO_TCHAR(__FUNCTION__), *CurrentSelectedLocationName.ToString());
+	}
+	
 	RequestDestroyWindow();
 	return FReply::Handled();
 }
diff --git a/Engine/Plugins/Experimental/ChaosVD/Source/ChaosVD/Private/Widgets/SChaosVDMainTab.cpp b/Engine/Plugins/Experimental/ChaosVD/Source/ChaosVD/Private/Widgets/SChaosVDMainTab.cpp
--- a/Engine/Plugins/Experimental/ChaosVD/Source/ChaosVD/Private/Widgets/SChaosVDMainTab.cpp
+++ b/Engine/Plugins/Experimental/ChaosVD/Source/ChaosVD/Private/Widgets/SChaosVDMainTab.cpp
@@ -434,11 +434,16 @@ void SChaosVDMainTab::BrowseAndOpenChaosVDRecording()
 	EChaosVDBrowseFileModalResponse Response = SessionBrowserModal->ShowModal();
 	switch(Response)
 	{
-		case EChaosVDBrowseFileModalResponse::OpenFolder:
+		case EChaosVDBrowseFileModalResponse::OpenLastFolder:
 			{
 				BrowseChaosVDRecordingFromFolder();
 				break;
 			}
+		case EChaosVDBrowseFileModalResponse::OpenProfilingFolder:
+			{
+				BrowseChaosVDRecordingFromFolder(*FPaths::ProfilingDir());
+				break;
+			}
 		case EChaosVDBrowseFileModalResponse::OpenTraceStore:
 			{
 				//TODO: Support remote Trace Stores
@@ -474,7 +479,7 @@ void SChaosVDMainTab::BrowseChaosVDRecordingFromFolder(FStringView FolderPath)
 			FSlateApplication::Get().FindBestParentWindowHandleForDialogs(AsShared()),
 			LOCTEXT("OpenDialogTitle", "Open Chaos Visual Debug File").ToString(),
 			FolderPath.GetData(),
-			*FPaths::ProfilingDir(),
+			TEXT(""),
 			*ExtensionStr,
 			EFileDialogFlags::None,
 			OutOpenFilenames
diff --git a/Engine/Plugins/Experimental/ChaosVD/Source/ChaosVD/Public/Widgets/SChaosBrowseTraceFileSourceModal.h b/Engine/Plugins/Experimental/ChaosVD/Source/ChaosVD/Public/Widgets/SChaosBrowseTraceFileSourceModal.h
--- a/Engine/Plugins/Experimental/ChaosVD/Source/ChaosVD/Public/Widgets/SChaosBrowseTraceFileSourceModal.h
+++ b/Engine/Plugins/Experimental/ChaosVD/Source/ChaosVD/Public/Widgets/SChaosBrowseTraceFileSourceModal.h
@@ -5,11 +5,13 @@
 #include "Widgets/SCompoundWidget.h"
 #include "Widgets/SWindow.h"
 
+class SChaosVDNameListPicker;
 class FReply;
 
 enum class EChaosVDBrowseFileModalResponse
 {
-	OpenFolder,
+	OpenLastFolder,
+	OpenProfilingFolder,
 	OpenTraceStore,
 	Cancel
 };
@@ -34,7 +36,21 @@ public:
 
 protected:
 
-	FReply OnButtonClick(EChaosVDBrowseFileModalResponse Response);
+	void PopulateLocationNamesList();
+
+	void RegisterAvailableLocationOptions();
+
+	void HandleSessionNameSelected(TSharedPtr<FName> SelectedName);
+	
+	FReply OnButtonClick();
+
+	TSharedPtr<SChaosVDNameListPicker> NamePickerWidget;
+	
+	TMap<FName, EChaosVDBrowseFileModalResponse> LocationNameToResponseID;
+
+	FName CurrentSelectedLocationName;
 
 	EChaosVDBrowseFileModalResponse UserResponse = EChaosVDBrowseFileModalResponse::Cancel;
+
+	FText CVDFilesLocationNames[3];
 };
