[ChaosVD] Added support for Record and Visualize Joint Constraints

Also Added Collision data recording and joint constraints data recording support for RBAN Solvers.

Joint constraint data recording is gated behind a dedicated CVD Data channel, which is disabled by default

diff --git a/Engine/Plugins/Experimental/ChaosVD/Source/ChaosVD/Private/Actors/ChaosVDSolverInfoActor.cpp b/Engine/Plugins/Experimental/ChaosVD/Source/ChaosVD/Private/Actors/ChaosVDSolverInfoActor.cpp
--- a/Engine/Plugins/Experimental/ChaosVD/Source/ChaosVD/Private/Actors/ChaosVDSolverInfoActor.cpp
+++ b/Engine/Plugins/Experimental/ChaosVD/Source/ChaosVD/Private/Actors/ChaosVDSolverInfoActor.cpp
@@ -8,6 +8,7 @@
 #include "EditorActorFolders.h"
 #include "Components/ChaosVDParticleDataComponent.h"
 #include "Components/ChaosVDSolverCollisionDataComponent.h"
+#include "Components/ChaosVDSolverJointConstraintDataComponent.h"
 #include "Elements/Framework/TypedElementSelectionSet.h"
 #include "Engine/World.h"
 #include "Misc/ScopedSlowTask.h"
@@ -18,6 +19,7 @@ AChaosVDSolverInfoActor::AChaosVDSolverInfoActor(const FObjectInitializer& Objec
 {
 	CollisionDataComponent = CreateDefaultSubobject<UChaosVDSolverCollisionDataComponent>(TEXT("SolverCollisionDataComponent"));
 	ParticleDataComponent = CreateDefaultSubobject<UChaosVDParticleDataComponent>(TEXT("ParticleCollisionDataComponent"));
+	JointsDataComponent = CreateDefaultSubobject<UChaosVDSolverJointConstraintDataComponent>(TEXT("JointDataComponent"));
 	bIsServer = false;
 }
 
diff --git a/Engine/Plugins/Experimental/ChaosVD/Source/ChaosVD/Private/ChaosVDConstraintDataInspectorTab.cpp b/Engine/Plugins/Experimental/ChaosVD/Source/ChaosVD/Private/ChaosVDConstraintDataInspectorTab.cpp
new file mode 100644
--- /dev/null
+++ b/Engine/Plugins/Experimental/ChaosVD/Source/ChaosVD/Private/ChaosVDConstraintDataInspectorTab.cpp
@@ -0,0 +1,45 @@
+﻿// Copyright Epic Games, Inc. All Rights Reserved.
+
+#include "ChaosVDConstraintDataInspectorTab.h"
+
+#include "Widgets/SChaosVDConstraintDataInspector.h"
+
+#define LOCTEXT_NAMESPACE "ChaosVisualDebugger"
+
+FChaosVDConstraintDataInspectorTab::~FChaosVDConstraintDataInspectorTab()
+{
+}
+
+TSharedRef<SDockTab> FChaosVDConstraintDataInspectorTab::HandleTabSpawnRequest(const FSpawnTabArgs& Args)
+{
+	TSharedRef<SDockTab> DetailsPanelTab =
+	SNew(SDockTab)
+	.TabRole(ETabRole::MajorTab)
+	.Label(LOCTEXT("ConstraintDataInspectorTab", "Constraint Inspector"))
+	.ToolTipText(LOCTEXT("ConstraintDataInspectorTabToolTip", "See the details of the any non-collision constraint selected in the viewport"));
+
+	if (const TSharedPtr<SChaosVDMainTab> MainTabPtr = OwningTabWidget.Pin())
+	{
+		DetailsPanelTab->SetContent
+		(
+			SAssignNew(ConstraintDataInspector, SChaosVDConstraintDataInspector, GetChaosVDScene())
+		);
+	}
+	else
+	{
+		DetailsPanelTab->SetContent(GenerateErrorWidget());
+	}
+
+	HandleTabSpawned(DetailsPanelTab);
+
+	return DetailsPanelTab;
+}
+
+void FChaosVDConstraintDataInspectorTab::HandleTabClosed(TSharedRef<SDockTab> InTabClosed)
+{
+	FChaosVDTabSpawnerBase::HandleTabClosed(InTabClosed);
+
+	ConstraintDataInspector.Reset();
+}
+
+#undef LOCTEXT_NAMESPACE
diff --git a/Engine/Plugins/Experimental/ChaosVD/Source/ChaosVD/Private/ChaosVDParticleActor.cpp b/Engine/Plugins/Experimental/ChaosVD/Source/ChaosVD/Private/ChaosVDParticleActor.cpp
--- a/Engine/Plugins/Experimental/ChaosVD/Source/ChaosVD/Private/ChaosVDParticleActor.cpp
+++ b/Engine/Plugins/Experimental/ChaosVD/Source/ChaosVD/Private/ChaosVDParticleActor.cpp
@@ -433,7 +433,7 @@ const TArray<TSharedPtr<FChaosVDParticlePairMidPhase>>* AChaosVDParticleActor::G
 	{
 		if (const UChaosVDSolverCollisionDataComponent* CollisionDataComponent = SolverInfoActor->GetCollisionDataComponent())
 		{
-			return CollisionDataComponent->GetMidPhasesForParticle(ParticleDataPtr->ParticleIndex, EChaosVDCollisionParticlePairSlot::Any);
+			return CollisionDataComponent->GetMidPhasesForParticle(ParticleDataPtr->ParticleIndex, EChaosVDParticlePairSlot::Any);
 		}
 	}
 
diff --git a/Engine/Plugins/Experimental/ChaosVD/Source/ChaosVD/Private/ChaosVDPlaybackViewportClient.cpp b/Engine/Plugins/Experimental/ChaosVD/Source/ChaosVD/Private/ChaosVDPlaybackViewportClient.cpp
--- a/Engine/Plugins/Experimental/ChaosVD/Source/ChaosVD/Private/ChaosVDPlaybackViewportClient.cpp
+++ b/Engine/Plugins/Experimental/ChaosVD/Source/ChaosVD/Private/ChaosVDPlaybackViewportClient.cpp
@@ -20,6 +20,7 @@
 #include "Actors/ChaosVDSolverInfoActor.h"
 #include "Components/ChaosVDParticleDataComponent.h"
 #include "Components/ChaosVDSolverCollisionDataComponent.h"
+#include "Components/ChaosVDSolverJointConstraintDataComponent.h"
 #include "Components/InstancedStaticMeshComponent.h"
 #include "Visualizers/ChaosVDDebugDrawUtils.h"
 #include "Widgets/SChaosVDMainTab.h"
@@ -371,6 +372,14 @@ void FChaosVDPlaybackViewportClient::Draw(const FSceneView* View, FPrimitiveDraw
 					Visualizer->DrawVisualization(CollisionDataComponent, View, PDI);
 				}
 			}
+
+			if (const UChaosVDSolverJointConstraintDataComponent* JointConstraintDataComponent = SolverInfoWithID.Value ? SolverInfoWithID.Value->GetJointsDataComponent() : nullptr)
+			{
+				if (const TSharedPtr<FComponentVisualizer> Visualizer = MainTabToolkitHost->FindComponentVisualizer(JointConstraintDataComponent->StaticClass()))
+				{
+					Visualizer->DrawVisualization(JointConstraintDataComponent, View, PDI);
+				}
+			}
 		}
 
 		if (const UChaosVDSceneQueryDataComponent* SceneQueryDataComponent = ScenePtr->GetSceneQueryDataContainerComponent())
diff --git a/Engine/Plugins/Experimental/ChaosVD/Source/ChaosVD/Private/ChaosVDScene.cpp b/Engine/Plugins/Experimental/ChaosVD/Source/ChaosVD/Private/ChaosVDScene.cpp
--- a/Engine/Plugins/Experimental/ChaosVD/Source/ChaosVD/Private/ChaosVDScene.cpp
+++ b/Engine/Plugins/Experimental/ChaosVD/Source/ChaosVD/Private/ChaosVDScene.cpp
@@ -27,6 +27,7 @@
 #include "Selection.h"
 #include "UObject/Package.h"
 #include "WorldPersistentFolders.h"
+#include "Components/ChaosVDSolverJointConstraintDataComponent.h"
 #include "Engine/Level.h"
 
 #define LOCTEXT_NAMESPACE "ChaosVisualDebugger"
@@ -201,6 +202,8 @@ void FChaosVDScene::UpdateFromRecordedStepData(const int32 SolverID, const FChao
 	}
 
 	UpdateParticlesCollisionData(InRecordedStepData, SolverID);
+
+	UpdateJointConstraintsData(InRecordedStepData, SolverID);
 	
 	const TMap<int32, AChaosVDParticleActor*>& AllSolverParticlesByID = SolverSceneData->GetAllParticleActorsByIDMap();
 
@@ -243,6 +246,17 @@ void FChaosVDScene::UpdateParticlesCollisionData(const FChaosVDStepData& InRecor
 	}
 }
 
+void FChaosVDScene::UpdateJointConstraintsData(const FChaosVDStepData& InRecordedStepData, int32 SolverID)
+{
+	if (AChaosVDSolverInfoActor* SolverDataInfoContainer = SolverDataContainerBySolverID.FindChecked(SolverID))
+	{
+		if (UChaosVDSolverJointConstraintDataComponent* JointsDataContainer = SolverDataInfoContainer->GetJointsDataComponent())
+		{
+			JointsDataContainer->UpdateConstraintData(InRecordedStepData.RecordedJointConstraints);
+		}
+	}
+}
+
 void FChaosVDScene::HandleNewGeometryData(const Chaos::FConstImplicitObjectPtr& GeometryData, const uint32 GeometryID) const
 {
 	NewGeometryAvailableDelegate.Broadcast(GeometryData, GeometryID);
@@ -267,10 +281,11 @@ void FChaosVDScene::CreateSolverInfoActor(int32 SolverID)
 		SolverDataInfo->SetIsServer(bIsServer);
 
 		SolverDataContainerBySolverID.Add(SolverID, SolverDataInfo);
+
+		SolverInfoActorCreatedDelegate.Broadcast(SolverDataInfo);
 	}
 }
 
-
 void FChaosVDScene::HandleEnterNewGameFrame(int32 FrameNumber, const TArray<int32>& AvailableSolversIds, const FChaosVDGameFrameData& InNewGameFrameData)
 {
 	// Currently the particle actors from all the solvers are in the same level, and we manage them by keeping track
diff --git a/Engine/Plugins/Experimental/ChaosVD/Source/ChaosVD/Private/ChaosVDTabsIDs.cpp b/Engine/Plugins/Experimental/ChaosVD/Source/ChaosVD/Private/ChaosVDTabsIDs.cpp
--- a/Engine/Plugins/Experimental/ChaosVD/Source/ChaosVD/Private/ChaosVDTabsIDs.cpp
+++ b/Engine/Plugins/Experimental/ChaosVD/Source/ChaosVD/Private/ChaosVDTabsIDs.cpp
@@ -12,4 +12,5 @@ const FName FChaosVDTabID::SolversTrack = TEXT("SolversTrack");
 const FName FChaosVDTabID::CVDEditorSettings = TEXT("CVDEditorSettings");
 const FName FChaosVDTabID::StatusBar = TEXT("StatusBar");
 const FName FChaosVDTabID::CollisionDataDetails = TEXT("CollisionDataDetails");
-const FName FChaosVDTabID::SceneQueryDataDetails = TEXT("SceneQueryDataDetails");
\ No newline at end of file
+const FName FChaosVDTabID::SceneQueryDataDetails = TEXT("SceneQueryDataDetails");
+const FName FChaosVDTabID::JointsDataDetails = TEXT("JointsDataDetails");
\ No newline at end of file
diff --git a/Engine/Plugins/Experimental/ChaosVD/Source/ChaosVD/Private/Components/ChaosVDSolverCollisionDataComponent.cpp b/Engine/Plugins/Experimental/ChaosVD/Source/ChaosVD/Private/Components/ChaosVDSolverCollisionDataComponent.cpp
--- a/Engine/Plugins/Experimental/ChaosVD/Source/ChaosVD/Private/Components/ChaosVDSolverCollisionDataComponent.cpp
+++ b/Engine/Plugins/Experimental/ChaosVD/Source/ChaosVD/Private/Components/ChaosVDSolverCollisionDataComponent.cpp
@@ -22,31 +22,27 @@ void UChaosVDSolverCollisionDataComponent::UpdateCollisionData(const TArray<TSha
 
 	for (const TSharedPtr<FChaosVDParticlePairMidPhase>& ParticleIDMidPhase : InMidPhaseData)
 	{
-		AddCollisionDataToParticleIDMap(MidPhasesByParticleID0, ParticleIDMidPhase, ParticleIDMidPhase->Particle0Idx);
-		AddCollisionDataToParticleIDMap(MidPhasesByParticleID1, ParticleIDMidPhase, ParticleIDMidPhase->Particle1Idx);
+		Chaos::VisualDebugger::Utils::AddDataDataToParticleIDMap(MidPhasesByParticleID0, ParticleIDMidPhase, ParticleIDMidPhase->Particle0Idx);
+		Chaos::VisualDebugger::Utils::AddDataDataToParticleIDMap(MidPhasesByParticleID1, ParticleIDMidPhase, ParticleIDMidPhase->Particle1Idx);
 
 		ConstraintsByParticleID0.Reserve(ParticleIDMidPhase->Constraints.Num());
 		ConstraintsByParticleID1.Reserve(ParticleIDMidPhase->Constraints.Num());
 		for (FChaosVDConstraint& Constraint : ParticleIDMidPhase->Constraints)
 		{
-			AddCollisionDataToParticleIDMap<FChaosVDConstraintByParticleMap, FChaosVDConstraint*>(ConstraintsByParticleID0, &Constraint, Constraint.Particle0Index);
-			AddCollisionDataToParticleIDMap<FChaosVDConstraintByParticleMap, FChaosVDConstraint*>(ConstraintsByParticleID1, &Constraint, Constraint.Particle1Index);
+			Chaos::VisualDebugger::Utils::AddDataDataToParticleIDMap<FChaosVDConstraintByParticleMap, FChaosVDConstraint*>(ConstraintsByParticleID0, &Constraint, Constraint.Particle0Index);
+			Chaos::VisualDebugger::Utils::AddDataDataToParticleIDMap<FChaosVDConstraintByParticleMap, FChaosVDConstraint*>(ConstraintsByParticleID1, &Constraint, Constraint.Particle1Index);
 		}
 	}
 }
 
-const TArray<TSharedPtr<FChaosVDParticlePairMidPhase>>* UChaosVDSolverCollisionDataComponent::GetMidPhasesForParticle(int32 ParticleID, EChaosVDCollisionParticlePairSlot Options) const
+const TArray<TSharedPtr<FChaosVDParticlePairMidPhase>>* UChaosVDSolverCollisionDataComponent::GetMidPhasesForParticle(int32 ParticleID, EChaosVDParticlePairSlot Options) const
 {
-	return GetCollisionDataFromMap<FChaosVDMidPhaseByParticleMap, TSharedPtr<FChaosVDParticlePairMidPhase>>(MidPhasesByParticleID0, MidPhasesByParticleID1, ParticleID, Options);
+	return Chaos::VisualDebugger::Utils::GetDataFromParticlePairMaps<FChaosVDMidPhaseByParticleMap, TSharedPtr<FChaosVDParticlePairMidPhase>>(MidPhasesByParticleID0, MidPhasesByParticleID1, ParticleID, Options);
 }
 
-const TArray<FChaosVDConstraint*>* UChaosVDSolverCollisionDataComponent::GetConstraintsForParticle(int32 ParticleID, EChaosVDCollisionParticlePairSlot Options) const
-{
-	return GetCollisionDataFromMap<FChaosVDConstraintByParticleMap, FChaosVDConstraint*>(ConstraintsByParticleID0, ConstraintsByParticleID1, ParticleID, Options);
-}
-
-void UChaosVDSolverCollisionDataComponent::DrawVisualization(const FSceneView* View, FPrimitiveDrawInterface* PDI)
+const TArray<FChaosVDConstraint*>* UChaosVDSolverCollisionDataComponent::GetConstraintsForParticle(int32 ParticleID, EChaosVDParticlePairSlot Options) const
 {
+	return Chaos::VisualDebugger::Utils::GetDataFromParticlePairMaps<FChaosVDConstraintByParticleMap, FChaosVDConstraint*>(ConstraintsByParticleID0, ConstraintsByParticleID1, ParticleID, Options);
 }
 
 void UChaosVDSolverCollisionDataComponent::ClearCollisionData()
diff --git a/Engine/Plugins/Experimental/ChaosVD/Source/ChaosVD/Private/Components/ChaosVDSolverJointConstraintDataComponent.cpp b/Engine/Plugins/Experimental/ChaosVD/Source/ChaosVD/Private/Components/ChaosVDSolverJointConstraintDataComponent.cpp
new file mode 100644
--- /dev/null
+++ b/Engine/Plugins/Experimental/ChaosVD/Source/ChaosVD/Private/Components/ChaosVDSolverJointConstraintDataComponent.cpp
@@ -0,0 +1,106 @@
+﻿// Copyright Epic Games, Inc. All Rights Reserved.
+
+#include "Components/ChaosVDSolverJointConstraintDataComponent.h"
+
+#include "Actors/ChaosVDSolverInfoActor.h"
+#include "Components/ChaosVDConstraintDataHelpers.h"
+#include "DataWrappers/ChaosVDJointDataWrappers.h"
+
+void FChaosVDJointConstraintSelectionHandle::SetIsSelected(bool bNewSelected)
+{
+	if (const TSharedPtr<FChaosVDJointConstraint> JointDataPtr = JointData.Pin())
+	{
+		JointDataPtr->bIsSelectedInEditor = bNewSelected;
+	}
+}
+
+bool FChaosVDJointConstraintSelectionHandle::IsSelected() const
+{
+	if (const TSharedPtr<FChaosVDJointConstraint> JointDataPtr = JointData.Pin())
+	{
+		return JointDataPtr->bIsSelectedInEditor;
+	}
+
+	return false;
+}
+
+UChaosVDSolverJointConstraintDataComponent::UChaosVDSolverJointConstraintDataComponent()
+{
+	SetCanEverAffectNavigation(false);
+	bNavigationRelevant = false;
+	PrimaryComponentTick.bCanEverTick = false;
+}
+
+void UChaosVDSolverJointConstraintDataComponent::UpdateConstraintData(const TArray<TSharedPtr<FChaosVDJointConstraint>>& InJointData)
+{
+	ClearData();
+
+	AllJointConstraints = InJointData;
+
+	JointConstraintByParticle0.Reserve(InJointData.Num());
+	JointConstraintByParticle1.Reserve(InJointData.Num());
+	JointConstraintByConstraintIndex.Reserve(InJointData.Num());
+
+	for (const TSharedPtr<FChaosVDJointConstraint>& JointConstraint : InJointData)
+	{
+		if (!JointConstraint)
+		{
+			continue;
+		}
+
+		JointConstraintByConstraintIndex.Add(JointConstraint->ConstraintIndex, JointConstraint);
+		Chaos::VisualDebugger::Utils::AddDataDataToParticleIDMap(JointConstraintByParticle0, JointConstraint, JointConstraint->ParticleParIndexes[0]);
+		Chaos::VisualDebugger::Utils::AddDataDataToParticleIDMap(JointConstraintByParticle1, JointConstraint, JointConstraint->ParticleParIndexes[1]);
+	}
+}
+
+const FChaosVDJointDataArray* UChaosVDSolverJointConstraintDataComponent::GetJointConstraintsForParticle(int32 ParticleID, EChaosVDParticlePairSlot Options) const
+{
+	return Chaos::VisualDebugger::Utils::GetDataFromParticlePairMaps<FChaosVDJointDataByParticleMap, TSharedPtr<FChaosVDJointConstraint>>(JointConstraintByParticle0, JointConstraintByParticle1, ParticleID, Options);
+}
+
+void UChaosVDSolverJointConstraintDataComponent::SelectJoint(int32 ConstraintIndex)
+{
+	CurrentJointSelectionHandle.SetIsSelected(false);
+	CurrentJointSelectionHandle = FChaosVDJointConstraintSelectionHandle(GetJointConstraintByIndex(ConstraintIndex));
+	CurrentJointSelectionHandle.SetIsSelected(true);
+
+	SelectionChangeDelegate.Broadcast(CurrentJointSelectionHandle);
+}
+
+void UChaosVDSolverJointConstraintDataComponent::SelectJoint(const FChaosVDJointConstraintSelectionHandle& SelectionHandle)
+{
+	CurrentJointSelectionHandle.SetIsSelected(false);
+	CurrentJointSelectionHandle = SelectionHandle;
+	CurrentJointSelectionHandle.SetIsSelected(true);
+	
+	SelectionChangeDelegate.Broadcast(CurrentJointSelectionHandle);
+}
+
+bool UChaosVDSolverJointConstraintDataComponent::IsJointSelected(int32 ConstraintIndex) const
+{
+	if (const TSharedPtr<FChaosVDJointConstraint> JointDataPtr = CurrentJointSelectionHandle.GetData().Pin())
+	{
+		return JointDataPtr->bIsSelectedInEditor && JointDataPtr->ConstraintIndex == ConstraintIndex;
+	}
+
+	return false;
+}
+
+TSharedPtr<FChaosVDJointConstraint> UChaosVDSolverJointConstraintDataComponent::GetJointConstraintByIndex(int32 ConstraintIndex)
+{
+	if (TSharedPtr<FChaosVDJointConstraint>* JointConstraintPtrPtr = JointConstraintByConstraintIndex.Find(ConstraintIndex))
+	{
+		return *JointConstraintPtrPtr;
+	}
+
+	return nullptr;
+}
+
+void UChaosVDSolverJointConstraintDataComponent::ClearData()
+{
+	AllJointConstraints.Empty();
+	JointConstraintByParticle0.Empty();
+	JointConstraintByParticle1.Empty();
+	JointConstraintByConstraintIndex.Empty();
+}
diff --git a/Engine/Plugins/Experimental/ChaosVD/Source/ChaosVD/Private/Trace/ChaosVDTraceProvider.cpp b/Engine/Plugins/Experimental/ChaosVD/Source/ChaosVD/Private/Trace/ChaosVDTraceProvider.cpp
--- a/Engine/Plugins/Experimental/ChaosVD/Source/ChaosVD/Private/Trace/ChaosVDTraceProvider.cpp
+++ b/Engine/Plugins/Experimental/ChaosVD/Source/ChaosVD/Private/Trace/ChaosVDTraceProvider.cpp
@@ -12,6 +12,7 @@
 #include "Compression/OodleDataCompressionUtil.h"
 #include "Serialization/MemoryReader.h"
 #include "Trace/DataProcessors/ChaosVDConstraintDataProcessor.h"
+#include "Trace/DataProcessors/ChaosVDJointConstraintDataProcessor.h"
 #include "Trace/DataProcessors/ChaosVDMidPhaseDataProcessor.h"
 #include "Trace/DataProcessors/ChaosVDSceneQueryDataProcessor.h"
 #include "Trace/DataProcessors/ChaosVDSceneQueryVisitDataProcessor.h"
@@ -223,6 +224,10 @@ void FChaosVDTraceProvider::RegisterDefaultDataProcessorsIfNeeded()
 	TSharedPtr<FChaosVDSerializedNameEntryDataProcessor> NameEntryDataProcessor = MakeShared<FChaosVDSerializedNameEntryDataProcessor>();
 	NameEntryDataProcessor->SetTraceProvider(AsShared());
 	RegisterDataProcessor(NameEntryDataProcessor);
+	
+	TSharedPtr<FChaosVDJointConstraintDataProcessor> JointConstraintDataProcessor = MakeShared<FChaosVDJointConstraintDataProcessor>();
+	JointConstraintDataProcessor->SetTraceProvider(AsShared());
+	RegisterDataProcessor(JointConstraintDataProcessor);
 
 	bDefaultDataProcessorsRegistered = true;
 }
diff --git a/Engine/Plugins/Experimental/ChaosVD/Source/ChaosVD/Private/Trace/DataProcessors/ChaosVDJointConstraintDataProcessor.cpp b/Engine/Plugins/Experimental/ChaosVD/Source/ChaosVD/Private/Trace/DataProcessors/ChaosVDJointConstraintDataProcessor.cpp
new file mode 100644
--- /dev/null
+++ b/Engine/Plugins/Experimental/ChaosVD/Source/ChaosVD/Private/Trace/DataProcessors/ChaosVDJointConstraintDataProcessor.cpp
@@ -0,0 +1,43 @@
+﻿// Copyright Epic Games, Inc. All Rights Reserved.
+
+#include "Trace/DataProcessors/ChaosVDJointConstraintDataProcessor.h"
+
+#include "DataWrappers/ChaosVDJointDataWrappers.h"
+#include "Trace/ChaosVDTraceProvider.h"
+
+
+FChaosVDJointConstraintDataProcessor::FChaosVDJointConstraintDataProcessor()
+	: IChaosVDDataProcessor(FChaosVDJointConstraint::WrapperTypeName)
+{
+}
+
+bool FChaosVDJointConstraintDataProcessor::ProcessRawData(const TArray<uint8>& InData)
+{
+	TSharedPtr<FChaosVDTraceProvider> ProviderSharedPtr = TraceProvider.Pin();
+	if (!ensure(ProviderSharedPtr.IsValid()))
+	{
+		return false;
+	}
+
+	if (!ensure(ProviderSharedPtr->GetNameTable().IsValid()))
+	{
+		return false;
+	}
+
+	TSharedPtr<FChaosVDJointConstraint> JointConstraint = MakeShared<FChaosVDJointConstraint>();
+	const bool bSuccess = ReadDataFromBuffer(InData, *JointConstraint, ProviderSharedPtr->GetNameTable().ToSharedRef());
+
+	if (bSuccess)
+	{
+		// This can be null if the recording started Mid-Frame. In this case we just discard the data for now
+		if (FChaosVDSolverFrameData* FrameData = ProviderSharedPtr->GetCurrentSolverFrame(JointConstraint->SolverID))
+		{
+			if (ensureMsgf(FrameData->SolverSteps.Num() > 0, TEXT("A Joint Constraint was traced without a valid step scope")))
+			{
+				FrameData->SolverSteps.Last().RecordedJointConstraints.Add(JointConstraint);
+			}
+		}
+	}
+
+	return bSuccess;
+}
diff --git a/Engine/Plugins/Experimental/ChaosVD/Source/ChaosVD/Private/Visualizers/ChaosVDDebugDrawUtils.cpp b/Engine/Plugins/Experimental/ChaosVD/Source/ChaosVD/Private/Visualizers/ChaosVDDebugDrawUtils.cpp
--- a/Engine/Plugins/Experimental/ChaosVD/Source/ChaosVD/Private/Visualizers/ChaosVDDebugDrawUtils.cpp
+++ b/Engine/Plugins/Experimental/ChaosVD/Source/ChaosVD/Private/Visualizers/ChaosVDDebugDrawUtils.cpp
@@ -234,3 +234,8 @@ void FChaosVDDebugDrawUtils::DrawCanvas(FViewport& InViewport, FSceneView& View,
 		}
 	}
 }
+
+FString Chaos::VisualDebugger::Utils::GenerateDebugTextForVector(const FVector& InVector, const FString& VectorName, const FString& InVectorUnits)
+{
+	return FString::Format(TEXT("{5} : {0} {4} \n	|-- X : {1} {4} \n	|-- Y : {2} {4} \n	|-- Z : {3} {4}"), {InVector.Size(), InVector.X, InVector.Y, InVector.Z, InVectorUnits, VectorName });
+}
diff --git a/Engine/Plugins/Experimental/ChaosVD/Source/ChaosVD/Private/Visualizers/ChaosVDJointConstraintsDataComponentVisualizer.cpp b/Engine/Plugins/Experimental/ChaosVD/Source/ChaosVD/Private/Visualizers/ChaosVDJointConstraintsDataComponentVisualizer.cpp
new file mode 100644
--- /dev/null
+++ b/Engine/Plugins/Experimental/ChaosVD/Source/ChaosVD/Private/Visualizers/ChaosVDJointConstraintsDataComponentVisualizer.cpp
@@ -0,0 +1,383 @@
+﻿// Copyright Epic Games, Inc. All Rights Reserved.
+
+#include "Visualizers/ChaosVDJointConstraintsDataComponentVisualizer.h"
+
+#include "ChaosVDScene.h"
+#include "ChaosVDTabsIDs.h"
+#include "EditorModeManager.h"
+#include "EditorViewportClient.h"
+#include "SceneManagement.h"
+#include "SceneView.h"
+#include "Actors/ChaosVDSolverInfoActor.h"
+#include "Components/ChaosVDSolverJointConstraintDataComponent.h"
+#include "Visualizers/ChaosVDDebugDrawUtils.h"
+#include "Widgets/SChaosVDMainTab.h"
+
+IMPLEMENT_HIT_PROXY(HChaosVDJointConstraintProxy, HComponentVisProxy)
+
+namespace Chaos::VisualDebugger::Utils
+{
+	FVector GetCoMWorldPosition(const FChaosVDParticleDataWrapper& ParticleData)
+	{
+		if (ensure(ParticleData.ParticlePositionRotation.HasValidData() && ParticleData.ParticleMassProps.HasValidData()))
+		{
+			return ParticleData.ParticlePositionRotation.MX + ParticleData.ParticlePositionRotation.MR.RotateVector(ParticleData.ParticleMassProps.MCenterOfMass);
+		}
+
+		return FVector::ZeroVector;
+	}
+
+	FQuat GetCoMWorldRotation(const FChaosVDParticleDataWrapper& ParticleData)
+	{
+		if (ensure(ParticleData.ParticlePositionRotation.HasValidData() && ParticleData.ParticleMassProps.HasValidData()))
+		{
+			return ParticleData.ParticlePositionRotation.MR * ParticleData.ParticleMassProps.MRotationOfMass;
+		}
+
+		return FQuat::Identity;
+	}
+
+	FVector ParticleLocalToCoMLocal(const FChaosVDParticleDataWrapper& ParticleData, const FVector& Position)
+	{
+		if (ensure(ParticleData.ParticleMassProps.HasValidData()))
+		{
+			return ParticleData.ParticleMassProps.MRotationOfMass.UnrotateVector(Position - ParticleData.ParticleMassProps.MCenterOfMass);
+		}
+
+		return FVector::ZeroVector;
+	}
+
+	FQuat ParticleLocalToCoMLocal(const FChaosVDParticleDataWrapper& ParticleData, const FQuat& Rotation)
+	{
+		if (ensure(ParticleData.ParticleMassProps.HasValidData()))
+		{
+			return ParticleData.ParticleMassProps.MRotationOfMass.Inverse()* Rotation;
+		}
+
+		return FQuat::Identity;
+	}
+
+	FTransform ParticleLocalToCoMLocal(const FChaosVDParticleDataWrapper& ParticleData, const FTransform& Transform)
+	{
+		return FTransform(ParticleLocalToCoMLocal(ParticleData, Transform.GetRotation()), ParticleLocalToCoMLocal(ParticleData, Transform.GetTranslation()));
+	}
+
+	void CalculateConstraintSpace(const FChaosVDParticleDataWrapper& ParticleData0, const FChaosVDParticleDataWrapper& ParticleData1, const FChaosVDJointConstraint& InJointConstraintData, FVector& OutX0, FMatrix33& OutR0, FVector& OutX1, FMatrix33& OutR1)
+	{
+		// Replicating what is done in FPBDJointConstraints::CalculateConstraintSpace. There invert the index for the particles. As we are not dealing with particle indexes here, just invert the references so Particle 0 is now 1 and viceversa 
+		// Copying the explanation given in that method
+		// "In solvers we need Particle0 to be the parent particle but ConstraintInstance has Particle1 as the parent, so by default
+		// we need to flip the indices before we pass them to the solver."
+
+		const FChaosVDParticleDataWrapper& ParticleDataToEvaluate0 = ParticleData1;
+		const FChaosVDParticleDataWrapper& ParticleDataToEvaluate1 = ParticleData0;
+
+		const FVector P0 = GetCoMWorldPosition(ParticleDataToEvaluate0);
+		const FRotation3 Q0 = GetCoMWorldRotation(ParticleDataToEvaluate0);
+		const FVector P1 = GetCoMWorldPosition(ParticleDataToEvaluate1);
+		const FRotation3 Q1 = GetCoMWorldRotation(ParticleDataToEvaluate1);
+		const FRigidTransform3& XL0 = ParticleLocalToCoMLocal(ParticleDataToEvaluate0, InJointConstraintData.JointSettings.ConnectorTransforms[1]);
+		const FRigidTransform3& XL1 = ParticleLocalToCoMLocal(ParticleDataToEvaluate1,  InJointConstraintData.JointSettings.ConnectorTransforms[0]);
+
+		OutX0 = P0 + Q0 * XL0.GetTranslation();
+		OutX1 = P1 + Q1 * XL1.GetTranslation();
+		OutR0 = FRotation3(Q0 * XL0.GetRotation()).ToMatrix();
+		OutR1 = FRotation3(Q1 * XL1.GetRotation()).ToMatrix();
+	}
+
+	/** Generates a random color based on the selection state and query ID, which will be used to debug draw the scene query */
+	FColor GenerateSelectionAwareDebugColor(FLinearColor InBaseColor, bool bIsSelected)
+	{
+		return bIsSelected ? InBaseColor.ToFColorSRGB() :  (InBaseColor * 0.75f).ToFColorSRGB();
+	}
+
+	bool IsParticleSleepingOrDynamic(const FChaosVDParticleDataWrapper& ParticleData)
+	{
+		return ParticleData.ParticleDynamicsMisc.MObjectState == EChaosVDObjectStateType::Dynamic || ParticleData.ParticleDynamicsMisc.MObjectState == EChaosVDObjectStateType::Sleeping;
+	}
+}
+
+void FChaosVDJointConstraintsDataComponentVisualizer::DrawVisualization(const UActorComponent* Component, const FSceneView* View, FPrimitiveDrawInterface* PDI)
+{
+	const UChaosVDSolverJointConstraintDataComponent* JointConstraintDataComponent = Cast<UChaosVDSolverJointConstraintDataComponent>(Component);
+	if (!JointConstraintDataComponent)
+	{
+		return;
+	}
+	
+	AChaosVDSolverInfoActor* SolverInfoActor = Cast<AChaosVDSolverInfoActor>(Component->GetOwner());
+	if (!SolverInfoActor)
+	{
+		return;
+	}
+
+	const TSharedPtr<FChaosVDScene> CVDScene = SolverInfoActor->GetScene().Pin();
+	if (!CVDScene)
+	{
+		return;
+	}
+
+	const TSharedPtr<FChaosVDRecording> CVDRecording = CVDScene->LoadedRecording;
+	if (!CVDRecording)
+	{
+		return;
+	}
+
+	FChaosVDJointVisualizationDataContext VisualizationContext;
+	VisualizationContext.CVDScene = CVDScene;
+	VisualizationContext.SpaceTransform = SolverInfoActor->GetSimulationTransform();
+	VisualizationContext.SolverInfoActor = SolverInfoActor;
+
+	if (const UChaosVDEditorSettings* EditorSettings = GetDefault<UChaosVDEditorSettings>())
+	{
+		VisualizationContext.VisualizationFlags = EditorSettings->GlobalJointsDataVisualizationFlags;
+		VisualizationContext.bShowDebugText = EditorSettings->bShowDebugText;
+		VisualizationContext.DebugDrawSettings = &EditorSettings->JointsDataDebugDrawSettings;
+	}
+
+	if (!VisualizationContext.IsVisualizationFlagEnabled(EChaosVDJointsDataVisualizationFlags::EnableDraw))
+	{
+		return;
+	}
+
+	bool bCanDrawRecordedSolverInstance = (SolverInfoActor->GetIsServer() && VisualizationContext.IsVisualizationFlagEnabled(EChaosVDJointsDataVisualizationFlags::EnableServerDraw)) ||
+											(!SolverInfoActor->GetIsServer() && VisualizationContext.IsVisualizationFlagEnabled(EChaosVDJointsDataVisualizationFlags::EnableClientDraw));
+	
+	if (!bCanDrawRecordedSolverInstance)
+	{
+		return;
+	}
+
+	// If nothing is selected, fallback to draw all joints
+	const bool bDrawOnlySelected = VisualizationContext.IsVisualizationFlagEnabled(EChaosVDJointsDataVisualizationFlags::OnlyDrawSelected) && JointConstraintDataComponent->GetCurrentSelectionHandle().IsSelected();
+	if (bDrawOnlySelected)
+	{
+		if (const TSharedPtr<FChaosVDJointConstraint> JointConstraint = JointConstraintDataComponent->GetCurrentSelectionHandle().GetData().Pin())
+		{
+			VisualizationContext.DataSelectionHandle = FChaosVDJointConstraintSelectionHandle(JointConstraint);
+			DrawJointConstraint(Component, *JointConstraint, VisualizationContext, View, PDI);
+		}
+	}
+	else
+	{
+		for (const TSharedPtr<FChaosVDJointConstraint>& JointConstraint : JointConstraintDataComponent->GetAllJointConstraints())
+		{
+			if (JointConstraint)
+			{
+				VisualizationContext.DataSelectionHandle = FChaosVDJointConstraintSelectionHandle(JointConstraint);
+				DrawJointConstraint(Component, *JointConstraint, VisualizationContext, View, PDI);
+			}
+		}
+	}
+}
+
+bool FChaosVDJointConstraintsDataComponentVisualizer::VisProxyHandleClick(FEditorViewportClient* InViewportClient, HComponentVisProxy* VisProxy, const FViewportClick& Click)
+{
+	const HChaosVDJointConstraintProxy* JointDataProxy = HitProxyCast<HChaosVDJointConstraintProxy>(VisProxy);
+	if (JointDataProxy == nullptr)
+	{
+		return false;
+	}
+	
+	if (const UChaosVDSolverJointConstraintDataComponent* JointsDataComponent = Cast<UChaosVDSolverJointConstraintDataComponent>(JointDataProxy->Component.Get()))
+	{
+		// Bring the SQ Inspector into focus if available
+		const TSharedPtr<SChaosVDMainTab> MainTabToolkitHost = InViewportClient->GetModeTools() ? StaticCastSharedPtr<SChaosVDMainTab>(InViewportClient->GetModeTools()->GetToolkitHost()) : nullptr;
+		if (const TSharedPtr<FTabManager> TabManager = MainTabToolkitHost ? MainTabToolkitHost->GetTabManager() : nullptr)
+		{
+			TabManager->TryInvokeTab(FChaosVDTabID::JointsDataDetails);
+		}
+
+		const_cast<UChaosVDSolverJointConstraintDataComponent*>(JointsDataComponent)->SelectJoint(JointDataProxy->DataSelectionHandle);
+
+		return true;
+	}
+
+	return false;
+}
+
+void FChaosVDJointConstraintsDataComponentVisualizer::DebugDrawAllAxis(const FChaosVDJointConstraint& InJointConstraintData, FChaosVDJointVisualizationDataContext& VisualizationContext, FPrimitiveDrawInterface* PDI, const float LineThickness, const FVector& InPosition, const Chaos::FMatrix33& InRotationMatrix, TConstArrayView<FLinearColor> AxisColors)
+{
+	for(int32 AxisIndex = 0; AxisIndex < 3 ; AxisIndex++)
+	{
+		FChaosVDDebugDrawUtils::DrawArrowVector(PDI, InPosition, InPosition + VisualizationContext.DebugDrawSettings->GeneralScale * VisualizationContext.DebugDrawSettings->ConstraintAxisLength * VisualizationContext.SpaceTransform.TransformVector(InRotationMatrix.GetAxis(AxisIndex)),nullptr, Chaos::VisualDebugger::Utils::GenerateSelectionAwareDebugColor(AxisColors[AxisIndex], InJointConstraintData.bIsSelectedInEditor), VisualizationContext.DebugDrawSettings->DepthPriority, LineThickness * 0.2f);
+	}
+}
+
+void FChaosVDJointConstraintsDataComponentVisualizer::DrawJointConstraint(const UActorComponent* Component, const FChaosVDJointConstraint& InJointConstraintData, FChaosVDJointVisualizationDataContext& VisualizationContext, const FSceneView* View, FPrimitiveDrawInterface* PDI)
+{
+	using namespace Chaos::VisualDebugger::Utils;
+
+	if (!VisualizationContext.IsVisualizationFlagEnabled(EChaosVDJointsDataVisualizationFlags::DrawDisabled))
+	{
+		if (InJointConstraintData.JointState.bDisabled)
+		{
+			return;
+		}
+	}
+
+	if (!Component)
+	{
+		return;
+	}
+	
+	if (!VisualizationContext.DebugDrawSettings)
+	{
+		return;
+	}
+
+	if (!PDI)
+	{
+		return;
+	}
+
+	if (!View)
+	{
+		return;
+	}
+
+	const FChaosVDParticleDataWrapper* ParticleData0 = nullptr;
+	const FChaosVDParticleDataWrapper* ParticleData1 = nullptr;
+
+	if (AChaosVDParticleActor* Particle0 = VisualizationContext.SolverInfoActor->GetParticleActor(InJointConstraintData.ParticleParIndexes[0]))
+	{
+		ParticleData0 = Particle0->GetParticleData();
+	}
+
+	if (AChaosVDParticleActor* Particle1 = VisualizationContext.SolverInfoActor->GetParticleActor(InJointConstraintData.ParticleParIndexes[1]))
+	{
+		ParticleData1 = Particle1->GetParticleData();
+	}
+
+	if (!ParticleData0 || !ParticleData1)
+	{
+		return;
+	}
+
+	if (!VisualizationContext.IsVisualizationFlagEnabled(EChaosVDJointsDataVisualizationFlags::DrawKinematic))
+	{
+		if (!IsParticleSleepingOrDynamic(*ParticleData0) || !IsParticleSleepingOrDynamic(*ParticleData1))
+		{
+			return;
+		}
+	}
+
+	if (!ParticleData0->ParticleMassProps.HasValidData() || !ParticleData1->ParticleMassProps.HasValidData())
+	{
+		// If we don't have mass data, all the following calculations will be off
+		// TODO: Should we draw just a line between the two particles as fallback?
+		return;
+	}
+
+	// Create a sphere containing the two particle positions to use as a pseudo bounds to determine if we should draw this joint data
+	FVector DiameterSphereViewVector = ParticleData1->ParticlePositionRotation.MX - ParticleData0->ParticlePositionRotation.MX;
+	float ViewRadius = DiameterSphereViewVector.Size() * 0.5f;
+	FVector MiddleViewPoint = ParticleData0->ParticlePositionRotation.MX + DiameterSphereViewVector.GetSafeNormal() * ViewRadius;
+
+	bool bCurrentViewContainsBothLocations = View->ViewFrustum.IntersectSphere(MiddleViewPoint, ViewRadius);
+	if (!bCurrentViewContainsBothLocations)
+	{
+		return;
+	}
+
+	PDI->SetHitProxy(new HChaosVDJointConstraintProxy(Component, VisualizationContext.DataSelectionHandle));
+
+	const float LineThickness = InJointConstraintData.bIsSelectedInEditor ?  VisualizationContext.DebugDrawSettings->BaseLineThickness * 1.5f :  VisualizationContext.DebugDrawSettings->BaseLineThickness;
+
+	FVector Pa = ParticleData1->ParticlePositionRotation.MX;
+	FVector Pb = ParticleData0->ParticlePositionRotation.MX;
+	FVector Ca = GetCoMWorldPosition(*ParticleData1);
+	FVector Cb =  GetCoMWorldPosition(*ParticleData0);
+	FVector Xa, Xb;
+	Chaos::FMatrix33 Ra, Rb;
+	CalculateConstraintSpace(*ParticleData0, *ParticleData1, InJointConstraintData, Xa, Ra, Xb, Rb);
+
+	Pa = VisualizationContext.SpaceTransform.TransformPosition(Pa);
+	Pb = VisualizationContext.SpaceTransform.TransformPosition(Pb);
+	Ca = VisualizationContext.SpaceTransform.TransformPosition(Ca);
+	Cb = VisualizationContext.SpaceTransform.TransformPosition(Cb);
+	Xa = VisualizationContext.SpaceTransform.TransformPosition(Xa);
+	Xb = VisualizationContext.SpaceTransform.TransformPosition(Xb);
+
+	const FLinearColor Red = FLinearColor::Red;
+	const FLinearColor Green = FLinearColor::Green;
+	const FLinearColor Blue = FLinearColor::Blue;
+	const FLinearColor Cyan = FLinearColor::FromSRGBColor(FColor::Cyan);
+	const FLinearColor Magenta = FLinearColor::FromSRGBColor(FColor::Magenta);
+	const FLinearColor Yellow = FLinearColor::Yellow;
+	const FLinearColor White = FLinearColor::White;
+	const FLinearColor Black = FLinearColor::Black;
+
+	if (VisualizationContext.IsVisualizationFlagEnabled(EChaosVDJointsDataVisualizationFlags::ActorConnector))
+	{
+		const float ConnectorThickness = 1.5f * LineThickness;
+		const double CoMSize = VisualizationContext.DebugDrawSettings->CenterOfMassSize;
+		// Leave a gap around the actor position so we can see where the center is
+		FVector Sa = Pa;
+		const double Lena = (Xa - Pa).Size();
+		if (Lena > UE_KINDA_SMALL_NUMBER)
+		{
+			Sa = FMath::Lerp(Pa, Xa, FMath::Clamp<double>(CoMSize / Lena, 0., 1.));
+		}
+		FVector Sb = Pb;
+		const double Lenb = (Xb - Pb).Size();
+		if (Lenb > UE_KINDA_SMALL_NUMBER)
+		{
+			Sb = FMath::Lerp(Pb, Xb, FMath::Clamp<double>(CoMSize / Lena, 0., 1.));
+		}
+
+		FChaosVDDebugDrawUtils::DrawLine(PDI, Pa, Sa, GenerateSelectionAwareDebugColor(White, InJointConstraintData.bIsSelectedInEditor), nullptr, VisualizationContext.DebugDrawSettings->DepthPriority, ConnectorThickness);
+		FChaosVDDebugDrawUtils::DrawLine(PDI, Pb, Sb, GenerateSelectionAwareDebugColor(White, InJointConstraintData.bIsSelectedInEditor), nullptr, VisualizationContext.DebugDrawSettings->DepthPriority, ConnectorThickness);
+		FChaosVDDebugDrawUtils::DrawLine(PDI, Sa, Xa, GenerateSelectionAwareDebugColor(Red, InJointConstraintData.bIsSelectedInEditor), nullptr, VisualizationContext.DebugDrawSettings->DepthPriority, ConnectorThickness);
+		FChaosVDDebugDrawUtils::DrawLine(PDI, Sb, Xb, GenerateSelectionAwareDebugColor(Cyan, InJointConstraintData.bIsSelectedInEditor), nullptr, VisualizationContext.DebugDrawSettings->DepthPriority, ConnectorThickness);
+	}
+
+	if (VisualizationContext.IsVisualizationFlagEnabled(EChaosVDJointsDataVisualizationFlags::CenterOfMassConnector))
+	{
+		const float ConnectorThickness = 1.5f * LineThickness;
+		const double CoMSize = VisualizationContext.DebugDrawSettings->CenterOfMassSize;
+		// Leave a gap around the body position so we can see where the center is
+		FVector Sa = Ca;
+		const double Lena = (Xa - Ca).Size();
+		if (Lena > UE_KINDA_SMALL_NUMBER)
+		{
+			Sa = FMath::Lerp(Ca, Xa, FMath::Clamp<double>(CoMSize / Lena, 0., 1.));
+		}
+		FVector Sb = Cb;
+		const double Lenb = (Xb - Cb).Size();
+		if (Lenb > UE_KINDA_SMALL_NUMBER)
+		{
+			Sb = FMath::Lerp(Cb, Xb, FMath::Clamp<double>(CoMSize / Lena, 0., 1.));
+		}
+
+		FChaosVDDebugDrawUtils::DrawLine(PDI, Ca, Sa, GenerateSelectionAwareDebugColor(Black, InJointConstraintData.bIsSelectedInEditor), nullptr, VisualizationContext.DebugDrawSettings->DepthPriority, ConnectorThickness);
+		FChaosVDDebugDrawUtils::DrawLine(PDI, Cb, Sb, GenerateSelectionAwareDebugColor(Black, InJointConstraintData.bIsSelectedInEditor), nullptr, VisualizationContext.DebugDrawSettings->DepthPriority, ConnectorThickness);
+		FChaosVDDebugDrawUtils::DrawLine(PDI, Sa, Xa, GenerateSelectionAwareDebugColor(Red, InJointConstraintData.bIsSelectedInEditor), nullptr, VisualizationContext.DebugDrawSettings->DepthPriority, ConnectorThickness);
+		FChaosVDDebugDrawUtils::DrawLine(PDI, Sb, Xb, GenerateSelectionAwareDebugColor(Cyan, InJointConstraintData.bIsSelectedInEditor), nullptr, VisualizationContext.DebugDrawSettings->DepthPriority, ConnectorThickness);
+	}
+
+	if (VisualizationContext.IsVisualizationFlagEnabled(EChaosVDJointsDataVisualizationFlags::Stretch))
+	{
+		const float StretchThickness = 3.0f * LineThickness;
+		FChaosVDDebugDrawUtils::DrawLine(PDI, Xa, Xb, GenerateSelectionAwareDebugColor(Magenta, InJointConstraintData.bIsSelectedInEditor), nullptr, VisualizationContext.DebugDrawSettings->DepthPriority, StretchThickness);
+	}
+
+	if (VisualizationContext.IsVisualizationFlagEnabled(EChaosVDJointsDataVisualizationFlags::Axes))
+	{
+		const FLinearColor AxisAColors[3] = {Red, Green, Blue};
+		const FLinearColor AxisBColors[3] = {Cyan, Magenta, Yellow};
+		DebugDrawAllAxis(InJointConstraintData, VisualizationContext, PDI, LineThickness, Xa, Ra, AxisAColors);
+		DebugDrawAllAxis(InJointConstraintData, VisualizationContext, PDI, LineThickness, Xb, Rb, AxisBColors);
+	}
+
+	// NOTE: GetLinearImpulse is the positional impulse (pushout)
+	if (VisualizationContext.IsVisualizationFlagEnabled(EChaosVDJointsDataVisualizationFlags::PushOut))
+	{
+		FLinearColor PushOutImpulseColor = FColor(0, 250, 250);
+		FChaosVDDebugDrawUtils::DrawLine(PDI, Xa, Xa + VisualizationContext.DebugDrawSettings->LinearImpulseScale * VisualizationContext.SpaceTransform.TransformVectorNoScale(InJointConstraintData.JointState.LinearImpulse), GenerateSelectionAwareDebugColor(PushOutImpulseColor, InJointConstraintData.bIsSelectedInEditor), nullptr, VisualizationContext.DebugDrawSettings->DepthPriority, LineThickness);
+	}
+
+	//TODO: Should we draw the Angular Impulse
+	
+	PDI->SetHitProxy(nullptr);
+}
diff --git a/Engine/Plugins/Experimental/ChaosVD/Source/ChaosVD/Private/Visualizers/ChaosVDParticleDataComponentVisualizer.cpp b/Engine/Plugins/Experimental/ChaosVD/Source/ChaosVD/Private/Visualizers/ChaosVDParticleDataComponentVisualizer.cpp
--- a/Engine/Plugins/Experimental/ChaosVD/Source/ChaosVD/Private/Visualizers/ChaosVDParticleDataComponentVisualizer.cpp
+++ b/Engine/Plugins/Experimental/ChaosVD/Source/ChaosVD/Private/Visualizers/ChaosVDParticleDataComponentVisualizer.cpp
@@ -130,11 +130,6 @@ bool FChaosVDParticleDataComponentVisualizer::VisProxyHandleClick(FEditorViewpor
 	return bHandled;	
 }
 
-FString FChaosVDParticleDataComponentVisualizer::GenerateDebugTextForVector(const FVector& InVector, const FString& VectorName, const FString& InVectorUnits)
-{
-	return FString::Format(TEXT("{5} : {0} {4} \n	|-- X : {1} {4} \n	|-- Y : {2} {4} \n	|-- Z : {3} {4}"), {InVector.Size(), InVector.X, InVector.Y, InVector.Z, InVectorUnits, VectorName });
-}
-
 void FChaosVDParticleDataComponentVisualizer::DrawParticleVector(FPrimitiveDrawInterface* PDI, const FVector& StartLocation, const FVector& InVector, EChaosVDParticleDataVisualizationFlags VectorID, const FChaosVDParticleDataVisualizationContext& InVisualizationContext, float LineThickness)
 {
 	if (!InVisualizationContext.IsVisualizationFlagEnabled(VectorID))
@@ -147,7 +142,7 @@ void FChaosVDParticleDataComponentVisualizer::DrawParticleVector(FPrimitiveDrawI
 		return;
 	}
 
-	const FString DebugText = InVisualizationContext.bShowDebugText ? GenerateDebugTextForVector(InVector, UEnum::GetDisplayValueAsText(VectorID).ToString(), Chaos::VisualDebugger::ParticleDataUnitsStrings::GetUnitByID(VectorID)) : TEXT("");
+	const FString DebugText = InVisualizationContext.bShowDebugText ? Chaos::VisualDebugger::Utils::GenerateDebugTextForVector(InVector, UEnum::GetDisplayValueAsText(VectorID).ToString(), Chaos::VisualDebugger::ParticleDataUnitsStrings::GetUnitByID(VectorID)) : TEXT("");
 	FChaosVDDebugDrawUtils::DrawArrowVector(PDI, StartLocation, StartLocation +  InVisualizationContext.DebugDrawSettings->GetScaleFortDataID(VectorID) * InVector, DebugText, InVisualizationContext.DebugDrawSettings->ColorSettings.GetColorForDataID(VectorID, InVisualizationContext.bIsSelectedData),  InVisualizationContext.DebugDrawSettings->DepthPriority, LineThickness);
 }
 
diff --git a/Engine/Plugins/Experimental/ChaosVD/Source/ChaosVD/Private/Visualizers/ChaosVDSolverCollisionDataComponentVisualizer.cpp b/Engine/Plugins/Experimental/ChaosVD/Source/ChaosVD/Private/Visualizers/ChaosVDSolverCollisionDataComponentVisualizer.cpp
--- a/Engine/Plugins/Experimental/ChaosVD/Source/ChaosVD/Private/Visualizers/ChaosVDSolverCollisionDataComponentVisualizer.cpp
+++ b/Engine/Plugins/Experimental/ChaosVD/Source/ChaosVD/Private/Visualizers/ChaosVDSolverCollisionDataComponentVisualizer.cpp
@@ -78,7 +78,7 @@ void FChaosVDSolverCollisionDataComponentVisualizer::DrawVisualization(const UAc
 
 		for (const int32 SelectedParticleID : SelectedParticlesIDs)
 		{
-			if (const TArray<TSharedPtr<FChaosVDParticlePairMidPhase>>* ParticleMidPhases = CollisionDataComponent->GetMidPhasesForParticle(SelectedParticleID, EChaosVDCollisionParticlePairSlot::Any))
+			if (const TArray<TSharedPtr<FChaosVDParticlePairMidPhase>>* ParticleMidPhases = CollisionDataComponent->GetMidPhasesForParticle(SelectedParticleID, EChaosVDParticlePairSlot::Any))
 			{
 				for (const TSharedPtr<FChaosVDParticlePairMidPhase>& MidPhase : *ParticleMidPhases)
 				{
diff --git a/Engine/Plugins/Experimental/ChaosVD/Source/ChaosVD/Private/Widgets/SChaosVDConstraintDataInspector.cpp b/Engine/Plugins/Experimental/ChaosVD/Source/ChaosVD/Private/Widgets/SChaosVDConstraintDataInspector.cpp
new file mode 100644
--- /dev/null
+++ b/Engine/Plugins/Experimental/ChaosVD/Source/ChaosVD/Private/Widgets/SChaosVDConstraintDataInspector.cpp
@@ -0,0 +1,264 @@
+﻿// Copyright Epic Games, Inc. All Rights Reserved.
+
+#include "Widgets/SChaosVDConstraintDataInspector.h"
+
+#include "ChaosVDScene.h"
+#include "IStructureDetailsView.h"
+#include "PropertyEditorModule.h"
+#include "Actors/ChaosVDSolverInfoActor.h"
+#include "Widgets/SChaosVDWarningMessageBox.h"
+#include "Components/ChaosVDSolverJointConstraintDataComponent.h"
+#include "DataWrappers/ChaosVDJointDataWrappers.h"
+#include "Modules/ModuleManager.h"
+#include "Widgets/Input/SButton.h"
+#include "Widgets/Layout/SBox.h"
+#include "Widgets/Layout/SScrollBox.h"
+#include "Widgets/Layout/SSplitter.h"
+#include "Widgets/Layout/SUniformGridPanel.h"
+#include "Widgets/Text/STextBlock.h"
+
+#define LOCTEXT_NAMESPACE "ChaosVisualDebugger"
+
+SChaosVDConstraintDataInspector::~SChaosVDConstraintDataInspector()
+{
+	UnregisterSceneEvents();
+}
+
+void SChaosVDConstraintDataInspector::Construct(const FArguments& InArgs, const TWeakPtr<FChaosVDScene>& InScenePtr)
+{
+	SceneWeakPtr = InScenePtr;
+
+	RegisterSceneEvents();
+
+	ConstraintDataDetailsView = CreateDataDetailsView();
+
+	constexpr float NoPadding = 0.0f;
+	constexpr float OuterBoxPadding = 2.0f;
+	constexpr float OuterInnerPadding = 5.0f;
+	constexpr float TagTitleBoxHorizontalPadding = 10.0f;
+	constexpr float TagTitleBoxVerticalPadding = 5.0f;
+	constexpr float InnerDetailsPanelsHorizontalPadding = 15.0f;
+	constexpr float InnerDetailsPanelsVerticalPadding = 15.0f;
+
+	ChildSlot
+	[
+		SNew(SVerticalBox)
+		+SVerticalBox::Slot()
+		.AutoHeight()
+		.Padding(OuterInnerPadding)
+		[
+			SNew(SBox)
+			.Visibility_Raw(this, &SChaosVDConstraintDataInspector::GetOutOfDateWarningVisibility)
+			.Padding(OuterBoxPadding, OuterBoxPadding,OuterBoxPadding,NoPadding)
+			[
+				SNew(SChaosVDWarningMessageBox)
+				.WarningText(LOCTEXT("ConstraintDataOutOfDate", "Scene change detected!. Selected constraint data is out of date..."))
+			]
+		]
+		+SVerticalBox::Slot()
+		.AutoHeight()
+		.Padding(TagTitleBoxHorizontalPadding, TagTitleBoxVerticalPadding, TagTitleBoxHorizontalPadding, NoPadding)
+		[
+			GenerateParticleSelectorButtons()
+		]
+		+SVerticalBox::Slot()
+		.Padding(TagTitleBoxHorizontalPadding, TagTitleBoxVerticalPadding, TagTitleBoxHorizontalPadding, NoPadding)
+		.AutoHeight()
+		[
+			SNew(STextBlock)
+			.Visibility_Raw(this, &SChaosVDConstraintDataInspector::GetNothingSelectedMessageVisibility)
+			.Justification(ETextJustify::Center)
+			.TextStyle(FAppStyle::Get(), "DetailsView.BPMessageTextStyle")
+			.Text(LOCTEXT("ConstraintDataNoSelectedMessage", "Select a Joint Constraint in the viewport to see its details..."))
+			.AutoWrapText(true)
+		]
+		+SVerticalBox::Slot()
+		.Padding(OuterInnerPadding)
+		[
+			SNew(SScrollBox)
+			.Visibility_Raw(this, &SChaosVDConstraintDataInspector::GetDetailsSectionVisibility)
+			+SScrollBox::Slot()
+			.Padding(InnerDetailsPanelsHorizontalPadding,NoPadding,InnerDetailsPanelsHorizontalPadding,InnerDetailsPanelsVerticalPadding)
+			[
+				ConstraintDataDetailsView->GetWidget().ToSharedRef()
+			]
+		]
+	];
+}
+
+void SChaosVDConstraintDataInspector::RegisterSelectionEventsForSolver(AChaosVDSolverInfoActor* SolverInfo)
+{
+	if (UChaosVDSolverJointConstraintDataComponent* JointDataComponent = SolverInfo ? SolverInfo->GetJointsDataComponent() : nullptr)
+	{
+		JointDataComponent->OnSelectionChanged().AddRaw(this, &SChaosVDConstraintDataInspector::SetConstraintDataToInspect);
+	}
+}
+
+void SChaosVDConstraintDataInspector::UnregisterSelectionEventsForSolver(AChaosVDSolverInfoActor* SolverInfo)
+{
+	if (UChaosVDSolverJointConstraintDataComponent* JointDataComponent = SolverInfo ? SolverInfo->GetJointsDataComponent() : nullptr)
+	{
+		JointDataComponent->OnSelectionChanged().RemoveAll(this);
+	}
+}
+
+void SChaosVDConstraintDataInspector::RegisterSceneEvents()
+{
+	if (const TSharedPtr<FChaosVDScene> ScenePtr = SceneWeakPtr.Pin())
+	{
+		ScenePtr->OnSceneUpdated().AddRaw(this, &SChaosVDConstraintDataInspector::HandleSceneUpdated);
+		ScenePtr->OnSolverInfoActorCreated().AddRaw(this, &SChaosVDConstraintDataInspector::RegisterSelectionEventsForSolver);
+
+		for (const TPair<int32, AChaosVDSolverInfoActor*> SolverInfos : ScenePtr->GetSolverInfoActorsMap())
+		{
+			RegisterSelectionEventsForSolver(SolverInfos.Value);
+		}
+	}
+}
+
+void SChaosVDConstraintDataInspector::UnregisterSceneEvents()
+{
+	if (const TSharedPtr<FChaosVDScene> ScenePtr = SceneWeakPtr.Pin())
+	{
+		ScenePtr->OnSceneUpdated().RemoveAll(this);
+		ScenePtr->OnSolverInfoActorCreated().RemoveAll(this);
+
+		for (const TPair<int32, AChaosVDSolverInfoActor*> SolverInfos : ScenePtr->GetSolverInfoActorsMap())
+		{
+			UnregisterSelectionEventsForSolver(SolverInfos.Value);
+		}
+	}
+}
+
+void SChaosVDConstraintDataInspector::SetConstraintDataToInspect(const FChaosVDJointConstraintSelectionHandle& InDataSelectionHandle)
+{
+	ClearInspector();
+
+	if (const TSharedPtr<FChaosVDJointConstraint> QueryDataToInspect = InDataSelectionHandle.GetData().Pin())
+	{
+		CurrentDataSelectionHandle = InDataSelectionHandle;
+		const TSharedPtr<FStructOnScope> QueryDataView = MakeShared<FStructOnScope>(FChaosVDJointConstraint::StaticStruct(), reinterpret_cast<uint8*>(QueryDataToInspect.Get()));
+		ConstraintDataDetailsView->SetStructureData(QueryDataView);
+	}
+
+	bIsUpToDate = true;
+}
+
+FText SChaosVDConstraintDataInspector::GetParticleName(EChaosVDParticleSelector PairIndex) const
+{
+	const TSharedPtr<FChaosVDJointConstraint> JointConstraintPtr = CurrentDataSelectionHandle.GetData().Pin();
+	if (!JointConstraintPtr)
+	{
+		return FText::AsCultureInvariant("UnnamedParticle");
+	}
+
+	if (const TSharedPtr<FChaosVDScene> ScenePtr = SceneWeakPtr.Pin())
+	{
+		if (AChaosVDParticleActor* ParticleActor = ScenePtr->GetParticleActor(JointConstraintPtr->SolverID, PairIndex == EChaosVDParticleSelector::Index_0 ? JointConstraintPtr->ParticleParIndexes[0] : JointConstraintPtr->ParticleParIndexes[1]))
+		{
+			if (const FChaosVDParticleDataWrapper* ParticleData = ParticleActor->GetParticleData())
+			{
+				return FText::AsCultureInvariant(ParticleData->DebugName);
+			}
+		}
+	}
+
+	return FText::AsCultureInvariant("UnnamedParticle");
+}
+
+TSharedRef<SWidget> SChaosVDConstraintDataInspector::GenerateParticleSelectorButtons()
+{
+	return SNew(SUniformGridPanel)
+			.Visibility_Raw(this, &SChaosVDConstraintDataInspector::GetDetailsSectionVisibility)
+			.SlotPadding(FAppStyle::GetMargin("StandardDialog.SlotPadding"))
+			.MinDesiredSlotWidth(FAppStyle::GetFloat("StandardDialog.MinDesiredSlotWidth"))
+			.MinDesiredSlotHeight(FAppStyle::GetFloat("StandardDialog.MinDesiredSlotHeight"))
+			+SUniformGridPanel::Slot(0, 0)
+			[
+				SNew(SButton)
+				.HAlign(HAlign_Center)
+				.ContentPadding(FAppStyle::GetMargin("StandardDialog.ContentPadding"))
+				.Text(LOCTEXT("SelectParticle0", "Select Particle 0"))
+				.ToolTipText_Raw(this, &SChaosVDConstraintDataInspector::GetParticleName, EChaosVDParticleSelector::Index_0)
+				.OnClicked(this, &SChaosVDConstraintDataInspector::SelectParticleForCurrentSelectedData, EChaosVDParticleSelector::Index_0)
+			]
+			+SUniformGridPanel::Slot(1, 0)
+			[
+				SNew(SButton)
+				.HAlign(HAlign_Center)
+				.ContentPadding(FAppStyle::GetMargin("StandardDialog.ContentPadding"))
+				.Text(LOCTEXT("SelectParticle1", "Select Particle 1"))
+				.ToolTipText_Raw(this, &SChaosVDConstraintDataInspector::GetParticleName, EChaosVDParticleSelector::Index_1)
+				.OnClicked(this, &SChaosVDConstraintDataInspector::SelectParticleForCurrentSelectedData, EChaosVDParticleSelector::Index_1)
+			];
+}
+
+FReply SChaosVDConstraintDataInspector::SelectParticleForCurrentSelectedData(EChaosVDParticleSelector ParticlePairIndex)
+{
+	const TSharedPtr<FChaosVDJointConstraint> JointConstraintPtr = CurrentDataSelectionHandle.GetData().Pin();
+	if (!JointConstraintPtr)
+	{
+		return FReply::Handled();
+	}
+
+	if (const TSharedPtr<FChaosVDScene> ScenePtr = SceneWeakPtr.Pin())
+	{
+		if (AChaosVDParticleActor* ParticleActor = ScenePtr->GetParticleActor(JointConstraintPtr->SolverID, ParticlePairIndex == EChaosVDParticleSelector::Index_0 ? JointConstraintPtr->ParticleParIndexes[0] : JointConstraintPtr->ParticleParIndexes[1]))
+		{
+			ScenePtr->SetSelectedObject(ParticleActor);
+		}
+	}
+	
+	return FReply::Handled();
+}
+
+EVisibility SChaosVDConstraintDataInspector::GetOutOfDateWarningVisibility() const
+{
+	return !bIsUpToDate && CurrentDataSelectionHandle.GetData().IsValid() ? EVisibility::Visible : EVisibility::Collapsed;
+}
+
+EVisibility SChaosVDConstraintDataInspector::GetDetailsSectionVisibility() const
+{
+	return CurrentDataSelectionHandle.GetData().IsValid() ? EVisibility::Visible : EVisibility::Collapsed;
+}
+
+EVisibility SChaosVDConstraintDataInspector::GetNothingSelectedMessageVisibility() const
+{
+	return !CurrentDataSelectionHandle.GetData().IsValid() ? EVisibility::Visible : EVisibility::Collapsed;
+}
+
+void SChaosVDConstraintDataInspector::HandleSceneUpdated()
+{
+	// TODO: Disabling the "out of date" message system, it is a little bit confusing with joint constraints as the debug draw position will be updated between frames (because the constraint only has particles IDs)
+	// but the data itself is not updated.
+	// We are clearing out the selection altogether instead for now
+	//bIsUpToDate = false;
+	
+	CurrentDataSelectionHandle.SetIsSelected(false);
+	CurrentDataSelectionHandle = FChaosVDJointConstraintSelectionHandle();
+
+	// TODO: To Keep a selection up to date we need a persistent ID for the constraint
+	// We could hash the pointer for that or add an ID to the Constraint Handle only compiled in when CVD is enabled
+}
+
+void SChaosVDConstraintDataInspector::ClearInspector()
+{
+	ConstraintDataDetailsView->SetStructureData(nullptr);
+	CurrentDataSelectionHandle = FChaosVDJointConstraintSelectionHandle();
+}
+
+TSharedPtr<IStructureDetailsView> SChaosVDConstraintDataInspector::CreateDataDetailsView()
+{
+	FPropertyEditorModule& PropertyEditorModule = FModuleManager::GetModuleChecked<FPropertyEditorModule>("PropertyEditor");
+
+	const FStructureDetailsViewArgs StructDetailsViewArgs;
+	FDetailsViewArgs DetailsViewArgs;
+	DetailsViewArgs.bShowOptions = false;
+	DetailsViewArgs.bAllowFavoriteSystem = false;
+	DetailsViewArgs.bAllowSearch = true;
+	DetailsViewArgs.bShowScrollBar = false;
+
+	return PropertyEditorModule.CreateStructureDetailView(DetailsViewArgs, StructDetailsViewArgs, nullptr);
+}
+
+#undef LOCTEXT_NAMESPACE
diff --git a/Engine/Plugins/Experimental/ChaosVD/Source/ChaosVD/Private/Widgets/SChaosVDMainTab.cpp b/Engine/Plugins/Experimental/ChaosVD/Source/ChaosVD/Private/Widgets/SChaosVDMainTab.cpp
--- a/Engine/Plugins/Experimental/ChaosVD/Source/ChaosVD/Private/Widgets/SChaosVDMainTab.cpp
+++ b/Engine/Plugins/Experimental/ChaosVD/Source/ChaosVD/Private/Widgets/SChaosVDMainTab.cpp
@@ -13,6 +13,7 @@
 #include "ChaosVDScene.h"
 #include "ChaosVDSolversTracksTab.h"
 #include "ChaosVDCollisionDataDetailsTab.h"
+#include "ChaosVDConstraintDataInspectorTab.h"
 #include "ChaosVDSceneQueryDataInspectorTab.h"
 #include "ChaosVDStyle.h"
 #include "ChaosVDTabsIDs.h"
@@ -29,9 +30,11 @@
 #include "StatusBarSubsystem.h"
 #include "Components/ChaosVDParticleDataComponent.h"
 #include "Components/ChaosVDSolverCollisionDataComponent.h"
+#include "Components/ChaosVDSolverJointConstraintDataComponent.h"
 #include "Styling/StyleColors.h"
 #include "Styling/ToolBarStyle.h"
 #include "Trace/ChaosVDTraceManager.h"
+#include "Visualizers/ChaosVDJointConstraintsDataComponentVisualizer.h"
 #include "Visualizers/ChaosVDParticleDataComponentVisualizer.h"
 #include "Visualizers/ChaosVDSceneQueryDataComponentVisualizer.h"
 #include "Visualizers/ChaosVDSolverCollisionDataComponentVisualizer.h"
@@ -56,6 +59,7 @@ void SChaosVDMainTab::Construct(const FArguments& InArgs, TSharedPtr<FChaosVDEng
 	RegisterComponentVisualizer(UChaosVDSolverCollisionDataComponent::StaticClass()->GetFName(), MakeShared<FChaosVDSolverCollisionDataComponentVisualizer>());
 	RegisterComponentVisualizer(UChaosVDSceneQueryDataComponent::StaticClass()->GetFName(), MakeShared<FChaosVDSceneQueryDataComponentVisualizer>());
 	RegisterComponentVisualizer(UChaosVDParticleDataComponent::StaticClass()->GetFName(), MakeShared<FChaosVDParticleDataComponentVisualizer>());
+	RegisterComponentVisualizer(UChaosVDSolverJointConstraintDataComponent::StaticClass()->GetFName(), MakeShared<FChaosVDJointConstraintsDataComponentVisualizer>());
 
 	TabManager = FGlobalTabmanager::Get()->NewTabManager(InArgs._OwnerTab.ToSharedRef()).ToSharedPtr();
 
@@ -67,6 +71,7 @@ void SChaosVDMainTab::Construct(const FArguments& InArgs, TSharedPtr<FChaosVDEng
 	RegisterTabSpawner<FChaosVDEditorVisualizationSettingsTab>(FChaosVDTabID::CVDEditorSettings);
 	RegisterTabSpawner<FChaosVDCollisionDataDetailsTab>(FChaosVDTabID::CollisionDataDetails);
 	RegisterTabSpawner<FChaosVDSceneQueryDataInspectorTab>(FChaosVDTabID::SceneQueryDataDetails);
+	RegisterTabSpawner<FChaosVDConstraintDataInspectorTab>(FChaosVDTabID::JointsDataDetails);
 
 	StatusBarID = FName(FChaosVDTabID::StatusBar.ToString() + InChaosVDEngine->GetInstanceGuid().ToString());
 	
@@ -390,6 +395,7 @@ TSharedRef<FTabManager::FLayout> SChaosVDMainTab::GenerateMainLayout()
 					->AddTab(FChaosVDTabID::DetailsPanel, ETabState::OpenedTab)
 					->AddTab(FChaosVDTabID::CollisionDataDetails, ETabState::OpenedTab)
 					->AddTab(FChaosVDTabID::SceneQueryDataDetails, ETabState::OpenedTab)
+					->AddTab(FChaosVDTabID::JointsDataDetails, ETabState::ClosedTab)
 					->AddTab(FChaosVDTabID::CVDEditorSettings, ETabState::ClosedTab)
 				)
 			)
diff --git a/Engine/Plugins/Experimental/ChaosVD/Source/ChaosVD/Public/Actors/ChaosVDSolverInfoActor.h b/Engine/Plugins/Experimental/ChaosVD/Source/ChaosVD/Public/Actors/ChaosVDSolverInfoActor.h
--- a/Engine/Plugins/Experimental/ChaosVD/Source/ChaosVD/Public/Actors/ChaosVDSolverInfoActor.h
+++ b/Engine/Plugins/Experimental/ChaosVD/Source/ChaosVD/Public/Actors/ChaosVDSolverInfoActor.h
@@ -8,6 +8,7 @@
 #include "GameFramework/Actor.h"
 #include "ChaosVDSolverInfoActor.generated.h"
 
+class UChaosVDSolverJointConstraintDataComponent;
 struct FChaosVDParticleDataWrapper;
 class AChaosVDParticleActor;
 class UChaosVDSolverCollisionDataComponent;
@@ -40,6 +41,7 @@ public:
 
 	UChaosVDSolverCollisionDataComponent* GetCollisionDataComponent() { return CollisionDataComponent; }
 	UChaosVDParticleDataComponent* GetParticleDataComponent() { return ParticleDataComponent; }
+	UChaosVDSolverJointConstraintDataComponent* GetJointsDataComponent() { return JointsDataComponent; }
 
 	void RegisterParticleActor(int32 ParticleID, AChaosVDParticleActor* ParticleActor);
 
@@ -95,6 +97,9 @@ protected:
 
 	UPROPERTY()
 	TObjectPtr<UChaosVDParticleDataComponent> ParticleDataComponent;
+
+	UPROPERTY()
+	TObjectPtr<UChaosVDSolverJointConstraintDataComponent> JointsDataComponent;
 };
 
 template <typename TCallback>
diff --git a/Engine/Plugins/Experimental/ChaosVD/Source/ChaosVD/Public/ChaosVDConstraintDataInspectorTab.h b/Engine/Plugins/Experimental/ChaosVD/Source/ChaosVD/Public/ChaosVDConstraintDataInspectorTab.h
new file mode 100644
--- /dev/null
+++ b/Engine/Plugins/Experimental/ChaosVD/Source/ChaosVD/Public/ChaosVDConstraintDataInspectorTab.h
@@ -0,0 +1,28 @@
+﻿// Copyright Epic Games, Inc. All Rights Reserved.
+
+#pragma once
+
+#include "ChaosVDTabSpawnerBase.h"
+#include "Templates/SharedPointer.h"
+
+class SChaosVDConstraintDataInspector;
+class SChaosVDSceneQueryDataInspector;
+
+/** Spawns and handles and instance for the visual debugger Constraints Inspector tab */
+class FChaosVDConstraintDataInspectorTab final : public FChaosVDTabSpawnerBase
+{
+public:
+	FChaosVDConstraintDataInspectorTab(const FName& InTabID, const TSharedPtr<FTabManager>& InTabManager, const TWeakPtr<SChaosVDMainTab>& InOwningTabWidget)
+		: FChaosVDTabSpawnerBase(InTabID, InTabManager, InOwningTabWidget)
+	{
+	}
+
+	virtual ~FChaosVDConstraintDataInspectorTab() override;
+	virtual TSharedRef<SDockTab> HandleTabSpawnRequest(const FSpawnTabArgs& Args) override;
+	virtual void HandleTabClosed(TSharedRef<SDockTab> InTabClosed) override;
+
+	TWeakPtr<SChaosVDConstraintDataInspector> GetConstraintDataInspectorInstance() const { return ConstraintDataInspector; }
+
+protected:
+	TSharedPtr<SChaosVDConstraintDataInspector> ConstraintDataInspector;
+};
diff --git a/Engine/Plugins/Experimental/ChaosVD/Source/ChaosVD/Public/ChaosVDEditorSettings.h b/Engine/Plugins/Experimental/ChaosVD/Source/ChaosVD/Public/ChaosVDEditorSettings.h
--- a/Engine/Plugins/Experimental/ChaosVD/Source/ChaosVD/Public/ChaosVDEditorSettings.h
+++ b/Engine/Plugins/Experimental/ChaosVD/Source/ChaosVD/Public/ChaosVDEditorSettings.h
@@ -7,6 +7,7 @@
 #include "Engine/StaticMesh.h"
 #include "UObject/Object.h"
 #include "UObject/SoftObjectPtr.h"
+#include "Visualizers/ChaosVDJointConstraintsDataComponentVisualizer.h"
 #include "Visualizers/ChaosVDParticleDataComponentVisualizer.h"
 #include "Visualizers/ChaosVDSceneQueryDataComponentVisualizer.h"
 #include "Visualizers/ChaosVDSolverCollisionDataComponentVisualizer.h"
@@ -185,11 +186,45 @@ struct FChaosParticleDataDebugDrawSettings
 	float CenterOfMassRadius = 10.0f;
 
 	float GetScaleFortDataID(EChaosVDParticleDataVisualizationFlags DataID) const;
-	
+
 	UPROPERTY(EditAnywhere, Category=DebugDraw)
 	FChaosParticleDataDebugDrawColors ColorSettings;
 };
 
+USTRUCT()
+struct FChaosVDJointsDebugDrawSettings
+{
+	GENERATED_BODY()
+
+	/** The depth priority used for while drawing data. Can be World or Foreground (with this one the shapes will be drawn on top of the geometry and be always visible) */
+	UPROPERTY(EditAnywhere, Category=DebugDraw)
+	TEnumAsByte<ESceneDepthPriorityGroup> DepthPriority = ESceneDepthPriorityGroup::SDPG_Foreground;
+
+	/** Scale to apply to the Linear Impulse vector before draw it. */
+	UPROPERTY(EditAnywhere, Category=DebugDraw)
+	float LinearImpulseScale = 0.001;
+
+	/** Scale to apply to the Angular Impulse vector before draw it. */
+	UPROPERTY(EditAnywhere, Category=DebugDraw)
+	float AngularImpulseScale = 0.1f;
+
+	/** Scale to apply to anything that does not have a dedicated scale setting before draw it. */
+	UPROPERTY(EditAnywhere, Category=DebugDraw)
+	float GeneralScale = 1.0f;
+
+	/** Line thickness to use as a base to calculate the different line thickness values used to debug draw the data. */
+	UPROPERTY(EditAnywhere, Category=DebugDraw)
+	float BaseLineThickness = 2.0f;
+
+	/** Size of the debug drawn Center Of Mass. */
+	UPROPERTY(EditAnywhere, Category=DebugDraw)
+	float CenterOfMassSize = 1.0f;
+
+	/** Size of the debug drawn if the Constraint Axis */
+	UPROPERTY(EditAnywhere, Category=DebugDraw)
+    float ConstraintAxisLength = 10.0f;
+};
+
 /** Structure holding the settings using to debug draw Particles shape based on their shape type on the Chaos Visual Debugger */
 USTRUCT()
 struct FChaosDebugDrawColorsByShapeType
@@ -295,6 +330,10 @@ public:
 	UPROPERTY(EditAnywhere, Category = "Viewport Visualization Flags", meta = (Bitmask, BitmaskEnum = "/Script/ChaosVD.EChaosVDSceneQueryVisualizationFlags"))
 	uint32 GlobalSceneQueriesVisualizationFlags = static_cast<uint32>(EChaosVDSceneQueryVisualizationFlags::DrawClientQueries | EChaosVDSceneQueryVisualizationFlags::DrawServerQueries |  EChaosVDSceneQueryVisualizationFlags::DrawHits | EChaosVDSceneQueryVisualizationFlags::DrawLineTraceQueries);
 
+	/** Set of flags to enable/disable visualization of specific scene queries data as debug draw */
+	UPROPERTY(EditAnywhere, Category = "Viewport Visualization Flags", meta = (Bitmask, BitmaskEnum = "/Script/ChaosVD.EChaosVDJointsDataVisualizationFlags"))
+	uint32 GlobalJointsDataVisualizationFlags = static_cast<uint32>(EChaosVDJointsDataVisualizationFlags::EnableClientDraw | EChaosVDJointsDataVisualizationFlags::EnableServerDraw | EChaosVDJointsDataVisualizationFlags::ActorConnector | EChaosVDJointsDataVisualizationFlags::DrawKinematic);
+
 	/** If true, text information (if available) will be drawn alongside any other debug draw shape */
 	UPROPERTY(EditAnywhere, Category = "Viewport Visualization")
 	bool bShowDebugText = false;
@@ -305,6 +344,9 @@ public:
 	UPROPERTY(EditAnywhere, Category = "Viewport Visualization",  meta=(EditCondition = "GlobalParticleDataVisualizationFlags != 0", EditConditionHides))
 	FChaosParticleDataDebugDrawSettings ParticleDataDebugDrawSettings;
 
+	UPROPERTY(EditAnywhere, Category = "Viewport Visualization",  meta=(EditCondition = "GlobalParticleDataVisualizationFlags != 0", EditConditionHides))
+	FChaosVDJointsDebugDrawSettings JointsDataDebugDrawSettings;
+
 	UPROPERTY(EditAnywhere, Category = "Viewport Visualization")
 	EChaosVDParticleDebugColorMode ParticleColorMode;
 	
diff --git a/Engine/Plugins/Experimental/ChaosVD/Source/ChaosVD/Public/ChaosVDScene.h b/Engine/Plugins/Experimental/ChaosVD/Source/ChaosVD/Public/ChaosVDScene.h
--- a/Engine/Plugins/Experimental/ChaosVD/Source/ChaosVD/Public/ChaosVDScene.h
+++ b/Engine/Plugins/Experimental/ChaosVD/Source/ChaosVD/Public/ChaosVDScene.h
@@ -30,6 +30,7 @@ typedef TMap<int32, AChaosVDSolverInfoActor*> FChaosVDSolverInfoByIDMap;
 DECLARE_MULTICAST_DELEGATE(FChaosVDSceneUpdatedDelegate)
 DECLARE_MULTICAST_DELEGATE_OneParam(FChaosVDActorActiveStateUpdateDelegate, AChaosVDParticleActor*)
 DECLARE_MULTICAST_DELEGATE_OneParam(FChaosVDOnObjectSelectedDelegate, UObject*)
+DECLARE_MULTICAST_DELEGATE_OneParam(FChaosVDSolverInfoActorCreatedDelegate, AChaosVDSolverInfoActor*)
 
 /** Recreates a UWorld from a recorded Chaos VD Frame */
 class FChaosVDScene : public FGCObject , public TSharedFromThis<FChaosVDScene>
@@ -54,6 +55,7 @@ public:
 	void UpdateFromRecordedStepData(const int32 SolverID, const FChaosVDStepData& InRecordedStepData, const FChaosVDSolverFrameData& InFrameData);
 
 	void UpdateParticlesCollisionData(const FChaosVDStepData& InRecordedStepData, int32 SolverID);
+	void UpdateJointConstraintsData(const FChaosVDStepData& InRecordedStepData, int32 SolverID);
 
 	// No need to deprecate the old version since it is not a public API nor inline 
 	void HandleNewGeometryData(const Chaos::FConstImplicitObjectPtr& Geometry, const uint32 GeometryID) const;
@@ -111,6 +113,8 @@ public:
 
 	FChaosVDActorActiveStateUpdateDelegate& OnActorActiveStateChanged() { return ParticleActorUpdateDelegate; }
 
+	FChaosVDSolverInfoActorCreatedDelegate& OnSolverInfoActorCreated() { return SolverInfoActorCreatedDelegate; }
+
 	TSharedPtr<FChaosVDRecording> LoadedRecording;
 
 private:
@@ -189,4 +193,6 @@ private:
 	FChaosVDActorActiveStateUpdateDelegate ParticleActorUpdateDelegate;
 
 	FDelegateHandle ActorDestroyedHandle;
+
+	FChaosVDSolverInfoActorCreatedDelegate SolverInfoActorCreatedDelegate;
 };
diff --git a/Engine/Plugins/Experimental/ChaosVD/Source/ChaosVD/Public/ChaosVDTabsIDs.h b/Engine/Plugins/Experimental/ChaosVD/Source/ChaosVD/Public/ChaosVDTabsIDs.h
--- a/Engine/Plugins/Experimental/ChaosVD/Source/ChaosVD/Public/ChaosVDTabsIDs.h
+++ b/Engine/Plugins/Experimental/ChaosVD/Source/ChaosVD/Public/ChaosVDTabsIDs.h
@@ -18,4 +18,5 @@ public:
 	static const FName StatusBar;
 	static const FName CollisionDataDetails;
 	static const FName SceneQueryDataDetails;
+	static const FName JointsDataDetails;
 };
diff --git a/Engine/Plugins/Experimental/ChaosVD/Source/ChaosVD/Public/Components/ChaosVDConstraintDataHelpers.h b/Engine/Plugins/Experimental/ChaosVD/Source/ChaosVD/Public/Components/ChaosVDConstraintDataHelpers.h
new file mode 100644
--- /dev/null
+++ b/Engine/Plugins/Experimental/ChaosVD/Source/ChaosVD/Public/Components/ChaosVDConstraintDataHelpers.h
@@ -0,0 +1,58 @@
+﻿// Copyright Epic Games, Inc. All Rights Reserved.
+
+#pragma once
+
+#include "Containers/Array.h"
+#include "Containers/Map.h"
+
+class FChaosVDCollisionDataVisualizer;
+class FChaosVDScene;
+class FPrimitiveDrawInterface;
+
+enum class EChaosVDParticlePairSlot : uint8
+{
+	Primary,
+	Secondary,
+	Any
+};
+
+namespace Chaos::VisualDebugger::Utils
+{
+	template <typename MapType, typename TData>
+	void AddDataDataToParticleIDMap(MapType& MapToUpdate, const TData& MidPhaseData, int32 ParticleID)
+	{
+		if (TArray<TData>* ParticleCollisionData = MapToUpdate.Find(ParticleID))
+		{
+			ParticleCollisionData->Add(MidPhaseData);
+		}
+		else
+		{
+			MapToUpdate.Add(ParticleID, { MidPhaseData });
+		}
+	}
+
+	template <typename MapType, typename TData>
+	const TArray<TData>* GetDataFromParticlePairMaps(const MapType& Map0ToQuery, const MapType& Map1ToQuery, int32 ParticleID, EChaosVDParticlePairSlot Options)
+	{
+		switch (Options)
+		{
+		case EChaosVDParticlePairSlot::Primary:
+			{
+				return Map0ToQuery.Find(ParticleID);
+				break;
+			}
+		case EChaosVDParticlePairSlot::Secondary:
+			{
+				return Map1ToQuery.Find(ParticleID);
+				break;
+			}
+		case EChaosVDParticlePairSlot::Any:
+			{
+				const TArray<TData>* FoundMidPhasesContainer = Map0ToQuery.Find(ParticleID);
+				return FoundMidPhasesContainer ? FoundMidPhasesContainer : Map1ToQuery.Find(ParticleID);
+				break;
+			}
+		}
+		return nullptr;
+	}
+}
diff --git a/Engine/Plugins/Experimental/ChaosVD/Source/ChaosVD/Public/Components/ChaosVDSolverCollisionDataComponent.h b/Engine/Plugins/Experimental/ChaosVD/Source/ChaosVD/Public/Components/ChaosVDSolverCollisionDataComponent.h
--- a/Engine/Plugins/Experimental/ChaosVD/Source/ChaosVD/Public/Components/ChaosVDSolverCollisionDataComponent.h
+++ b/Engine/Plugins/Experimental/ChaosVD/Source/ChaosVD/Public/Components/ChaosVDSolverCollisionDataComponent.h
@@ -2,6 +2,7 @@
 
 #pragma once
 
+#include "ChaosVDConstraintDataHelpers.h"
 #include "Components/ActorComponent.h"
 
 #include "Containers/Array.h"
@@ -12,23 +13,11 @@
 
 #include "ChaosVDSolverCollisionDataComponent.generated.h"
 
-class FChaosVDCollisionDataVisualizer;
-class FChaosVDScene;
-class FPrimitiveDrawInterface;
-class FSceneView;
-
 struct FChaosVDConstraint;
 
 typedef TMap<int32, TArray<TSharedPtr<FChaosVDParticlePairMidPhase>>> FChaosVDMidPhaseByParticleMap;
 typedef TMap<int32, TArray<FChaosVDConstraint*>> FChaosVDConstraintByParticleMap;
 
-enum class EChaosVDCollisionParticlePairSlot
-{
-	Primary,
-	Secondary,
-	Any
-};
-
 UCLASS()
 class UChaosVDSolverCollisionDataComponent : public UActorComponent
 {
@@ -40,21 +29,13 @@ public:
 	void UpdateCollisionData(const TArray<TSharedPtr<FChaosVDParticlePairMidPhase>>& InMidPhaseData);
 
 	const TArray<TSharedPtr<FChaosVDParticlePairMidPhase>>& GetMidPhases() const { return AllMidPhases; }
-	const TArray<TSharedPtr<FChaosVDParticlePairMidPhase>>* GetMidPhasesForParticle(int32 ParticleID, EChaosVDCollisionParticlePairSlot Options) const;
-	const TArray<FChaosVDConstraint*>* GetConstraintsForParticle(int32 ParticleID, EChaosVDCollisionParticlePairSlot Options) const;
-
-	void DrawVisualization(const FSceneView* View, FPrimitiveDrawInterface* PDI);
+	const TArray<TSharedPtr<FChaosVDParticlePairMidPhase>>* GetMidPhasesForParticle(int32 ParticleID, EChaosVDParticlePairSlot Options) const;
+	const TArray<FChaosVDConstraint*>* GetConstraintsForParticle(int32 ParticleID, EChaosVDParticlePairSlot Options) const;
 
 protected:
 
 	void ClearCollisionData();
 
-	template<typename MapType, typename CollisionDataType>
-	void AddCollisionDataToParticleIDMap(MapType& MapToUpdate, const CollisionDataType& MidPhaseData, int32 ParticleID);
-
-	template<typename MapType, typename CollisionDataType>
-	const TArray<CollisionDataType>* GetCollisionDataFromMap(const MapType& Map0ToQuery, const MapType& Map1ToQuery, int32 ParticleID, EChaosVDCollisionParticlePairSlot Options) const;
-
 	TArray<TSharedPtr<FChaosVDParticlePairMidPhase>> AllMidPhases;
 	FChaosVDMidPhaseByParticleMap MidPhasesByParticleID0;
 	FChaosVDMidPhaseByParticleMap MidPhasesByParticleID1;
@@ -62,42 +43,4 @@ protected:
 	FChaosVDConstraintByParticleMap ConstraintsByParticleID0;
 	FChaosVDConstraintByParticleMap ConstraintsByParticleID1;
 
-};
-
-template <typename MapType, typename CollisionDataType>
-void UChaosVDSolverCollisionDataComponent::AddCollisionDataToParticleIDMap(MapType& MapToUpdate, const CollisionDataType& MidPhaseData, int32 ParticleID)
-{
-	if (TArray<CollisionDataType>* ParticleCollisionData = MapToUpdate.Find(ParticleID))
-	{
-		ParticleCollisionData->Add(MidPhaseData);
-	}
-	else
-	{
-		MapToUpdate.Add(ParticleID, { MidPhaseData });
-	}
-}
-
-template <typename MapType, typename CollisionDataType>
-const TArray<CollisionDataType>* UChaosVDSolverCollisionDataComponent::GetCollisionDataFromMap(const MapType& Map0ToQuery, const MapType& Map1ToQuery, int32 ParticleID, EChaosVDCollisionParticlePairSlot Options) const
-{
-	switch (Options)
-	{
-	case EChaosVDCollisionParticlePairSlot::Primary:
-		{
-			return Map0ToQuery.Find(ParticleID);
-			break;
-		}
-	case EChaosVDCollisionParticlePairSlot::Secondary:
-		{
-			return Map1ToQuery.Find(ParticleID);
-			break;
-		}
-	case EChaosVDCollisionParticlePairSlot::Any:
-		{
-			const TArray<CollisionDataType>* FoundMidPhasesContainer = Map0ToQuery.Find(ParticleID);
-			return FoundMidPhasesContainer ? FoundMidPhasesContainer : Map1ToQuery.Find(ParticleID);
-			break;
-		}
-	}
-	return nullptr;
-}
+};
\ No newline at end of file
diff --git a/Engine/Plugins/Experimental/ChaosVD/Source/ChaosVD/Public/Components/ChaosVDSolverJointConstraintDataComponent.h b/Engine/Plugins/Experimental/ChaosVD/Source/ChaosVD/Public/Components/ChaosVDSolverJointConstraintDataComponent.h
new file mode 100644
--- /dev/null
+++ b/Engine/Plugins/Experimental/ChaosVD/Source/ChaosVD/Public/Components/ChaosVDSolverJointConstraintDataComponent.h
@@ -0,0 +1,80 @@
+﻿// Copyright Epic Games, Inc. All Rights Reserved.
+
+#pragma once
+
+#include "Components/ActorComponent.h"
+#include "ChaosVDSolverJointConstraintDataComponent.generated.h"
+
+struct FChaosVDJointConstraint;
+
+typedef TMap<int32, TArray<TSharedPtr<FChaosVDJointConstraint>>> FChaosVDJointDataByParticleMap;
+typedef TMap<int32, TSharedPtr<FChaosVDJointConstraint>> FChaosVDJointDataByConstraintIndexMap;
+typedef TArray<TSharedPtr<FChaosVDJointConstraint>> FChaosVDJointDataArray;
+
+enum class EChaosVDParticlePairSlot : uint8;
+
+/** Struct used to pass data about a specific joint constraint to other objects */
+struct FChaosVDJointConstraintSelectionHandle
+{
+	FChaosVDJointConstraintSelectionHandle()
+	{		
+	}
+
+	FChaosVDJointConstraintSelectionHandle(const TWeakPtr<FChaosVDJointConstraint>& InJointData)
+		: JointData(InJointData)
+	{
+	}
+
+	void SetIsSelected(bool bNewSelected);
+	bool IsSelected() const;
+
+	bool IsValid();
+
+	TWeakPtr<FChaosVDJointConstraint> GetData() const { return JointData; } 
+
+private:
+	TWeakPtr<FChaosVDJointConstraint> JointData = nullptr;
+};
+
+DECLARE_MULTICAST_DELEGATE_OneParam(FChaosVDJointSlectionChangedDelegate, const FChaosVDJointConstraintSelectionHandle& SelectionHandle)
+
+UCLASS()
+class CHAOSVD_API UChaosVDSolverJointConstraintDataComponent : public UActorComponent
+{
+	GENERATED_BODY()
+
+public:
+
+	UChaosVDSolverJointConstraintDataComponent();
+
+	void UpdateConstraintData(const TArray<TSharedPtr<FChaosVDJointConstraint>>& InJointData);
+	
+	const FChaosVDJointDataArray& GetAllJointConstraints() const { return AllJointConstraints; }
+	const FChaosVDJointDataArray* GetJointConstraintsForParticle(int32 ParticleID, EChaosVDParticlePairSlot Options) const;
+
+	void SelectJoint(int32 ConstraintIndex);
+	void SelectJoint(const FChaosVDJointConstraintSelectionHandle& SelectionHandle);
+
+	bool IsJointSelected(int32 ConstraintIndex) const;
+
+	FChaosVDJointSlectionChangedDelegate& OnSelectionChanged() { return SelectionChangeDelegate; }
+
+	const FChaosVDJointConstraintSelectionHandle& GetCurrentSelectionHandle() const { return CurrentJointSelectionHandle; }
+
+protected:
+
+	TSharedPtr<FChaosVDJointConstraint> GetJointConstraintByIndex(int32 ConstraintIndex);
+	
+	void ClearData();
+
+	FChaosVDJointDataArray AllJointConstraints;
+	
+	FChaosVDJointDataByParticleMap JointConstraintByParticle0;
+	FChaosVDJointDataByParticleMap JointConstraintByParticle1;
+
+	FChaosVDJointDataByConstraintIndexMap JointConstraintByConstraintIndex;
+
+	FChaosVDJointSlectionChangedDelegate SelectionChangeDelegate;
+
+	FChaosVDJointConstraintSelectionHandle CurrentJointSelectionHandle;
+};
diff --git a/Engine/Plugins/Experimental/ChaosVD/Source/ChaosVD/Public/Trace/DataProcessors/ChaosVDJointConstraintDataProcessor.h b/Engine/Plugins/Experimental/ChaosVD/Source/ChaosVD/Public/Trace/DataProcessors/ChaosVDJointConstraintDataProcessor.h
new file mode 100644
--- /dev/null
+++ b/Engine/Plugins/Experimental/ChaosVD/Source/ChaosVD/Public/Trace/DataProcessors/ChaosVDJointConstraintDataProcessor.h
@@ -0,0 +1,22 @@
+﻿// Copyright Epic Games, Inc. All Rights Reserved.
+
+#pragma once
+
+#pragma once
+#include "IChaosVDDataProcessor.h"
+
+struct FChaosVDSolverFrameData;
+struct FChaosVDConstraint;
+
+/**
+ * Data processor implementation that is able to deserialize traced joint constraints
+ */
+class FChaosVDJointConstraintDataProcessor final : public IChaosVDDataProcessor
+{
+public:
+	explicit FChaosVDJointConstraintDataProcessor();
+	
+	virtual bool ProcessRawData(const TArray<uint8>& InData) override;
+};
+
+
diff --git a/Engine/Plugins/Experimental/ChaosVD/Source/ChaosVD/Public/Visualizers/ChaosVDDebugDrawUtils.h b/Engine/Plugins/Experimental/ChaosVD/Source/ChaosVD/Public/Visualizers/ChaosVDDebugDrawUtils.h
--- a/Engine/Plugins/Experimental/ChaosVD/Source/ChaosVD/Public/Visualizers/ChaosVDDebugDrawUtils.h
+++ b/Engine/Plugins/Experimental/ChaosVD/Source/ChaosVD/Public/Visualizers/ChaosVDDebugDrawUtils.h
@@ -34,3 +34,8 @@ private:
 
 	static TQueue<FChaosVDQueuedTextToDraw> TexToDrawQueue;
 };
+
+namespace Chaos::VisualDebugger::Utils
+{
+	FString GenerateDebugTextForVector(const FVector& InVector, const FString& VectorName, const FString& InVectorUnits);
+}
diff --git a/Engine/Plugins/Experimental/ChaosVD/Source/ChaosVD/Public/Visualizers/ChaosVDJointConstraintsDataComponentVisualizer.h b/Engine/Plugins/Experimental/ChaosVD/Source/ChaosVD/Public/Visualizers/ChaosVDJointConstraintsDataComponentVisualizer.h
new file mode 100644
--- /dev/null
+++ b/Engine/Plugins/Experimental/ChaosVD/Source/ChaosVD/Public/Visualizers/ChaosVDJointConstraintsDataComponentVisualizer.h
@@ -0,0 +1,93 @@
+﻿// Copyright Epic Games, Inc. All Rights Reserved.
+
+#pragma once
+#include "ComponentVisualizer.h"
+#include "IChaosVDParticleVisualizationDataProvider.h"
+#include "Chaos/Core.h"
+#include "Components/ChaosVDSolverJointConstraintDataComponent.h"
+
+struct FChaosVDJointsDebugDrawSettings;
+class AChaosVDSolverInfoActor;
+struct FChaosVDJointConstraint;
+
+UENUM(meta = (Bitflags, UseEnumValuesAsMaskValuesInEditor = "true"))
+enum class EChaosVDJointsDataVisualizationFlags : uint32
+{
+	None					= 0 UMETA(Hidden),
+	/** Draw the PushOut vector based on the constraint's data */
+	PushOut					= 1 << 0,
+	/** Draw the Angular Impulse vector based on the constraint's data */
+	AngularImpulse			= 1 << 1 UMETA(Hidden),
+	ActorConnector			= 1 << 2,
+	CenterOfMassConnector	= 1 << 3,
+	Stretch					= 1 << 4,
+	Axes					= 1 << 5,
+	/** Draw the joint even if one of the particles or both are kinematic */
+	DrawKinematic			= 1 << 6,
+	/** Draw the joint even if it is disabled */
+	DrawDisabled			= 1 << 7,
+	/** Enables Debug draw for Joint Constraint data from Game Client solvers */
+	EnableClientDraw		= 1 << 8,
+	/** Enables Debug draw for Joint Constraint data from Game Server solvers */
+	EnableServerDraw		= 1 << 9,
+	/** Only debugs draw data for a selected joint constraint */
+	OnlyDrawSelected		= 1 << 10,
+	/** Enables Debug draw for Joint Constraint data from any solver */
+	EnableDraw				= EnableClientDraw | EnableServerDraw,
+};
+ENUM_CLASS_FLAGS(EChaosVDJointsDataVisualizationFlags);
+
+/** Visualization context structure specific for Joints visualizations */
+struct FChaosVDJointVisualizationDataContext : public FChaosVDVisualizationContext
+{
+	FChaosVDJointConstraintSelectionHandle DataSelectionHandle = FChaosVDJointConstraintSelectionHandle(nullptr);
+	bool bIsServerVisualizationEnabled = false;
+	
+	AChaosVDSolverInfoActor* SolverInfoActor = nullptr;
+
+	const FChaosVDJointsDebugDrawSettings* DebugDrawSettings = nullptr;
+
+	bool bShowDebugText = false;
+
+	bool IsVisualizationFlagEnabled(EChaosVDJointsDataVisualizationFlags Flag) const
+	{
+		const EChaosVDJointsDataVisualizationFlags FlagsAsParticleFlags = static_cast<EChaosVDJointsDataVisualizationFlags>(VisualizationFlags);
+		return EnumHasAnyFlags(FlagsAsParticleFlags, Flag);
+	}
+};
+
+/** Custom Hit Proxy for debug drawn scene queries */
+struct HChaosVDJointConstraintProxy : public HComponentVisProxy
+{
+	DECLARE_HIT_PROXY()
+	
+	HChaosVDJointConstraintProxy(const UActorComponent* Component, const FChaosVDJointConstraintSelectionHandle& InContactFinderData) : HComponentVisProxy(Component, HPP_UI), DataSelectionHandle(InContactFinderData)
+	{	
+	}
+
+	virtual EMouseCursor::Type GetMouseCursor() override
+	{
+		return EMouseCursor::Crosshairs;
+	}
+
+	FChaosVDJointConstraintSelectionHandle DataSelectionHandle;
+};
+
+/**
+ * Component visualizer in charge of generating debug draw visualizations for Joint Constraints in a UChaosVDSolverJointConstraintDataComponent
+ */
+class FChaosVDJointConstraintsDataComponentVisualizer final : public FComponentVisualizer
+{
+public:
+	FChaosVDJointConstraintsDataComponentVisualizer()
+	{
+	}
+
+	virtual void DrawVisualization(const UActorComponent* Component, const FSceneView* View, FPrimitiveDrawInterface* PDI) override;
+	virtual bool VisProxyHandleClick(FEditorViewportClient* InViewportClient, HComponentVisProxy* VisProxy, const FViewportClick& Click) override;
+
+protected:
+
+	void DebugDrawAllAxis(const FChaosVDJointConstraint& InJointConstraintData, FChaosVDJointVisualizationDataContext& VisualizationContext, FPrimitiveDrawInterface* PDI, const float LineThickness, const FVector& InPosition, const Chaos::FMatrix33& InRotationMatrix, TConstArrayView<FLinearColor>);
+	void DrawJointConstraint(const UActorComponent* Component, const FChaosVDJointConstraint& InJointConstraintData, FChaosVDJointVisualizationDataContext& VisualizationContext, const FSceneView* View, FPrimitiveDrawInterface* PDI);
+};
diff --git a/Engine/Plugins/Experimental/ChaosVD/Source/ChaosVD/Public/Visualizers/ChaosVDParticleDataComponentVisualizer.h b/Engine/Plugins/Experimental/ChaosVD/Source/ChaosVD/Public/Visualizers/ChaosVDParticleDataComponentVisualizer.h
--- a/Engine/Plugins/Experimental/ChaosVD/Source/ChaosVD/Public/Visualizers/ChaosVDParticleDataComponentVisualizer.h
+++ b/Engine/Plugins/Experimental/ChaosVD/Source/ChaosVD/Public/Visualizers/ChaosVDParticleDataComponentVisualizer.h
@@ -114,7 +114,6 @@ public:
 	virtual bool VisProxyHandleClick(FEditorViewportClient* InViewportClient, HComponentVisProxy* VisProxy, const FViewportClick& Click) override;
 
 protected:
-	FString GenerateDebugTextForVector(const FVector& InVector, const FString& VectorName, const FString& InVectorUnits);
 	void DrawParticleVector(FPrimitiveDrawInterface* PDI, const FVector& StartLocation, const FVector& InVector, EChaosVDParticleDataVisualizationFlags VectorID, const FChaosVDParticleDataVisualizationContext& InVisualizationContext, float LineThickness);
 	void DrawVisualizationForParticleData(const UActorComponent* Component, FPrimitiveDrawInterface* PDI, const FSceneView* View, const FChaosVDParticleDataVisualizationContext& InVisualizationContext, const FChaosVDParticleDataWrapper& InParticleDataViewer);
 };
diff --git a/Engine/Plugins/Experimental/ChaosVD/Source/ChaosVD/Public/Widgets/SChaosVDConstraintDataInspector.h b/Engine/Plugins/Experimental/ChaosVD/Source/ChaosVD/Public/Widgets/SChaosVDConstraintDataInspector.h
new file mode 100644
--- /dev/null
+++ b/Engine/Plugins/Experimental/ChaosVD/Source/ChaosVD/Public/Widgets/SChaosVDConstraintDataInspector.h
@@ -0,0 +1,67 @@
+﻿// Copyright Epic Games, Inc. All Rights Reserved.
+
+#pragma once
+#include "Actors/ChaosVDSolverInfoActor.h"
+#include "Components/ChaosVDSolverJointConstraintDataComponent.h"
+#include "Widgets/SCompoundWidget.h"
+
+struct FChaosVDJointConstraintSelectionHandle;
+class FChaosVDScene;
+class IStructureDetailsView;
+class SChaosVDNameListPicker;
+
+enum class EChaosVDParticleSelector : int32
+{
+	Index_0,
+	Index_1
+};
+
+class SChaosVDConstraintDataInspector : public SCompoundWidget
+{
+public:
+
+	SLATE_BEGIN_ARGS(SChaosVDConstraintDataInspector)
+	{
+	}
+
+	SLATE_END_ARGS()
+
+	virtual ~SChaosVDConstraintDataInspector() override;
+
+	/** Constructs this widget with InArgs */
+	void Construct(const FArguments& InArgs, const TWeakPtr<FChaosVDScene>& InScenePtr);
+	/** Sets a new query data to be inspected */
+	void SetConstraintDataToInspect(const FChaosVDJointConstraintSelectionHandle& InDataSelectionHandle);
+
+protected:
+
+	FText GetParticleName(EChaosVDParticleSelector PairIndex) const;
+
+	TSharedPtr<IStructureDetailsView> CreateDataDetailsView();
+
+	TSharedRef<SWidget> GenerateParticleSelectorButtons();
+	
+	FReply SelectParticleForCurrentSelectedData(EChaosVDParticleSelector ParticlePairIndex);
+
+	EVisibility GetOutOfDateWarningVisibility() const;
+	EVisibility GetDetailsSectionVisibility() const;
+	EVisibility GetNothingSelectedMessageVisibility() const;
+
+	void RegisterSelectionEventsForSolver(AChaosVDSolverInfoActor* SolverInfo);
+	void UnregisterSelectionEventsForSolver(AChaosVDSolverInfoActor* SolverInfo);
+
+	void RegisterSceneEvents();
+	void UnregisterSceneEvents();
+
+	void HandleSceneUpdated();
+
+	void ClearInspector();
+
+	TSharedPtr<IStructureDetailsView> ConstraintDataDetailsView;
+	
+	TWeakPtr<FChaosVDScene> SceneWeakPtr;
+
+	FChaosVDJointConstraintSelectionHandle CurrentDataSelectionHandle;
+
+	bool bIsUpToDate = true;
+};
diff --git a/Engine/Source/Runtime/Experimental/Chaos/Private/Chaos/Evolution/PBDMinEvolution.cpp b/Engine/Source/Runtime/Experimental/Chaos/Private/Chaos/Evolution/PBDMinEvolution.cpp
--- a/Engine/Source/Runtime/Experimental/Chaos/Private/Chaos/Evolution/PBDMinEvolution.cpp
+++ b/Engine/Source/Runtime/Experimental/Chaos/Private/Chaos/Evolution/PBDMinEvolution.cpp
@@ -15,6 +15,7 @@
 #include "Chaos/PerParticlePBDGroundConstraint.h"
 #include "Chaos/PerParticlePBDUpdateFromDeltaPosition.h"
 #include "ChaosStats.h"
+#include "Chaos/PBDJointConstraints.h"
 #include "ChaosVisualDebugger/ChaosVisualDebuggerTrace.h"
 
 //UE_DISABLE_OPTIMIZATION
@@ -145,6 +146,8 @@ namespace Chaos
 			ScatterOutput(Dt);
 		}
 
+		CVD_TRACE_CONSTRAINTS_CONTAINER(ConstraintContainers);
+
 		CVD_TRACE_PARTICLES(Particles.GetParticleHandles());
 	}
 
diff --git a/Engine/Source/Runtime/Experimental/Chaos/Private/Chaos/PBDJointConstraints.cpp b/Engine/Source/Runtime/Experimental/Chaos/Private/Chaos/PBDJointConstraints.cpp
--- a/Engine/Source/Runtime/Experimental/Chaos/Private/Chaos/PBDJointConstraints.cpp
+++ b/Engine/Source/Runtime/Experimental/Chaos/Private/Chaos/PBDJointConstraints.cpp
@@ -127,6 +127,11 @@ namespace Chaos
 		return ConcreteContainer()->SetConstraintEnabledDuringResim(ConstraintIndex, bEnabled);
 	}
 
+	bool FPBDJointConstraintHandle::IsEnabledDuringResim() const
+	{
+		return ConcreteContainer()->IsConstraintEnabledDuringResim(ConstraintIndex);
+	}
+
 	EResimType FPBDJointConstraintHandle::ResimType() const
 	{
 		return ConcreteContainer()->GetConstraintResimType(ConstraintIndex);
@@ -896,6 +901,11 @@ namespace Chaos
 		ConstraintStates[ConstraintIndex].bEnabledDuringResim = bEnabled;
 	}
 
+	bool FPBDJointConstraints::IsConstraintEnabledDuringResim(int32 ConstraintIndex) const
+	{
+		return ConstraintStates[ConstraintIndex].bEnabledDuringResim;
+	}
+
 	EResimType FPBDJointConstraints::GetConstraintResimType(int32 ConstraintIndex) const
 	{
 		return ConstraintStates[ConstraintIndex].ResimType;
diff --git a/Engine/Source/Runtime/Experimental/Chaos/Private/Chaos/PBDRigidsEvolutionGBF.cpp b/Engine/Source/Runtime/Experimental/Chaos/Private/Chaos/PBDRigidsEvolutionGBF.cpp
--- a/Engine/Source/Runtime/Experimental/Chaos/Private/Chaos/PBDRigidsEvolutionGBF.cpp
+++ b/Engine/Source/Runtime/Experimental/Chaos/Private/Chaos/PBDRigidsEvolutionGBF.cpp
@@ -712,6 +712,8 @@ void FPBDRigidsEvolutionGBF::AdvanceOneTimeStepImpl(const FReal Dt, const FSubSt
 			CVD_TRACE_STEP_MID_PHASES_FROM_COLLISION_CONSTRAINTS(CVDDC_EndOfEvolutionCollisionConstraints, GetCollisionConstraints());
 		}
 
+		CVD_TRACE_JOINT_CONSTRAINTS(CVDDC_JointConstraints, JointConstraints);
+
 		CVD_TRACE_PARTICLES_SOA(Particles);
 	}
 
diff --git a/Engine/Source/Runtime/Experimental/Chaos/Private/ChaosVisualDebugger/ChaosVDDataWrapperUtils.cpp b/Engine/Source/Runtime/Experimental/Chaos/Private/ChaosVisualDebugger/ChaosVDDataWrapperUtils.cpp
--- a/Engine/Source/Runtime/Experimental/Chaos/Private/ChaosVisualDebugger/ChaosVDDataWrapperUtils.cpp
+++ b/Engine/Source/Runtime/Experimental/Chaos/Private/ChaosVisualDebugger/ChaosVDDataWrapperUtils.cpp
@@ -4,6 +4,9 @@
 
 #if WITH_CHAOS_VISUAL_DEBUGGER
 
+#include "Chaos/PBDJointConstraints.h"
+#include "DataWrappers/ChaosVDJointDataWrappers.h"
+
 #include "Chaos/Collision/ParticlePairMidPhase.h"
 #include "Chaos/Collision/CollisionConstraintAllocator.h"
 #include "Chaos/ParticleHandle.h"
@@ -208,6 +211,158 @@ FChaosVDParticlePairMidPhase FChaosVDDataWrapperUtils::BuildMidPhaseDataWrapperF
 	return MoveTemp(WrappedMidPhaseData);
 }
 
+#ifndef CVD_COPY_FIELD_TO_WRAPPER
+	#define CVD_COPY_FIELD_TO_WRAPPER(Source,Target, Field) \
+	Target.Field = Source.Field;
+#endif
+
+#ifndef CVD_COPY_FIELD_TO_WRAPPER_WITH_CAST
+	#define CVD_COPY_FIELD_TO_WRAPPER_WITH_CAST(Source,Target, Field, TargetType) \
+	Target.Field = static_cast<TargetType>(Source.Field);
+#endif
+
+#ifndef CVD_COPY_VECTOR_FIELD_TO_WRAPPER
+	#define CVD_COPY_VECTOR_FIELD_TO_WRAPPER(Source,Target, Field) \
+	Target.Field = FChaosVDDataWrapperUtils::ConvertToFVector(Source.Field);
+#endif
+
+FChaosVDJointConstraint FChaosVDDataWrapperUtils::BuildJointDataWrapper(const Chaos::FPBDJointConstraintHandle* ConstaintHanlde)
+{
+	FChaosVDJointConstraint WrappedJointData;
+
+	if (ConstaintHanlde)
+	{
+		WrappedJointData.ConstraintIndex = ConstaintHanlde->GetConstraintIndex();
+
+		WrappedJointData.ParticleParIndexes[0] = ConstaintHanlde->GetConstrainedParticles()[0]->UniqueIdx().Idx;
+		WrappedJointData.ParticleParIndexes[1] = ConstaintHanlde->GetConstrainedParticles()[1]->UniqueIdx().Idx;
+
+		WrappedJointData.JointState.bBroken = ConstaintHanlde->IsConstraintBroken();
+		WrappedJointData.JointState.bBreaking = ConstaintHanlde->IsConstraintBreaking();
+		WrappedJointData.JointState.bDisabled = !ConstaintHanlde->IsConstraintEnabled();
+		WrappedJointData.JointState.bDriveTargetChanged = ConstaintHanlde->IsDriveTargetChanged();
+
+		//TODO: Island related data getters are deprecated. We need to see where is bet to get that data now and if this should be recorded as part of the CVD Constraint wrapper
+		//WrappedJointData.JointState.Color = ConstaintHanlde->GetConstraintColor();
+		//WrappedJointData.JointState.Island = ConstaintHanlde->GetConstraintIsland();
+		//WrappedJointData.JointState.IslandSize = ConstaintHanlde->GetConstraintIsland();
+		
+		WrappedJointData.JointState.bEnabledDuringResim = ConstaintHanlde->IsEnabledDuringResim();
+		WrappedJointData.JointState.AngularImpulse = ConstaintHanlde->GetAngularImpulse();
+		WrappedJointData.JointState.LinearImpulse = ConstaintHanlde->GetLinearImpulse();
+
+		switch(ConstaintHanlde->ResimType())
+		{
+			case Chaos::EResimType::FullResim:
+				WrappedJointData.JointState.ResimType = EChaosVDJointReSimType::FullResim;
+				break;
+			case Chaos::EResimType::ResimAsFollower:
+				WrappedJointData.JointState.ResimType = EChaosVDJointReSimType::ResimAsFollower;
+				break;
+		}
+
+		switch(ConstaintHanlde->SyncState())
+		{
+			case Chaos::ESyncState::InSync:
+				WrappedJointData.JointState.SyncState = EChaosVDJointSyncType::InSync;
+				break;
+			case Chaos::ESyncState::HardDesync:
+				WrappedJointData.JointState.SyncState = EChaosVDJointSyncType::HardDesync;
+				break;
+		}
+
+		WrappedJointData.JointState.MarkAsValid();
+
+		CVD_COPY_FIELD_TO_WRAPPER(ConstaintHanlde->GetJointSettings(), WrappedJointData.JointSettings, Stiffness);
+		CVD_COPY_FIELD_TO_WRAPPER(ConstaintHanlde->GetJointSettings(), WrappedJointData.JointSettings, LinearProjection);
+		CVD_COPY_FIELD_TO_WRAPPER(ConstaintHanlde->GetJointSettings(), WrappedJointData.JointSettings, AngularProjection);
+		CVD_COPY_FIELD_TO_WRAPPER(ConstaintHanlde->GetJointSettings(), WrappedJointData.JointSettings, AngularProjection);
+		CVD_COPY_FIELD_TO_WRAPPER(ConstaintHanlde->GetJointSettings(), WrappedJointData.JointSettings, TeleportDistance);
+		CVD_COPY_FIELD_TO_WRAPPER(ConstaintHanlde->GetJointSettings(), WrappedJointData.JointSettings, TeleportAngle);
+		CVD_COPY_FIELD_TO_WRAPPER(ConstaintHanlde->GetJointSettings(), WrappedJointData.JointSettings, ParentInvMassScale);
+		CVD_COPY_FIELD_TO_WRAPPER(ConstaintHanlde->GetJointSettings(), WrappedJointData.JointSettings, bCollisionEnabled);
+		CVD_COPY_FIELD_TO_WRAPPER(ConstaintHanlde->GetJointSettings(), WrappedJointData.JointSettings, bMassConditioningEnabled);
+		CVD_COPY_FIELD_TO_WRAPPER(ConstaintHanlde->GetJointSettings(), WrappedJointData.JointSettings, bSoftLinearLimitsEnabled);
+		CVD_COPY_FIELD_TO_WRAPPER(ConstaintHanlde->GetJointSettings(), WrappedJointData.JointSettings, bSoftTwistLimitsEnabled);
+		CVD_COPY_FIELD_TO_WRAPPER(ConstaintHanlde->GetJointSettings(), WrappedJointData.JointSettings, bAngularSLerpPositionDriveEnabled);
+		CVD_COPY_FIELD_TO_WRAPPER(ConstaintHanlde->GetJointSettings(), WrappedJointData.JointSettings, bAngularSLerpVelocityDriveEnabled);
+		CVD_COPY_FIELD_TO_WRAPPER(ConstaintHanlde->GetJointSettings(), WrappedJointData.JointSettings, bAngularTwistPositionDriveEnabled);
+		CVD_COPY_FIELD_TO_WRAPPER(ConstaintHanlde->GetJointSettings(), WrappedJointData.JointSettings, bAngularTwistVelocityDriveEnabled);
+		CVD_COPY_FIELD_TO_WRAPPER(ConstaintHanlde->GetJointSettings(), WrappedJointData.JointSettings, bAngularSwingPositionDriveEnabled);
+		CVD_COPY_FIELD_TO_WRAPPER(ConstaintHanlde->GetJointSettings(), WrappedJointData.JointSettings, bAngularSwingVelocityDriveEnabled);
+
+		WrappedJointData.JointSettings.LinearMotionTypes[0] = static_cast<EChaosVDJointMotionType>(ConstaintHanlde->GetJointSettings().LinearMotionTypes[0]);
+		WrappedJointData.JointSettings.LinearMotionTypes[1] = static_cast<EChaosVDJointMotionType>(ConstaintHanlde->GetJointSettings().LinearMotionTypes[1]);
+		WrappedJointData.JointSettings.LinearMotionTypes[2] = static_cast<EChaosVDJointMotionType>(ConstaintHanlde->GetJointSettings().LinearMotionTypes[2]);
+
+		CVD_COPY_FIELD_TO_WRAPPER(ConstaintHanlde->GetJointSettings(), WrappedJointData.JointSettings, LinearLimit);
+
+		WrappedJointData.JointSettings.AngularMotionTypes[0] = static_cast<EChaosVDJointMotionType>(ConstaintHanlde->GetJointSettings().AngularMotionTypes[0]);
+		WrappedJointData.JointSettings.AngularMotionTypes[1] = static_cast<EChaosVDJointMotionType>(ConstaintHanlde->GetJointSettings().AngularMotionTypes[1]);
+		WrappedJointData.JointSettings.AngularMotionTypes[2] = static_cast<EChaosVDJointMotionType>(ConstaintHanlde->GetJointSettings().AngularMotionTypes[2]);
+
+		CVD_COPY_FIELD_TO_WRAPPER(ConstaintHanlde->GetJointSettings(), WrappedJointData.JointSettings, AngularLimits);
+		CVD_COPY_FIELD_TO_WRAPPER_WITH_CAST(ConstaintHanlde->GetJointSettings(), WrappedJointData.JointSettings, LinearSoftForceMode, EChaosVDJointForceMode);
+		CVD_COPY_FIELD_TO_WRAPPER_WITH_CAST(ConstaintHanlde->GetJointSettings(), WrappedJointData.JointSettings, AngularSoftForceMode, EChaosVDJointForceMode);
+		CVD_COPY_FIELD_TO_WRAPPER(ConstaintHanlde->GetJointSettings(), WrappedJointData.JointSettings, SoftLinearStiffness);
+		CVD_COPY_FIELD_TO_WRAPPER(ConstaintHanlde->GetJointSettings(), WrappedJointData.JointSettings, SoftLinearDamping);
+		CVD_COPY_FIELD_TO_WRAPPER(ConstaintHanlde->GetJointSettings(), WrappedJointData.JointSettings, SoftTwistStiffness);
+		CVD_COPY_FIELD_TO_WRAPPER(ConstaintHanlde->GetJointSettings(), WrappedJointData.JointSettings, SoftTwistDamping);
+		CVD_COPY_FIELD_TO_WRAPPER(ConstaintHanlde->GetJointSettings(), WrappedJointData.JointSettings, SoftSwingStiffness);
+		CVD_COPY_FIELD_TO_WRAPPER(ConstaintHanlde->GetJointSettings(), WrappedJointData.JointSettings, SoftSwingDamping);
+		CVD_COPY_FIELD_TO_WRAPPER(ConstaintHanlde->GetJointSettings(), WrappedJointData.JointSettings, LinearRestitution);
+		CVD_COPY_FIELD_TO_WRAPPER(ConstaintHanlde->GetJointSettings(), WrappedJointData.JointSettings, TwistRestitution);
+		CVD_COPY_FIELD_TO_WRAPPER(ConstaintHanlde->GetJointSettings(), WrappedJointData.JointSettings, TwistRestitution);
+		CVD_COPY_FIELD_TO_WRAPPER(ConstaintHanlde->GetJointSettings(), WrappedJointData.JointSettings, SwingRestitution);
+		CVD_COPY_FIELD_TO_WRAPPER(ConstaintHanlde->GetJointSettings(), WrappedJointData.JointSettings, LinearContactDistance);
+		CVD_COPY_FIELD_TO_WRAPPER(ConstaintHanlde->GetJointSettings(), WrappedJointData.JointSettings, TwistContactDistance);
+		CVD_COPY_FIELD_TO_WRAPPER(ConstaintHanlde->GetJointSettings(), WrappedJointData.JointSettings, SwingContactDistance);
+
+		CVD_COPY_VECTOR_FIELD_TO_WRAPPER(ConstaintHanlde->GetJointSettings(), WrappedJointData.JointSettings, LinearDrivePositionTarget);
+		CVD_COPY_VECTOR_FIELD_TO_WRAPPER(ConstaintHanlde->GetJointSettings(), WrappedJointData.JointSettings, LinearDriveVelocityTarget);
+		
+		WrappedJointData.JointSettings.bLinearPositionDriveEnabled0 = ConstaintHanlde->GetJointSettings().bLinearPositionDriveEnabled[0];
+		WrappedJointData.JointSettings.bLinearPositionDriveEnabled1 = ConstaintHanlde->GetJointSettings().bLinearPositionDriveEnabled[1];
+		WrappedJointData.JointSettings.bLinearPositionDriveEnabled2 = ConstaintHanlde->GetJointSettings().bLinearPositionDriveEnabled[2];
+
+		WrappedJointData.JointSettings.bLinearVelocityDriveEnabled0 = ConstaintHanlde->GetJointSettings().bLinearVelocityDriveEnabled[0];
+		WrappedJointData.JointSettings.bLinearVelocityDriveEnabled1 = ConstaintHanlde->GetJointSettings().bLinearVelocityDriveEnabled[1];
+		WrappedJointData.JointSettings.bLinearVelocityDriveEnabled2 = ConstaintHanlde->GetJointSettings().bLinearVelocityDriveEnabled[2];
+
+		CVD_COPY_FIELD_TO_WRAPPER_WITH_CAST(ConstaintHanlde->GetJointSettings(), WrappedJointData.JointSettings, LinearDriveForceMode, EChaosVDJointForceMode);
+
+		CVD_COPY_VECTOR_FIELD_TO_WRAPPER(ConstaintHanlde->GetJointSettings(), WrappedJointData.JointSettings, LinearDriveStiffness);
+		CVD_COPY_VECTOR_FIELD_TO_WRAPPER(ConstaintHanlde->GetJointSettings(), WrappedJointData.JointSettings, LinearDriveDamping);
+		CVD_COPY_VECTOR_FIELD_TO_WRAPPER(ConstaintHanlde->GetJointSettings(), WrappedJointData.JointSettings, LinearDriveMaxForce);
+		CVD_COPY_FIELD_TO_WRAPPER(ConstaintHanlde->GetJointSettings(), WrappedJointData.JointSettings, AngularDrivePositionTarget);
+		CVD_COPY_VECTOR_FIELD_TO_WRAPPER(ConstaintHanlde->GetJointSettings(), WrappedJointData.JointSettings, AngularDriveVelocityTarget);
+		CVD_COPY_FIELD_TO_WRAPPER_WITH_CAST(ConstaintHanlde->GetJointSettings(), WrappedJointData.JointSettings, AngularDriveForceMode, EChaosVDJointForceMode);
+		CVD_COPY_VECTOR_FIELD_TO_WRAPPER(ConstaintHanlde->GetJointSettings(), WrappedJointData.JointSettings, AngularDriveStiffness);
+		CVD_COPY_VECTOR_FIELD_TO_WRAPPER(ConstaintHanlde->GetJointSettings(), WrappedJointData.JointSettings, AngularDriveDamping);
+		CVD_COPY_VECTOR_FIELD_TO_WRAPPER(ConstaintHanlde->GetJointSettings(), WrappedJointData.JointSettings, AngularDriveMaxTorque);
+
+		CVD_COPY_FIELD_TO_WRAPPER(ConstaintHanlde->GetJointSettings(), WrappedJointData.JointSettings, LinearBreakForce);
+		CVD_COPY_FIELD_TO_WRAPPER(ConstaintHanlde->GetJointSettings(), WrappedJointData.JointSettings, LinearPlasticityLimit);
+
+		CVD_COPY_FIELD_TO_WRAPPER_WITH_CAST(ConstaintHanlde->GetJointSettings(), WrappedJointData.JointSettings, LinearPlasticityType, EChaosVDPlasticityType);
+		
+		CVD_COPY_FIELD_TO_WRAPPER(ConstaintHanlde->GetJointSettings(), WrappedJointData.JointSettings, LinearPlasticityInitialDistanceSquared);
+		CVD_COPY_FIELD_TO_WRAPPER(ConstaintHanlde->GetJointSettings(), WrappedJointData.JointSettings, AngularBreakTorque);
+		CVD_COPY_FIELD_TO_WRAPPER(ConstaintHanlde->GetJointSettings(), WrappedJointData.JointSettings, AngularPlasticityLimit);
+		CVD_COPY_FIELD_TO_WRAPPER(ConstaintHanlde->GetJointSettings(), WrappedJointData.JointSettings, ContactTransferScale);
+		
+		WrappedJointData.JointSettings.ConnectorTransforms[0] = ConstaintHanlde->GetJointSettings().ConnectorTransforms[0];
+		WrappedJointData.JointSettings.ConnectorTransforms[1] = ConstaintHanlde->GetJointSettings().ConnectorTransforms[1];
+	
+		WrappedJointData.JointSettings.MarkAsValid();
+		 
+	}
+	
+	WrappedJointData.MarkAsValid();
+	
+	return MoveTemp(WrappedJointData);
+}
+
 void FChaosVDDataWrapperUtils::CopyShapeDataToWrapper(const Chaos::FShapeInstancePtr& ShapeDataPtr, FChaosVDShapeCollisionData& OutCopyTo)
 {
 	const Chaos::FCollisionData& CollisionData = ShapeDataPtr->GetCollisionData();
diff --git a/Engine/Source/Runtime/Experimental/Chaos/Private/ChaosVisualDebugger/ChaosVDOptionalDataChannel.cpp b/Engine/Source/Runtime/Experimental/Chaos/Private/ChaosVisualDebugger/ChaosVDOptionalDataChannel.cpp
--- a/Engine/Source/Runtime/Experimental/Chaos/Private/ChaosVisualDebugger/ChaosVDOptionalDataChannel.cpp
+++ b/Engine/Source/Runtime/Experimental/Chaos/Private/ChaosVisualDebugger/ChaosVDOptionalDataChannel.cpp
@@ -82,5 +82,6 @@ CVD_DEFINE_OPTIONAL_DATA_CHANNEL(CollisionDetectionNarrowPhase, EChaosVDDataChan
 CVD_DEFINE_OPTIONAL_DATA_CHANNEL(EndOfEvolutionCollisionConstraints, EChaosVDDataChannelInitializationFlags::StartEnabled| EChaosVDDataChannelInitializationFlags::CanChangeEnabledState);
 CVD_DEFINE_OPTIONAL_DATA_CHANNEL(EvolutionEnd, EChaosVDDataChannelInitializationFlags::StartEnabled); // Intentionally not allowing changing the enabled state because otherwise there might be nothing to visualize
 CVD_DEFINE_OPTIONAL_DATA_CHANNEL(SceneQueries, EChaosVDDataChannelInitializationFlags::StartEnabled | EChaosVDDataChannelInitializationFlags::CanChangeEnabledState);
+CVD_DEFINE_OPTIONAL_DATA_CHANNEL(JointConstraints, EChaosVDDataChannelInitializationFlags::CanChangeEnabledState);
 
 #endif
diff --git a/Engine/Source/Runtime/Experimental/Chaos/Private/ChaosVisualDebugger/ChaosVisualDebuggerTrace.cpp b/Engine/Source/Runtime/Experimental/Chaos/Private/ChaosVisualDebugger/ChaosVisualDebuggerTrace.cpp
--- a/Engine/Source/Runtime/Experimental/Chaos/Private/ChaosVisualDebugger/ChaosVisualDebuggerTrace.cpp
+++ b/Engine/Source/Runtime/Experimental/Chaos/Private/ChaosVisualDebugger/ChaosVisualDebuggerTrace.cpp
@@ -8,6 +8,7 @@
 #include "Chaos/ImplicitObject.h"
 #include "Chaos/ParticleHandle.h"
 #include "Chaos/PBDCollisionConstraints.h"
+#include "Chaos/PBDJointConstraints.h"
 #include "Chaos/PBDRigidsSOAs.h"
 #include "ChaosVisualDebugger/ChaosVDDataWrapperUtils.h"
 #include "ChaosVisualDebugger/ChaosVDMemWriterReader.h"
@@ -15,6 +16,7 @@
 #include "Compression/OodleDataCompressionUtil.h"
 #include "DataWrappers/ChaosVDCollisionDataWrappers.h"
 #include "DataWrappers/ChaosVDImplicitObjectDataWrapper.h"
+#include "DataWrappers/ChaosVDJointDataWrappers.h"
 #include "DataWrappers/ChaosVDParticleDataWrapper.h"
 #include "DataWrappers/ChaosVDQueryDataWrappers.h"
 #include "HAL/CriticalSection.h"
@@ -318,6 +320,39 @@ void FChaosVisualDebuggerTrace::TraceMidPhasesFromCollisionConstraints(Chaos::FP
 	});
 }
 
+void FChaosVisualDebuggerTrace::TraceJointsConstraints(Chaos::FPBDJointConstraints& InJointConstraints)
+{
+	using namespace Chaos::VisualDebugger::Utils;
+
+	if (!IsTracing())
+	{
+		return;
+	}
+
+	const FChaosVDContext* CVDContextData = FChaosVDThreadContext::Get().GetCurrentContext(EChaosVDContextType::Solver);
+
+	if (!IsContextEnabledAndValid(CVDContextData))
+	{
+		return;
+	}
+
+	const Chaos::FPBDJointConstraints::FHandles& JointHandles = InJointConstraints.GetConstConstraintHandles();
+
+	ParallelFor(JointHandles.Num(), [&JointHandles, CopyContext = *CVDContextData](int32 ConstraintIndex)
+	{
+		CVD_SCOPE_CONTEXT(CopyContext);
+		
+		FChaosVDJointConstraint WrappedJointConstraintData = FChaosVDDataWrapperUtils::BuildJointDataWrapper(JointHandles[ConstraintIndex]);
+
+		WrappedJointConstraintData.SolverID = CopyContext.Id;
+
+		FChaosVDScopedTLSBufferAccessor TLSDataBuffer;
+		Chaos::VisualDebugger::WriteDataToBuffer(TLSDataBuffer.BufferRef, WrappedJointConstraintData);
+
+		TraceBinaryData(TLSDataBuffer.BufferRef, FChaosVDJointConstraint::WrapperTypeName);
+	});
+}
+
 void FChaosVisualDebuggerTrace::TraceCollisionConstraint(const Chaos::FPBDCollisionConstraint* CollisionConstraint)
 {
 	using namespace Chaos::VisualDebugger::Utils;
@@ -365,6 +400,29 @@ void FChaosVisualDebuggerTrace::TraceCollisionConstraintView(TArrayView<Chaos::F
 	});
 }
 
+void FChaosVisualDebuggerTrace::TraceConstraintsContainer(TConstArrayView<Chaos::FPBDConstraintContainer*> ConstraintContainersView)
+{
+	if (!IsTracing())
+	{
+		return;
+	}
+
+	for (Chaos::FPBDConstraintContainer* ConstraintContainer : ConstraintContainersView)
+	{
+		if (ConstraintContainer)
+		{
+			if (ConstraintContainer->GetConstraintHandleType().IsA(Chaos::FPBDJointConstraintHandle::StaticType()))
+			{
+				CVD_TRACE_JOINT_CONSTRAINTS(CVDDC_JointConstraints, *static_cast<Chaos::FPBDJointConstraints*>(ConstraintContainer));
+			}
+			else if (ConstraintContainer->GetConstraintHandleType().IsA(Chaos::FPBDCollisionConstraint::StaticType()))
+			{
+				CVD_TRACE_STEP_MID_PHASES_FROM_COLLISION_CONSTRAINTS(CVDDC_EndOfEvolutionCollisionConstraints, *static_cast<Chaos::FPBDCollisionConstraints*>(ConstraintContainer));
+			}
+		}
+	}
+}
+
 void FChaosVisualDebuggerTrace::TraceSolverFrameStart(const FChaosVDContext& ContextData, const FString& InDebugName)
 {
 	if (!IsTracing())
diff --git a/Engine/Source/Runtime/Experimental/Chaos/Public/Chaos/PBDJointConstraints.h b/Engine/Source/Runtime/Experimental/Chaos/Public/Chaos/PBDJointConstraints.h
--- a/Engine/Source/Runtime/Experimental/Chaos/Public/Chaos/PBDJointConstraints.h
+++ b/Engine/Source/Runtime/Experimental/Chaos/Public/Chaos/PBDJointConstraints.h
@@ -99,6 +99,7 @@ namespace Chaos
 		CHAOS_API void SetSyncState(ESyncState SyncState);
 
 		CHAOS_API void SetEnabledDuringResim(bool bEnabled);
+		CHAOS_API bool IsEnabledDuringResim() const;
 		CHAOS_API EResimType ResimType() const;
 
 		UE_DEPRECATED(5.2, "No longer used")
@@ -337,6 +338,8 @@ namespace Chaos
 		
 		CHAOS_API void SetConstraintEnabledDuringResim(int32 ConstraintIndex, bool bEnabled);
 		
+		CHAOS_API bool IsConstraintEnabledDuringResim(int32 ConstraintIndex) const;
+		
 		CHAOS_API EResimType GetConstraintResimType(int32 ConstraintIndex) const;
 
 		//
diff --git a/Engine/Source/Runtime/Experimental/Chaos/Public/ChaosVisualDebugger/ChaosVDDataWrapperUtils.h b/Engine/Source/Runtime/Experimental/Chaos/Public/ChaosVisualDebugger/ChaosVDDataWrapperUtils.h
--- a/Engine/Source/Runtime/Experimental/Chaos/Public/ChaosVisualDebugger/ChaosVDDataWrapperUtils.h
+++ b/Engine/Source/Runtime/Experimental/Chaos/Public/ChaosVisualDebugger/ChaosVDDataWrapperUtils.h
@@ -8,8 +8,12 @@
 #include "Chaos/ParticleHandleFwd.h"
 #include "DataWrappers/ChaosVDCollisionDataWrappers.h"
 
+struct FChaosVDJointConstraint;
+
 namespace Chaos
 {
+	class FPBDJointConstraintHandle;
+
 	namespace VisualDebugger
 	{
 		class FChaosVDSerializableNameTable;
@@ -55,6 +59,9 @@ private:
 	/** Creates and populates a FChaosVDParticlePairMidPhase with the data of the provided FParticlePairMidPhase */
 	static FChaosVDParticlePairMidPhase BuildMidPhaseDataWrapperFromMidPhase(const Chaos::FParticlePairMidPhase& InMidPhase);
 
+	/** Creates and populates a FChaosVDJointConstraint with the data of the provided FPBDJointConstraintHandle */
+	static FChaosVDJointConstraint BuildJointDataWrapper(const Chaos::FPBDJointConstraintHandle* ConstaintHanlde);
+
 	/** Converts a Chaos::FVec3 to a FVector. It is worth notice that FVector is double precision and FVec3 is single */
 	static FVector ConvertToFVector(const Chaos::FVec3f& VectorRef) { return FVector(VectorRef); }
 
diff --git a/Engine/Source/Runtime/Experimental/Chaos/Public/ChaosVisualDebugger/ChaosVDOptionalDataChannel.h b/Engine/Source/Runtime/Experimental/Chaos/Public/ChaosVisualDebugger/ChaosVDOptionalDataChannel.h
--- a/Engine/Source/Runtime/Experimental/Chaos/Public/ChaosVisualDebugger/ChaosVDOptionalDataChannel.h
+++ b/Engine/Source/Runtime/Experimental/Chaos/Public/ChaosVisualDebugger/ChaosVDOptionalDataChannel.h
@@ -150,6 +150,7 @@ CVD_DECLARE_OPTIONAL_DATA_CHANNEL(CollisionDetectionNarrowPhase);
 CVD_DECLARE_OPTIONAL_DATA_CHANNEL(EndOfEvolutionCollisionConstraints);
 CVD_DECLARE_OPTIONAL_DATA_CHANNEL(EvolutionEnd);
 CVD_DECLARE_OPTIONAL_DATA_CHANNEL(SceneQueries);
+CVD_DECLARE_OPTIONAL_DATA_CHANNEL(JointConstraints);
 
 #else  //WITH_CHAOS_VISUAL_DEBUGGER
 
diff --git a/Engine/Source/Runtime/Experimental/Chaos/Public/ChaosVisualDebugger/ChaosVDTraceMacros.h b/Engine/Source/Runtime/Experimental/Chaos/Public/ChaosVisualDebugger/ChaosVDTraceMacros.h
--- a/Engine/Source/Runtime/Experimental/Chaos/Public/ChaosVisualDebugger/ChaosVDTraceMacros.h
+++ b/Engine/Source/Runtime/Experimental/Chaos/Public/ChaosVisualDebugger/ChaosVDTraceMacros.h
@@ -155,6 +155,19 @@
 		FChaosVisualDebuggerTrace::TraceSceneQueryVisit(InQueryVisitData);
 	#endif
 
+	#ifndef CVD_TRACE_JOINT_CONSTRAINTS
+		#define CVD_TRACE_JOINT_CONSTRAINTS(DataChannel, InJointConstraints) \
+			{ \
+				CVD_SCOPED_DATA_CHANNEL_OVERRIDE(DataChannel) \
+				FChaosVisualDebuggerTrace::TraceJointsConstraints(InJointConstraints); \
+			}
+	#endif
+
+	#ifndef CVD_TRACE_CONSTRAINTS_CONTAINER
+		#define CVD_TRACE_CONSTRAINTS_CONTAINER(ContainerView) \
+			FChaosVisualDebuggerTrace::TraceConstraintsContainer(ContainerView);
+	#endif
+
 #else // WITH_CHAOS_VISUAL_DEBUGGER
 
 	#ifndef CVD_TRACE_PARTICLE
@@ -247,4 +260,12 @@
 		#define CVD_TRACE_SCENE_QUERY_VISIT(InQueryVisitData)
 	#endif
 
+	#ifndef CVD_TRACE_JOINT_CONSTRAINTS
+		#define CVD_TRACE_JOINT_CONSTRAINTS(DataChannel, InJointConstraints)
+	#endif
+
+	#ifndef CVD_TRACE_CONSTRAINTS_CONTAINER
+		#define CVD_TRACE_CONSTRAINTS_CONTAINER(ContainerView)
+	#endif
+
 #endif // WITH_CHAOS_VISUAL_DEBUGGER
\ No newline at end of file
diff --git a/Engine/Source/Runtime/Experimental/Chaos/Public/ChaosVisualDebugger/ChaosVisualDebuggerTrace.h b/Engine/Source/Runtime/Experimental/Chaos/Public/ChaosVisualDebugger/ChaosVisualDebuggerTrace.h
--- a/Engine/Source/Runtime/Experimental/Chaos/Public/ChaosVisualDebugger/ChaosVisualDebuggerTrace.h
+++ b/Engine/Source/Runtime/Experimental/Chaos/Public/ChaosVisualDebugger/ChaosVisualDebuggerTrace.h
@@ -26,6 +26,12 @@
 #include "Trace/Trace.h"
 #include "Trace/Trace.inl"
 
+namespace Chaos
+{
+class FPBDConstraintContainer;
+class FPBDJointConstraints;
+}
+
 UE_TRACE_CHANNEL_EXTERN(ChaosVDChannel, CHAOS_API)
 
 UE_TRACE_EVENT_BEGIN_EXTERN(ChaosVDLogger, ChaosVDSolverFrameStart)
@@ -182,12 +188,18 @@ public:
 	/** Traces a Particle pair MidPhase as binary data from a provided CollisionConstraints object */
 	static CHAOS_API void TraceMidPhasesFromCollisionConstraints(Chaos::FPBDCollisionConstraints& InCollisionConstraints);
 
+	/** Traces all joint constraints in the provided container */
+	static CHAOS_API void TraceJointsConstraints(Chaos::FPBDJointConstraints& InJointConstraints);
+
 	/** Traces a Particle pair MidPhase as binary data */
 	static CHAOS_API void TraceCollisionConstraint(const Chaos::FPBDCollisionConstraint* CollisionConstraint);
 
 	/** Traces a Particle pair MidPhase as binary data in parallel */
 	static CHAOS_API void TraceCollisionConstraintView(TArrayView<Chaos::FPBDCollisionConstraint* const> CollisionConstraintView);
 
+	/** Traces all supported constraints in the provided containers view */
+	static CHAOS_API void TraceConstraintsContainer(TConstArrayView<Chaos::FPBDConstraintContainer*> ConstraintContainersView);
+
 	/** Traces the start of a solver frame and it pushes its context data to the CVD TLS context stack */
 	static CHAOS_API void TraceSolverFrameStart(const FChaosVDContext& ContextData, const FString& InDebugName);
 	
diff --git a/Engine/Source/Runtime/Experimental/ChaosVDData/Public/ChaosVDRecording.h b/Engine/Source/Runtime/Experimental/ChaosVDData/Public/ChaosVDRecording.h
--- a/Engine/Source/Runtime/Experimental/ChaosVDData/Public/ChaosVDRecording.h
+++ b/Engine/Source/Runtime/Experimental/ChaosVDData/Public/ChaosVDRecording.h
@@ -10,6 +10,7 @@
 #include "Containers/UnrealString.h"
 #include "Chaos/ImplicitFwd.h"
 #include "Chaos/ImplicitObject.h"
+#include "DataWrappers/ChaosVDJointDataWrappers.h"
 #include "DataWrappers/ChaosVDQueryDataWrappers.h"
 
 namespace Chaos::VisualDebugger
@@ -25,6 +26,7 @@ struct FChaosVDStepData
 	FString StepName;
 	TArray<TSharedPtr<FChaosVDParticleDataWrapper>> RecordedParticlesData;
 	TArray<TSharedPtr<FChaosVDParticlePairMidPhase>> RecordedMidPhases;
+	TArray<TSharedPtr<FChaosVDJointConstraint>> RecordedJointConstraints;
 	TArray<FChaosVDConstraint> RecordedConstraints;
 	TMap<int32, TArray<FChaosVDConstraint>> RecordedConstraintsByParticleID;
 	TMap<int32, TArray<TSharedPtr<FChaosVDParticlePairMidPhase>>> RecordedMidPhasesByParticleID;
diff --git a/Engine/Source/Runtime/Experimental/ChaosVisualDebugger/Private/DataWrappers/ChaosVDJointDataWrappers.cpp b/Engine/Source/Runtime/Experimental/ChaosVisualDebugger/Private/DataWrappers/ChaosVDJointDataWrappers.cpp
new file mode 100644
--- /dev/null
+++ b/Engine/Source/Runtime/Experimental/ChaosVisualDebugger/Private/DataWrappers/ChaosVDJointDataWrappers.cpp
@@ -0,0 +1,237 @@
+﻿// Copyright Epic Games, Inc. All Rights Reserved.
+
+#include "DataWrappers/ChaosVDJointDataWrappers.h"
+
+#include "DataWrappers/ChaosVDDataSerializationMacros.h"
+
+bool FChaosVDJointStateDataWrapper::Serialize(FArchive& Ar)
+{
+	Ar << bHasValidData;
+	if (!bHasValidData)
+	{
+		return !Ar.IsError();
+	}
+
+	Ar << Island;
+	Ar << Level;
+	Ar << Color;
+	Ar << IslandSize;
+
+	EChaosVDJointStateFlags StateFlags = EChaosVDJointStateFlags::None;
+	
+	if (Ar.IsLoading())
+	{
+		Ar << StateFlags;
+		CVD_UNPACK_BITFIELD_DATA(bDisabled, StateFlags, EChaosVDJointStateFlags::Disabled);
+		CVD_UNPACK_BITFIELD_DATA(bBroken, StateFlags, EChaosVDJointStateFlags::Broken);
+		CVD_UNPACK_BITFIELD_DATA(bBreaking, StateFlags, EChaosVDJointStateFlags::Breaking);
+		CVD_UNPACK_BITFIELD_DATA(bDriveTargetChanged, StateFlags, EChaosVDJointStateFlags::DriveTargetChanged);
+		CVD_UNPACK_BITFIELD_DATA(bEnabledDuringResim, StateFlags, EChaosVDJointStateFlags::EnabledDuringResim);
+	}
+	else
+	{
+		CVD_PACK_BITFIELD_DATA(bDisabled, StateFlags, EChaosVDJointStateFlags::Disabled);
+		CVD_PACK_BITFIELD_DATA(bBroken, StateFlags, EChaosVDJointStateFlags::Broken);
+		CVD_PACK_BITFIELD_DATA(bBreaking, StateFlags, EChaosVDJointStateFlags::Breaking);
+		CVD_PACK_BITFIELD_DATA(bDriveTargetChanged, StateFlags, EChaosVDJointStateFlags::DriveTargetChanged);
+		CVD_PACK_BITFIELD_DATA(bEnabledDuringResim, StateFlags, EChaosVDJointStateFlags::EnabledDuringResim);
+
+		Ar << StateFlags;
+	}
+
+	Ar << LinearImpulse;
+	Ar << AngularImpulse;
+	Ar << ResimType;
+	Ar << SyncState;
+	
+	return !Ar.IsError();
+}
+
+bool FChaosVDJointSolverSettingsDataWrapper::Serialize(FArchive& Ar)
+{
+	Ar << bHasValidData;
+	if (!bHasValidData)
+	{
+		return !Ar.IsError();
+	}
+
+	Ar << SwingTwistAngleTolerance;
+	Ar << PositionTolerance;
+	Ar << AngleTolerance;
+	Ar << MinParentMassRatio;
+	Ar << MaxInertiaRatio;
+	Ar << MinSolverStiffness;
+	Ar << MaxSolverStiffness;
+	Ar << NumIterationsAtMaxSolverStiffness;
+	Ar << NumShockPropagationIterations;
+	Ar << LinearStiffnessOverride;
+	Ar << TwistStiffnessOverride;
+	Ar << SwingStiffnessOverride;
+	Ar << LinearProjectionOverride;
+	Ar << AngularProjectionOverride;
+	Ar << ShockPropagationOverride;
+	Ar << LinearDriveStiffnessOverride;
+	Ar << LinearDriveDampingOverride;
+	Ar << AngularDriveStiffnessOverride;
+	Ar << AngularDriveDampingOverride;
+	Ar << SoftLinearStiffnessOverride;
+	Ar << SoftLinearDampingOverride;
+	Ar << SoftTwistStiffnessOverride;
+	Ar << SoftTwistDampingOverride;
+	Ar << SoftSwingStiffnessOverride;
+	Ar << SoftSwingDampingOverride;
+
+	EChaosVDJointSolverSettingsFlags Flags = EChaosVDJointSolverSettingsFlags::None;
+	
+	if (Ar.IsLoading())
+	{
+		Ar << Flags;
+		CVD_PACK_BITFIELD_DATA(bUseLinearSolver, Flags, EChaosVDJointSolverSettingsFlags::UseLinearSolver);
+		CVD_PACK_BITFIELD_DATA(bSortEnabled, Flags, EChaosVDJointSolverSettingsFlags::SortEnabled);
+		CVD_PACK_BITFIELD_DATA(bSolvePositionLast, Flags, EChaosVDJointSolverSettingsFlags::SolvePositionLast);
+		CVD_PACK_BITFIELD_DATA(bUsePositionBasedDrives, Flags, EChaosVDJointSolverSettingsFlags::UsePositionBasedDrives);
+		CVD_PACK_BITFIELD_DATA(bEnableTwistLimits, Flags, EChaosVDJointSolverSettingsFlags::EnableTwistLimits);
+		CVD_PACK_BITFIELD_DATA(bEnableSwingLimits, Flags, EChaosVDJointSolverSettingsFlags::EnableSwingLimits);
+		CVD_PACK_BITFIELD_DATA(bEnableDrives, Flags, EChaosVDJointSolverSettingsFlags::EnableDrives);
+	}
+	else
+	{
+		CVD_UNPACK_BITFIELD_DATA(bUseLinearSolver, Flags, EChaosVDJointSolverSettingsFlags::UseLinearSolver);
+		CVD_UNPACK_BITFIELD_DATA(bSortEnabled, Flags, EChaosVDJointSolverSettingsFlags::SortEnabled);
+		CVD_UNPACK_BITFIELD_DATA(bSolvePositionLast, Flags, EChaosVDJointSolverSettingsFlags::SolvePositionLast);
+		CVD_UNPACK_BITFIELD_DATA(bUsePositionBasedDrives, Flags, EChaosVDJointSolverSettingsFlags::UsePositionBasedDrives);
+		CVD_UNPACK_BITFIELD_DATA(bEnableTwistLimits, Flags, EChaosVDJointSolverSettingsFlags::EnableTwistLimits);
+		CVD_UNPACK_BITFIELD_DATA(bEnableSwingLimits, Flags, EChaosVDJointSolverSettingsFlags::EnableSwingLimits);
+		CVD_UNPACK_BITFIELD_DATA(bEnableDrives, Flags, EChaosVDJointSolverSettingsFlags::EnableDrives);
+
+		Ar << Flags;
+	}
+
+	return !Ar.IsError();
+}
+
+bool FChaosVDJointSettingsDataWrapper::Serialize(FArchive& Ar)
+{
+	Ar << bHasValidData;
+	if (!bHasValidData)
+	{
+		return !Ar.IsError();
+	}
+	
+	CVD_SERIALIZE_STATIC_ARRAY(Ar, ConnectorTransforms);
+	Ar << Stiffness;
+	Ar << LinearProjection;
+	Ar << AngularProjection;
+	Ar << ShockPropagation;
+	Ar << TeleportDistance;
+	Ar << TeleportAngle;
+	Ar << ParentInvMassScale;
+	CVD_SERIALIZE_STATIC_ARRAY(Ar, LinearMotionTypes);
+	Ar << LinearLimit;
+	CVD_SERIALIZE_STATIC_ARRAY(Ar, AngularMotionTypes);
+	Ar << AngularLimits;
+	Ar << LinearSoftForceMode;
+	Ar << AngularSoftForceMode;
+	Ar << SoftLinearStiffness;
+	Ar << SoftLinearDamping;
+	Ar << SoftTwistStiffness;
+	Ar << SoftTwistDamping;
+	Ar << SoftSwingStiffness;
+	Ar << SoftSwingDamping;
+	Ar << LinearRestitution;
+	Ar << TwistRestitution;
+	Ar << SwingRestitution;
+	Ar << LinearContactDistance;
+	Ar << TwistContactDistance;
+	Ar << SwingContactDistance;
+	Ar << LinearDrivePositionTarget;
+	Ar << LinearDriveVelocityTarget;
+	Ar << LinearDriveForceMode;
+	Ar << LinearDriveStiffness;
+	Ar << LinearDriveDamping;
+	Ar << LinearDriveMaxForce;
+	Ar << AngularDrivePositionTarget;
+	Ar << AngularDriveVelocityTarget;
+	Ar << AngularDriveForceMode;
+	Ar << AngularDriveStiffness;
+	Ar << AngularDriveDamping;
+	Ar << AngularDriveMaxTorque;
+	Ar << LinearBreakForce;
+	Ar << LinearPlasticityLimit;
+	Ar << LinearPlasticityType;
+	Ar << LinearPlasticityInitialDistanceSquared;
+	Ar << AngularBreakTorque;
+	Ar << AngularPlasticityLimit;
+	Ar << ContactTransferScale;
+
+	EChaosVDJointSettingsFlags Flags = EChaosVDJointSettingsFlags::None;
+	
+	if (Ar.IsLoading())
+	{
+		Ar << Flags;
+		CVD_PACK_BITFIELD_DATA(bCollisionEnabled, Flags, EChaosVDJointSettingsFlags::CollisionEnabled);
+		CVD_PACK_BITFIELD_DATA(bMassConditioningEnabled, Flags, EChaosVDJointSettingsFlags::MassConditioningEnabled);
+		CVD_PACK_BITFIELD_DATA(bSoftLinearLimitsEnabled, Flags, EChaosVDJointSettingsFlags::SoftLinearLimitsEnabled);
+		CVD_PACK_BITFIELD_DATA(bSoftTwistLimitsEnabled, Flags, EChaosVDJointSettingsFlags::SoftTwistLimitsEnabled);
+		CVD_PACK_BITFIELD_DATA(bSoftSwingLimitsEnabled, Flags, EChaosVDJointSettingsFlags::SoftSwingLimitsEnabled);
+		CVD_PACK_BITFIELD_DATA(bAngularSLerpPositionDriveEnabled, Flags, EChaosVDJointSettingsFlags::AngularSLerpPositionDriveEnabled);
+		CVD_PACK_BITFIELD_DATA(bAngularSLerpVelocityDriveEnabled, Flags, EChaosVDJointSettingsFlags::AngularSLerpVelocityDriveEnabled);
+		CVD_PACK_BITFIELD_DATA(bAngularTwistPositionDriveEnabled, Flags, EChaosVDJointSettingsFlags::AngularTwistPositionDriveEnabled);
+		CVD_PACK_BITFIELD_DATA(bAngularSwingPositionDriveEnabled, Flags, EChaosVDJointSettingsFlags::AngularSwingPositionDriveEnabled);
+		CVD_PACK_BITFIELD_DATA(bAngularTwistVelocityDriveEnabled, Flags, EChaosVDJointSettingsFlags::AngularTwistVelocityDriveEnabled);
+		CVD_PACK_BITFIELD_DATA(bAngularSwingVelocityDriveEnabled, Flags, EChaosVDJointSettingsFlags::AngularSwingVelocityDriveEnabled);
+		CVD_PACK_BITFIELD_DATA(bAngularSwingVelocityDriveEnabled, Flags, EChaosVDJointSettingsFlags::AngularSwingVelocityDriveEnabled);
+		CVD_PACK_BITFIELD_DATA(bLinearPositionDriveEnabled0, Flags, EChaosVDJointSettingsFlags::LinearPositionDriveEnabled0);
+		CVD_PACK_BITFIELD_DATA(bLinearPositionDriveEnabled1, Flags, EChaosVDJointSettingsFlags::LinearPositionDriveEnable1);
+		CVD_PACK_BITFIELD_DATA(bLinearPositionDriveEnabled2, Flags, EChaosVDJointSettingsFlags::LinearPositionDriveEnable2);
+		CVD_PACK_BITFIELD_DATA(bLinearVelocityDriveEnabled0, Flags, EChaosVDJointSettingsFlags::LinearVelocityDriveEnabled0);
+		CVD_PACK_BITFIELD_DATA(bLinearVelocityDriveEnabled1, Flags, EChaosVDJointSettingsFlags::LinearVelocityDriveEnabled1);
+		CVD_PACK_BITFIELD_DATA(bLinearVelocityDriveEnabled2, Flags, EChaosVDJointSettingsFlags::LinearVelocityDriveEnabled2);
+	}
+	else
+	{
+		CVD_UNPACK_BITFIELD_DATA(bCollisionEnabled, Flags, EChaosVDJointSettingsFlags::CollisionEnabled);
+		CVD_UNPACK_BITFIELD_DATA(bMassConditioningEnabled, Flags, EChaosVDJointSettingsFlags::MassConditioningEnabled);
+		CVD_UNPACK_BITFIELD_DATA(bSoftLinearLimitsEnabled, Flags, EChaosVDJointSettingsFlags::SoftLinearLimitsEnabled);
+		CVD_UNPACK_BITFIELD_DATA(bSoftTwistLimitsEnabled, Flags, EChaosVDJointSettingsFlags::SoftTwistLimitsEnabled);
+		CVD_UNPACK_BITFIELD_DATA(bSoftSwingLimitsEnabled, Flags, EChaosVDJointSettingsFlags::SoftSwingLimitsEnabled);
+		CVD_UNPACK_BITFIELD_DATA(bAngularSLerpPositionDriveEnabled, Flags, EChaosVDJointSettingsFlags::AngularSLerpPositionDriveEnabled);
+		CVD_UNPACK_BITFIELD_DATA(bAngularSLerpVelocityDriveEnabled, Flags, EChaosVDJointSettingsFlags::AngularSLerpVelocityDriveEnabled);
+		CVD_UNPACK_BITFIELD_DATA(bAngularTwistPositionDriveEnabled, Flags, EChaosVDJointSettingsFlags::AngularTwistPositionDriveEnabled);
+		CVD_UNPACK_BITFIELD_DATA(bAngularSwingPositionDriveEnabled, Flags, EChaosVDJointSettingsFlags::AngularSwingPositionDriveEnabled);
+		CVD_UNPACK_BITFIELD_DATA(bAngularTwistVelocityDriveEnabled, Flags, EChaosVDJointSettingsFlags::AngularTwistVelocityDriveEnabled);
+		CVD_UNPACK_BITFIELD_DATA(bAngularSwingVelocityDriveEnabled, Flags, EChaosVDJointSettingsFlags::AngularSwingVelocityDriveEnabled);
+		CVD_UNPACK_BITFIELD_DATA(bAngularSwingVelocityDriveEnabled, Flags, EChaosVDJointSettingsFlags::AngularSwingVelocityDriveEnabled);
+		CVD_UNPACK_BITFIELD_DATA(bLinearPositionDriveEnabled0, Flags, EChaosVDJointSettingsFlags::LinearPositionDriveEnabled0);
+		CVD_UNPACK_BITFIELD_DATA(bLinearPositionDriveEnabled1, Flags, EChaosVDJointSettingsFlags::LinearPositionDriveEnable1);
+		CVD_UNPACK_BITFIELD_DATA(bLinearPositionDriveEnabled2, Flags, EChaosVDJointSettingsFlags::LinearPositionDriveEnable2);
+		CVD_UNPACK_BITFIELD_DATA(bLinearVelocityDriveEnabled0, Flags, EChaosVDJointSettingsFlags::LinearVelocityDriveEnabled0);
+		CVD_UNPACK_BITFIELD_DATA(bLinearVelocityDriveEnabled1, Flags, EChaosVDJointSettingsFlags::LinearVelocityDriveEnabled1);
+		CVD_UNPACK_BITFIELD_DATA(bLinearVelocityDriveEnabled2, Flags, EChaosVDJointSettingsFlags::LinearVelocityDriveEnabled2);;
+
+		Ar << Flags;
+	}
+
+	return !Ar.IsError();
+}
+
+
+bool FChaosVDJointConstraint::Serialize(FArchive& Ar)
+{
+	Ar << bHasValidData;
+	if (!bHasValidData)
+	{
+		return !Ar.IsError();
+	}
+
+	Ar << SolverID;
+
+	Ar << ConstraintIndex;
+
+	CVD_SERIALIZE_STATIC_ARRAY(Ar, ParticleParIndexes);
+
+	Ar << JointSettings;
+	Ar << JointState;
+
+	return !Ar.IsError();
+}
diff --git a/Engine/Source/Runtime/Experimental/ChaosVisualDebugger/Public/DataWrappers/ChaosVDJointDataWrappers.h b/Engine/Source/Runtime/Experimental/ChaosVisualDebugger/Public/DataWrappers/ChaosVDJointDataWrappers.h
new file mode 100644
--- /dev/null
+++ b/Engine/Source/Runtime/Experimental/ChaosVisualDebugger/Public/DataWrappers/ChaosVDJointDataWrappers.h
@@ -0,0 +1,409 @@
+﻿// Copyright Epic Games, Inc. All Rights Reserved.
+
+#pragma once
+#include "ChaosVDParticleDataWrapper.h"
+#include "UObject/Class.h"
+#include "UObject/ObjectMacros.h"
+
+#include "ChaosVDJointDataWrappers.generated.h"
+
+#ifndef CVD_IMPLEMENT_SERIALIZER
+		#define CVD_IMPLEMENT_SERIALIZER(Type) \
+		inline FArchive& operator<<(FArchive& Ar, Type& Data) \
+		{\
+			Data.Serialize(Ar); \
+			return Ar; \
+		} \
+		template<>\
+		struct TStructOpsTypeTraits<Type> : public TStructOpsTypeTraitsBase2<Type> \
+		{\
+			enum\
+			{\
+				WithSerializer = true,\
+			};\
+		};\
+
+#endif
+
+UENUM()
+enum class EChaosVDJointReSimType
+{
+	/** Fully re-run simulation and keep results (any forces must be applied again) */
+	FullResim = 0,
+	/** Use previous forces and snap to previous results regardless of variation - used to push other objects away */
+	ResimAsFollower = 1
+};
+
+UENUM()
+enum class EChaosVDJointSyncType
+{
+	/** In sync with recorded data */
+	InSync,
+	/** Recorded data mismatches, must run collision detection again */
+	HardDesync,
+};
+
+UENUM()
+enum class EChaosVDJointStateFlags : uint8
+{
+	None = 0,
+	Disabled = 1 << 0,
+	Broken = 1 << 1,
+	Breaking = 1 << 2,
+	DriveTargetChanged = 1 << 3,
+	EnabledDuringResim = 1 << 4,
+};
+
+ENUM_CLASS_FLAGS(EChaosVDJointStateFlags)
+
+USTRUCT()
+struct FChaosVDJointStateDataWrapper : public FChaosVDWrapperDataBase
+{
+	GENERATED_BODY()
+public:
+
+	CHAOSVDRUNTIME_API bool Serialize(FArchive& Ar);
+
+	//TODO: Make the island data visible when we add support to record that data
+
+	int32 Island = INDEX_NONE;
+	int32 Level = INDEX_NONE;
+	int32 Color = INDEX_NONE;
+	int32 IslandSize = INDEX_NONE;
+
+	UPROPERTY(VisibleAnywhere, Category=JointState)
+	uint8 bDisabled : 1 = false;
+	UPROPERTY(VisibleAnywhere, Category=JointState)
+	uint8 bBroken : 1 = false;
+	UPROPERTY(VisibleAnywhere, Category=JointState)
+	uint8 bBreaking : 1 = false;
+	UPROPERTY(VisibleAnywhere, Category=JointState)
+	uint8 bDriveTargetChanged : 1 = false;
+	UPROPERTY(VisibleAnywhere, Category=JointState)
+	uint8 bEnabledDuringResim : 1 = true;
+
+	UPROPERTY(VisibleAnywhere, Category=JointState)
+	FVector LinearImpulse = FVector::ZeroVector;
+	UPROPERTY(VisibleAnywhere, Category=JointState)
+	FVector AngularImpulse = FVector::ZeroVector;
+	UPROPERTY(VisibleAnywhere, Category=JointState)
+	EChaosVDJointReSimType ResimType = EChaosVDJointReSimType::FullResim;
+	UPROPERTY(VisibleAnywhere, Category=JointState)
+	EChaosVDJointSyncType SyncState = EChaosVDJointSyncType::InSync;
+	
+};
+CVD_IMPLEMENT_SERIALIZER(FChaosVDJointStateDataWrapper)
+
+UENUM()
+enum class EChaosVDJointSolverSettingsFlags : uint16
+{
+	None = 0,
+
+	UseLinearSolver = 1 << 0,
+	SortEnabled = 1 << 1,
+	SolvePositionLast = 1 << 2,
+	UsePositionBasedDrives = 1 << 3,
+	EnableTwistLimits = 1 << 4,
+	EnableSwingLimits = 1 << 5,
+	EnableDrives = 1 << 6,
+};
+
+ENUM_CLASS_FLAGS(EChaosVDJointSettingsFlags)
+
+USTRUCT()
+struct FChaosVDJointSolverSettingsDataWrapper : public FChaosVDWrapperDataBase
+{
+	GENERATED_BODY()
+public:
+	CHAOSVDRUNTIME_API bool Serialize(FArchive& Ar);
+
+	UPROPERTY(VisibleAnywhere, Category=Tolerances)
+	double SwingTwistAngleTolerance;
+	UPROPERTY(VisibleAnywhere, Category=Tolerances)
+	double PositionTolerance;
+	UPROPERTY(VisibleAnywhere, Category=Tolerances)
+	double AngleTolerance;
+
+	UPROPERTY(VisibleAnywhere, Category="Stability Control")
+	double MinParentMassRatio;
+	UPROPERTY(VisibleAnywhere, Category="Stability Control")
+	double MaxInertiaRatio;
+
+	UPROPERTY(VisibleAnywhere, Category="Solver Stiffness")
+	double MinSolverStiffness;
+	UPROPERTY(VisibleAnywhere, Category="Solver Stiffness")
+	double MaxSolverStiffness;
+	UPROPERTY(VisibleAnywhere, Category="Solver Stiffness")
+	int32 NumIterationsAtMaxSolverStiffness;
+	UPROPERTY(VisibleAnywhere, Category="Solver Stiffness")
+	int32 NumShockPropagationIterations;
+
+	UPROPERTY(VisibleAnywhere, Category="General")
+	uint16 bUseLinearSolver : 1 = false;
+	UPROPERTY(VisibleAnywhere, Category="General")
+	uint16 bSortEnabled : 1 = false;
+	UPROPERTY(VisibleAnywhere, Category="General")
+	uint16 bSolvePositionLast : 1 = false;
+	UPROPERTY(VisibleAnywhere, Category="General")
+	uint16 bUsePositionBasedDrives : 1 = false;
+	UPROPERTY(VisibleAnywhere, Category="General")
+	uint16 bEnableTwistLimits : 1 = false;
+	UPROPERTY(VisibleAnywhere, Category="General")
+	uint16 bEnableSwingLimits : 1 = false;
+	UPROPERTY(VisibleAnywhere, Category="General")
+	uint16 bEnableDrives : 1 = false;
+
+	UPROPERTY(VisibleAnywhere, Category="General")
+	double LinearStiffnessOverride = 0.0;
+	UPROPERTY(VisibleAnywhere, Category="General")
+	double TwistStiffnessOverride = 0.0;
+	UPROPERTY(VisibleAnywhere, Category="General")
+	double SwingStiffnessOverride = 0.0;
+	UPROPERTY(VisibleAnywhere, Category="General")
+	double LinearProjectionOverride = 0.0;
+	UPROPERTY(VisibleAnywhere, Category="General")
+	double AngularProjectionOverride = 0.0;
+	UPROPERTY(VisibleAnywhere, Category="General")
+	double ShockPropagationOverride = 0.0;
+	UPROPERTY(VisibleAnywhere, Category="General")
+	double LinearDriveStiffnessOverride = 0.0;
+	UPROPERTY(VisibleAnywhere, Category="General")
+	double LinearDriveDampingOverride = 0.0;
+	UPROPERTY(VisibleAnywhere, Category="General")
+	double AngularDriveStiffnessOverride = 0.0;
+	UPROPERTY(VisibleAnywhere, Category="General")
+	double AngularDriveDampingOverride = 0.0;
+	UPROPERTY(VisibleAnywhere, Category="General")
+	double SoftLinearStiffnessOverride = 0.0;
+	UPROPERTY(VisibleAnywhere, Category="General")
+	double SoftLinearDampingOverride = 0.0;
+	UPROPERTY(VisibleAnywhere, Category="General")
+	double SoftTwistStiffnessOverride = 0.0;
+	UPROPERTY(VisibleAnywhere, Category="General")
+	double SoftTwistDampingOverride = 0.0;
+	UPROPERTY(VisibleAnywhere, Category="General")
+	double SoftSwingStiffnessOverride = 0.0;
+	UPROPERTY(VisibleAnywhere, Category="General")
+	double SoftSwingDampingOverride = 0.0;
+};
+CVD_IMPLEMENT_SERIALIZER(FChaosVDJointSolverSettingsDataWrapper)
+
+UENUM()
+enum class EChaosVDJointMotionType : int32
+{
+	Free,
+	Limited,
+	Locked,
+};
+
+UENUM()
+enum class EChaosVDJointForceMode : int32
+{
+	Acceleration,
+	Force,
+};
+
+UENUM()
+enum class EChaosVDPlasticityType : int32
+{
+	Free,
+	Shrink,
+	Grow,
+};
+
+UENUM()
+enum class EChaosVDJointSettingsFlags : uint32
+{
+	None = 0,
+	CollisionEnabled = 1 << 0,
+	MassConditioningEnabled = 1 << 1,
+	AngularSLerpPositionDriveEnabled = 1 << 2,
+	AngularSLerpVelocityDriveEnabled = 1 << 3,
+	AngularTwistPositionDriveEnabled = 1 << 4,
+	AngularTwistVelocityDriveEnabled = 1 << 5,
+	AngularSwingPositionDriveEnabled = 1 << 6,
+	AngularSwingVelocityDriveEnabled = 1 << 7,
+	SoftLinearLimitsEnabled= 1 << 8,
+	SoftTwistLimitsEnabled = 1 << 9,
+	SoftSwingLimitsEnabled = 1 << 10,
+	LinearPositionDriveEnabled0 = 1 << 11,
+	LinearPositionDriveEnable1 = 1 << 12,
+	LinearPositionDriveEnable2 = 1 << 13,
+	LinearVelocityDriveEnabled0 = 1 << 14,
+	LinearVelocityDriveEnabled1 = 1 << 15,
+	LinearVelocityDriveEnabled2 = 1 << 16,
+};
+
+USTRUCT()
+struct FChaosVDJointSettingsDataWrapper : public FChaosVDWrapperDataBase
+{
+	GENERATED_BODY()
+public:
+	
+	CHAOSVDRUNTIME_API bool Serialize(FArchive& Ar);
+
+	UPROPERTY(VisibleAnywhere, Category="General")
+	FTransform ConnectorTransforms[2];
+	UPROPERTY(VisibleAnywhere, Category="General")
+	double Stiffness;
+	UPROPERTY(VisibleAnywhere, Category="General")
+	double LinearProjection;
+	UPROPERTY(VisibleAnywhere, Category="General")
+	double AngularProjection;
+	UPROPERTY(VisibleAnywhere, Category="General")
+	double ShockPropagation;
+	UPROPERTY(VisibleAnywhere, Category="General")
+	double TeleportDistance;
+	UPROPERTY(VisibleAnywhere, Category="General")
+	double TeleportAngle;
+	UPROPERTY(VisibleAnywhere, Category="General")
+	double ParentInvMassScale;
+
+	UPROPERTY(VisibleAnywhere, Category="General")
+	uint16 bCollisionEnabled : 1 = false;
+	UPROPERTY(VisibleAnywhere, Category="General")
+	uint16 bMassConditioningEnabled : 1 = false;
+	UPROPERTY(VisibleAnywhere, Category="General")
+	uint16 bSoftLinearLimitsEnabled : 1 = false;
+	UPROPERTY(VisibleAnywhere, Category="General")
+	uint16 bSoftTwistLimitsEnabled : 1 = false;
+	UPROPERTY(VisibleAnywhere, Category="General")
+	uint16 bSoftSwingLimitsEnabled : 1 = false;
+	UPROPERTY(VisibleAnywhere, Category="General")
+	uint16 bAngularSLerpPositionDriveEnabled : 1 = false;
+	UPROPERTY(VisibleAnywhere, Category="General")
+	uint16 bAngularSLerpVelocityDriveEnabled : 1 = false;
+	UPROPERTY(VisibleAnywhere, Category="General")
+	uint16 bAngularTwistPositionDriveEnabled : 1 = false;
+	UPROPERTY(VisibleAnywhere, Category="General")
+	uint16 bAngularTwistVelocityDriveEnabled : 1 = false;
+	UPROPERTY(VisibleAnywhere, Category="General")
+	uint16 bAngularSwingPositionDriveEnabled : 1 = false;
+	UPROPERTY(VisibleAnywhere, Category="General")
+	uint16 bAngularSwingVelocityDriveEnabled : 1 = false;
+
+	UPROPERTY(VisibleAnywhere, Category="General")
+	EChaosVDJointMotionType LinearMotionTypes[3];
+	UPROPERTY(VisibleAnywhere, Category="General")
+	double LinearLimit;
+	UPROPERTY(VisibleAnywhere, Category="General")
+	EChaosVDJointMotionType AngularMotionTypes[3];
+	UPROPERTY(VisibleAnywhere, Category="General")
+	FVector AngularLimits;
+
+	UPROPERTY(VisibleAnywhere, Category="General")
+	EChaosVDJointForceMode LinearSoftForceMode;
+	UPROPERTY(VisibleAnywhere, Category="General")
+	EChaosVDJointForceMode AngularSoftForceMode;
+	UPROPERTY(VisibleAnywhere, Category="General")
+	double SoftLinearStiffness;
+	UPROPERTY(VisibleAnywhere, Category="General")
+	double SoftLinearDamping;
+	UPROPERTY(VisibleAnywhere, Category="General")
+	double SoftTwistStiffness;
+	UPROPERTY(VisibleAnywhere, Category="General")
+	double SoftTwistDamping;
+	UPROPERTY(VisibleAnywhere, Category="General")
+	double SoftSwingStiffness;
+	UPROPERTY(VisibleAnywhere, Category="General")
+	double SoftSwingDamping;
+	UPROPERTY(VisibleAnywhere, Category="General")
+	double LinearRestitution;
+	UPROPERTY(VisibleAnywhere, Category="General")
+	double TwistRestitution;
+	UPROPERTY(VisibleAnywhere, Category="General")
+	double SwingRestitution;
+
+	UPROPERTY(VisibleAnywhere, Category="General")
+	double LinearContactDistance;
+	UPROPERTY(VisibleAnywhere, Category="General")
+	double TwistContactDistance;
+	UPROPERTY(VisibleAnywhere, Category="General")
+	double SwingContactDistance;
+
+	UPROPERTY(VisibleAnywhere, Category="General")
+	FVector LinearDrivePositionTarget;
+	UPROPERTY(VisibleAnywhere, Category="General")
+	FVector LinearDriveVelocityTarget;
+	
+	UPROPERTY(VisibleAnywhere, Category="General")
+	uint8 bLinearPositionDriveEnabled0 : 1 = false;
+	UPROPERTY(VisibleAnywhere, Category="General")
+	uint8 bLinearPositionDriveEnabled1 : 1 = false;
+	UPROPERTY(VisibleAnywhere, Category="General")
+	uint8 bLinearPositionDriveEnabled2 : 1 = false;
+	UPROPERTY(VisibleAnywhere, Category="General")
+	uint8 bLinearVelocityDriveEnabled0 : 1 = false;
+	UPROPERTY(VisibleAnywhere, Category="General")
+	uint8 bLinearVelocityDriveEnabled1 : 1 = false;
+	UPROPERTY(VisibleAnywhere, Category="General")
+	uint8 bLinearVelocityDriveEnabled2 : 1 = false;
+	
+	UPROPERTY(VisibleAnywhere, Category="General")
+	EChaosVDJointForceMode LinearDriveForceMode;
+	UPROPERTY(VisibleAnywhere, Category="General")
+	FVector LinearDriveStiffness;
+	UPROPERTY(VisibleAnywhere, Category="General")
+	FVector LinearDriveDamping;
+	UPROPERTY(VisibleAnywhere, Category="General")
+	FVector LinearDriveMaxForce;
+	UPROPERTY(VisibleAnywhere, Category="General")
+	FQuat AngularDrivePositionTarget;
+	UPROPERTY(VisibleAnywhere, Category="General")
+	FVector AngularDriveVelocityTarget;
+	UPROPERTY(VisibleAnywhere, Category="General")
+	EChaosVDJointForceMode AngularDriveForceMode;
+	UPROPERTY(VisibleAnywhere, Category="General")
+	FVector AngularDriveStiffness;
+	UPROPERTY(VisibleAnywhere, Category="General")
+	FVector AngularDriveDamping;
+	UPROPERTY(VisibleAnywhere, Category="General")
+	FVector AngularDriveMaxTorque;
+	UPROPERTY(VisibleAnywhere, Category="General")
+	double LinearBreakForce;
+	UPROPERTY(VisibleAnywhere, Category="General")
+	double LinearPlasticityLimit;
+	UPROPERTY(VisibleAnywhere, Category="General")
+	EChaosVDPlasticityType LinearPlasticityType;
+	UPROPERTY(VisibleAnywhere, Category="General")
+	double LinearPlasticityInitialDistanceSquared;
+	UPROPERTY(VisibleAnywhere, Category="General")
+	double AngularBreakTorque;
+	UPROPERTY(VisibleAnywhere, Category="General")
+	double AngularPlasticityLimit;
+	UPROPERTY(VisibleAnywhere, Category="General")
+	double ContactTransferScale;
+};
+
+CVD_IMPLEMENT_SERIALIZER(FChaosVDJointSettingsDataWrapper)
+
+USTRUCT()
+struct FChaosVDJointConstraint : public FChaosVDWrapperDataBase
+{
+	GENERATED_BODY()
+public:
+
+	inline static FStringView WrapperTypeName = TEXT("FChaosVDJointConstraint");
+
+	CHAOSVDRUNTIME_API bool Serialize(FArchive& Ar);
+
+	int32 SolverID = INDEX_NONE;
+
+	UPROPERTY(VisibleAnywhere, Category="General")
+	int32 ConstraintIndex = INDEX_NONE;
+
+	int32 ParticleParIndexes[2] = { INDEX_NONE, INDEX_NONE };
+
+	UPROPERTY(VisibleAnywhere, Category=JointState)
+	FChaosVDJointStateDataWrapper JointState;
+	
+	UPROPERTY(VisibleAnywhere, Category=JointSettings)
+	FChaosVDJointSettingsDataWrapper JointSettings;
+
+	bool bIsSelectedInEditor = false;
+};
+
+CVD_IMPLEMENT_SERIALIZER(FChaosVDJointConstraint)
+
+
+
