[ChaosVD] Fix for visibility control issues

- Added a visibility button to the solver tracks control widget. This can be used to hide all particles of a specific solver. The particles will remain hidden even if they are removed and then added while scrubbing, if the solver itself is still requested to be hidden
- All solver particles can still be hidden from the outliner (the new widget will reflect the correct state)
- You can hide all particles from a solver from the new widget and the un-hide them from the outliner (and vice versa)
- When particles from a solver are hidden, the scene outliner will disable the visibility button and use a tooltip explaining another visibility setting is controlling the visible state.

#tests Tested it in PIE by doing a CVD recording with 2 clients (meaning we will have multiple solvers), and playing it back while changing the visibility options (meaning we will have multiple solvers), and playing it back while changing the visibility options

diff --git a/Engine/Plugins/Experimental/ChaosVD/Source/ChaosVD/Private/Actors/ChaosVDSolverInfoActor.cpp b/Engine/Plugins/Experimental/ChaosVD/Source/ChaosVD/Private/Actors/ChaosVDSolverInfoActor.cpp
--- a/Engine/Plugins/Experimental/ChaosVD/Source/ChaosVD/Private/Actors/ChaosVDSolverInfoActor.cpp
+++ b/Engine/Plugins/Experimental/ChaosVD/Source/ChaosVD/Private/Actors/ChaosVDSolverInfoActor.cpp
@@ -23,6 +23,12 @@ AChaosVDSolverInfoActor::AChaosVDSolverInfoActor(const FObjectInitializer& Objec
 	bIsServer = false;
 }
 
+void AChaosVDSolverInfoActor::SetSolverName(const FString& InSolverName)
+{
+	SolverName = InSolverName;
+	SetActorLabel(TEXT("Solver Data Container | ") + InSolverName);
+}
+
 void AChaosVDSolverInfoActor::SetScene(TWeakPtr<FChaosVDScene> InScene)
 {
 	FChaosVDSceneObjectBase::SetScene(InScene);
@@ -43,6 +49,11 @@ void AChaosVDSolverInfoActor::RegisterParticleActor(int32 ParticleID, AChaosVDPa
 	if (!SolverParticlesByID.Contains(ParticleID))
 	{
 		SolverParticlesByID.Add(ParticleID, ParticleActor);
+
+#if WITH_EDITOR
+		ApplySolverVisibilityToParticle(ParticleActor, IsTemporarilyHiddenInEditor());
+#endif
+
 	}
 
 	ParticleActor->SetFolderPath(GetFolderPathForParticleType(ParticleActor->GetParticleData()->Type));
@@ -132,6 +143,54 @@ void AChaosVDSolverInfoActor::RemoveSolverFolders(UWorld* World)
 	}
 }
 
+void AChaosVDSolverInfoActor::ApplySolverVisibilityToParticle(AChaosVDParticleActor* ParticleActor, bool bIsHidden)
+{
+	if (!ParticleActor)
+	{
+		return;
+	}
+
+	if (bIsHidden)
+	{
+		// Note: We should probably add a priority system for the hide requests
+		// For now just clear the HideBySceneOutliner flag when a hide by solver request is done as this has priority
+		ParticleActor->RemoveHiddenFlag(EChaosVDHideParticleFlags::HiddenBySceneOutliner);
+		ParticleActor->AddHiddenFlag(EChaosVDHideParticleFlags::HiddenBySolverVisibility);
+	}
+	else
+	{
+		ParticleActor->RemoveHiddenFlag(EChaosVDHideParticleFlags::HiddenBySolverVisibility);
+	}
+}
+
+#if WITH_EDITOR
+void AChaosVDSolverInfoActor::SetIsTemporarilyHiddenInEditor(bool bIsHidden)
+{
+	if (IsTemporarilyHiddenInEditor() != bIsHidden)
+	{
+		Super::SetIsTemporarilyHiddenInEditor(bIsHidden);
+
+		constexpr float AmountOfWork = 1.0f;
+		const float PercentagePerElement = 1.0f / SolverParticlesByID.Num();
+
+		FScopedSlowTask VisibilityUpdateProgress(AmountOfWork, LOCTEXT("UpdatingParticlesVisisibility", "Updating Particles Visibility ..."));
+		VisibilityUpdateProgress.MakeDialog();
+
+		for (const TPair<int32, AChaosVDParticleActor*>& ParticleVDInstanceWithID : SolverParticlesByID)
+		{
+			ApplySolverVisibilityToParticle(ParticleVDInstanceWithID.Value, bIsHidden);
+
+			VisibilityUpdateProgress.EnterProgressFrame(PercentagePerElement);
+		}
+
+		if (const TSharedPtr<FChaosVDScene> CVDScene = SceneWeakPtr.Pin())
+		{
+			CVDScene->OnSolverVisibilityUpdated().Broadcast(SolverID, !bIsHidden);
+		}
+	}
+}
+#endif
+
 void AChaosVDSolverInfoActor::Destroyed()
 {
 	UWorld* World = GetWorld();
diff --git a/Engine/Plugins/Experimental/ChaosVD/Source/ChaosVD/Private/ChaosVDParticleActor.cpp b/Engine/Plugins/Experimental/ChaosVD/Source/ChaosVD/Private/ChaosVDParticleActor.cpp
--- a/Engine/Plugins/Experimental/ChaosVD/Source/ChaosVD/Private/ChaosVDParticleActor.cpp
+++ b/Engine/Plugins/Experimental/ChaosVD/Source/ChaosVD/Private/ChaosVDParticleActor.cpp
@@ -332,17 +332,17 @@ bool AChaosVDParticleActor::IsSelectedInEditor() const
 	return false;
 }
 
-#if WITH_EDITOR
 void AChaosVDParticleActor::SetIsTemporarilyHiddenInEditor(bool bIsHidden)
 {
-	Super::SetIsTemporarilyHiddenInEditor(bIsHidden);
-
-	VisitGeometryInstances([this, bIsHidden](const TSharedRef<FChaosVDMeshDataInstanceHandle>& MeshDataHandle)
+	if (bIsHidden)
 	{
-		FChaosVDGeometryComponentUtils::UpdateMeshVisibility(MeshDataHandle, ParticleDataPtr ? *ParticleDataPtr.Get() : FChaosVDParticleDataWrapper(), IsActive() && !bIsHidden);
-	});
+		AddHiddenFlag(EChaosVDHideParticleFlags::HiddenBySceneOutliner);
+	}
+	else
+	{
+		RemoveHiddenFlag(EChaosVDHideParticleFlags::HiddenBySceneOutliner);
+	}
 }
-#endif //WITH_EDITOR
 
 FBox AChaosVDParticleActor::GetComponentsBoundingBox(bool bNonColliding, bool bIncludeFromChildActors) const
 {
@@ -457,7 +457,7 @@ void AChaosVDParticleActor::UpdateGeometryComponentsVisibility()
 	{
 		if (ParticleDataPtr)
 		{
-			FChaosVDGeometryComponentUtils::UpdateMeshVisibility(MeshDataHandle, *ParticleDataPtr.Get(), IsActive());
+			FChaosVDGeometryComponentUtils::UpdateMeshVisibility(MeshDataHandle, *ParticleDataPtr.Get(), IsVisible());
 		}
 	});
 }
@@ -485,6 +485,15 @@ void AChaosVDParticleActor::SetIsActive(bool bNewActive)
 		bEditable = bNewActive;
 		bListedInSceneOutliner = bNewActive;
 
+		if (bNewActive)
+		{
+			RemoveHiddenFlag(EChaosVDHideParticleFlags::HiddenByActiveState);
+		}
+		else
+		{
+			AddHiddenFlag(EChaosVDHideParticleFlags::HiddenByActiveState);
+		}
+
 		UpdateGeometryComponentsVisibility();
 
 #endif
@@ -496,4 +505,24 @@ void AChaosVDParticleActor::SetIsActive(bool bNewActive)
 	}
 }
 
+void AChaosVDParticleActor::AddHiddenFlag(EChaosVDHideParticleFlags Flag)
+{
+	EnumAddFlags(HideParticleFlags, Flag);
+
+	VisitGeometryInstances([this](const TSharedRef<FChaosVDMeshDataInstanceHandle>& MeshDataHandle)
+	{
+		FChaosVDGeometryComponentUtils::UpdateMeshVisibility(MeshDataHandle, ParticleDataPtr ? *ParticleDataPtr.Get() : FChaosVDParticleDataWrapper(), IsVisible());
+	});
+}
+
+void AChaosVDParticleActor::RemoveHiddenFlag(EChaosVDHideParticleFlags Flag)
+{
+	EnumRemoveFlags(HideParticleFlags, Flag);
+
+	VisitGeometryInstances([this](const TSharedRef<FChaosVDMeshDataInstanceHandle>& MeshDataHandle)
+	{
+		FChaosVDGeometryComponentUtils::UpdateMeshVisibility(MeshDataHandle, ParticleDataPtr ? *ParticleDataPtr.Get() : FChaosVDParticleDataWrapper(), IsVisible());
+	});
+}
+
 #endif //WITH_EDITOR
diff --git a/Engine/Plugins/Experimental/ChaosVD/Source/ChaosVD/Private/ChaosVDPlaybackController.cpp b/Engine/Plugins/Experimental/ChaosVD/Source/ChaosVD/Private/ChaosVDPlaybackController.cpp
--- a/Engine/Plugins/Experimental/ChaosVD/Source/ChaosVD/Private/ChaosVDPlaybackController.cpp
+++ b/Engine/Plugins/Experimental/ChaosVD/Source/ChaosVD/Private/ChaosVDPlaybackController.cpp
@@ -8,6 +8,7 @@
 #include "ChaosVDRecording.h"
 #include "ChaosVDRuntimeModule.h"
 #include "ChaosVDScene.h"
+#include "Actors/ChaosVDSolverInfoActor.h"
 #include "Trace/ChaosVDTraceManager.h"
 #include "Trace/ChaosVDTraceProvider.h"
 #include "TraceServices/Model/AnalysisSession.h"
@@ -736,6 +737,28 @@ float FChaosVDPlaybackController::GetFrameTimeForTrack(EChaosVDTrackType TrackTy
 	return CurrentTargetFrameTime;
 }
 
+void FChaosVDPlaybackController::UpdateTrackVisibility(EChaosVDTrackType Type, int32 TrackID, bool bNewVisibility)
+{
+	switch (Type)
+	{
+		case EChaosVDTrackType::Solver:
+			{
+				if (const TSharedPtr<FChaosVDScene> ScenePtr = SceneToControl.Pin())
+				{
+					if (AChaosVDSolverInfoActor* SolverActorInfo = ScenePtr->GetSolverInfoActor(TrackID))
+					{
+						SolverActorInfo->SetIsTemporarilyHiddenInEditor(!bNewVisibility);
+					}
+				}
+				break;
+			}
+		case EChaosVDTrackType::Game:
+		default:
+			ensure(false);
+			break;
+	}
+}
+
 void FChaosVDPlaybackController::UpdateSolverTracksData()
 {
 	const TMap<int32, TArray<FChaosVDSolverFrameData>>& SolversByID = LoadedRecording->GetAvailableSolvers_AssumesLocked();
diff --git a/Engine/Plugins/Experimental/ChaosVD/Source/ChaosVD/Private/ChaosVDScene.cpp b/Engine/Plugins/Experimental/ChaosVD/Source/ChaosVD/Private/ChaosVDScene.cpp
--- a/Engine/Plugins/Experimental/ChaosVD/Source/ChaosVD/Private/ChaosVDScene.cpp
+++ b/Engine/Plugins/Experimental/ChaosVD/Source/ChaosVD/Private/ChaosVDScene.cpp
@@ -269,12 +269,14 @@ void FChaosVDScene::CreateSolverInfoActor(int32 SolverID)
 		AChaosVDSolverInfoActor* SolverDataInfo = PhysicsVDWorld->SpawnActor<AChaosVDSolverInfoActor>();
 		check(SolverDataInfo);
 
-		const FName FolderPath("ChaosVisualDebugger/SolverDataContainer");
-		SolverDataInfo->SetFolderPath(FolderPath);
-
 		FString SolverName = LoadedRecording->GetSolverName_AssumedLocked(SolverID);
 		const bool bIsServer = SolverName.Contains(TEXT("Server"));
 
+		const FStringFormatOrderedArguments Args {SolverName, FString::FromInt(SolverID)};
+		const FName FolderPath = *FString::Format(TEXT("Solver {0} | ID {1}"), Args);
+
+		SolverDataInfo->SetFolderPath(FolderPath);
+
 		SolverDataInfo->SetSolverID(SolverID);
 		SolverDataInfo->SetSolverName(SolverName);
 		SolverDataInfo->SetScene(AsWeak());
diff --git a/Engine/Plugins/Experimental/ChaosVD/Source/ChaosVD/Private/ChaosVDWorldOutlinerMode.cpp b/Engine/Plugins/Experimental/ChaosVD/Source/ChaosVD/Private/ChaosVDWorldOutlinerMode.cpp
--- a/Engine/Plugins/Experimental/ChaosVD/Source/ChaosVD/Private/ChaosVDWorldOutlinerMode.cpp
+++ b/Engine/Plugins/Experimental/ChaosVD/Source/ChaosVD/Private/ChaosVDWorldOutlinerMode.cpp
@@ -13,6 +13,97 @@ static FAutoConsoleVariable CVarChaosVDQueueAndCombineSceneOutlinerEvents(
 	true,
 	TEXT("If set to true, scene outliner events will be queued and sent once per frame. If there was a unprocessed event for an item, the las queued event will replace it"));
 
+const FSceneOutlinerTreeItemType FChaosVDActorTreeItem::Type(&FActorTreeItem::Type);
+
+#define LOCTEXT_NAMESPACE "ChaosVisualDebugger"
+
+const TSharedRef<SWidget> FChaosVDSceneOutlinerGutter::ConstructRowWidget(FSceneOutlinerTreeItemRef TreeItem, const STableRow<FSceneOutlinerTreeItemPtr>& Row)
+{
+	if (TreeItem->ShouldShowVisibilityState())
+	{
+		return SNew(SHorizontalBox)
+			+SHorizontalBox::Slot()
+			.AutoWidth()
+			.VAlign(VAlign_Center)
+			[
+				SNew(SVisibilityWidget, SharedThis(this), WeakOutliner, TreeItem, &Row)
+				.IsEnabled_Raw(this, &FChaosVDSceneOutlinerGutter::IsEnabled, TreeItem->AsWeak())
+				.ToolTipText_Raw(this, &FChaosVDSceneOutlinerGutter::GetVisibilityTooltip, TreeItem->AsWeak())
+			];
+	}
+	return SNullWidget::NullWidget;
+}
+
+FText FChaosVDSceneOutlinerGutter::GetVisibilityTooltip(TWeakPtr<ISceneOutlinerTreeItem> WeakTreeItem) const
+{
+
+	return IsEnabled(WeakTreeItem) ? LOCTEXT("SceneOutlinerVisibilityToggleTooltip", "Toggles the visibility of this object in the level editor.") :
+										LOCTEXT("SceneOutlinerVisibilityToggleDisabkedTooltip", "Visibility of this object is being controlled by another visibility setting");
+}
+
+bool FChaosVDSceneOutlinerGutter::IsEnabled(TWeakPtr<ISceneOutlinerTreeItem> WeakTreeItem) const
+{
+	bool bIsEnabled = true;
+	if (const TSharedPtr<ISceneOutlinerTreeItem> TreeItem = WeakTreeItem.Pin())
+	{
+		if (const FActorTreeItem* ActorItem = TreeItem->CastTo<FActorTreeItem>())
+		{
+			if (const AChaosVDParticleActor* CVDActor = Cast<AChaosVDParticleActor>(ActorItem->Actor.Get()))
+			{
+				bIsEnabled = CVDActor->GetHideFlags() == EChaosVDHideParticleFlags::HiddenBySceneOutliner || CVDActor->GetHideFlags() == EChaosVDHideParticleFlags::None;
+			}
+		}
+	}
+
+	return bIsEnabled;
+}
+
+bool FChaosVDActorTreeItem::GetVisibility() const
+{
+	if (const AChaosVDParticleActor* CVDActor = Cast<AChaosVDParticleActor>(Actor.Get()))
+	{
+		return CVDActor->IsVisible();
+	}
+	else
+	{
+		return FActorTreeItem::GetVisibility();
+	}
+}
+
+void FChaosVDActorTreeItem::OnVisibilityChanged(const bool bNewVisibility)
+{
+	if (AChaosVDParticleActor* CVDActor = Cast<AChaosVDParticleActor>(Actor.Get()))
+	{
+		if (bNewVisibility)
+		{
+			CVDActor->RemoveHiddenFlag(EChaosVDHideParticleFlags::HiddenBySceneOutliner);
+		}
+		else
+		{
+			CVDActor->AddHiddenFlag(EChaosVDHideParticleFlags::HiddenBySceneOutliner);
+		}
+	}
+	else
+	{
+		FActorTreeItem::OnVisibilityChanged(bNewVisibility);
+	}
+}
+
+TUniquePtr<FChaosVDOutlinerHierarchy> FChaosVDOutlinerHierarchy::Create(ISceneOutlinerMode* Mode, const TWeakObjectPtr<UWorld>& World)
+{
+	FChaosVDOutlinerHierarchy* Hierarchy = new FChaosVDOutlinerHierarchy(Mode, World);
+
+	Create_Internal(Hierarchy, World);
+
+	return TUniquePtr<FChaosVDOutlinerHierarchy>(Hierarchy);
+}
+
+
+FSceneOutlinerTreeItemPtr FChaosVDOutlinerHierarchy::CreateItemForActor(AActor* InActor, bool bForce) const
+{
+	return Mode->CreateItemFor<FChaosVDActorTreeItem>(InActor, bForce);
+}
+
 FChaosVDWorldOutlinerMode::FChaosVDWorldOutlinerMode(const FActorModeParams& InModeParams, TWeakPtr<FChaosVDScene> InScene)
 	: FActorMode(InModeParams),
 	CVDScene(InScene)
@@ -95,6 +186,19 @@ bool FChaosVDWorldOutlinerMode::Tick(float DeltaTime)
 	return true;
 }
 
+TUniquePtr<ISceneOutlinerHierarchy> FChaosVDWorldOutlinerMode::CreateHierarchy()
+{
+	TUniquePtr<FChaosVDOutlinerHierarchy> ActorHierarchy = FChaosVDOutlinerHierarchy::Create(this, RepresentingWorld);
+
+	ActorHierarchy->SetShowingComponents(!bHideComponents);
+	ActorHierarchy->SetShowingOnlyActorWithValidComponents(!bHideComponents && bHideActorWithNoComponent);
+	ActorHierarchy->SetShowingLevelInstances(!bHideLevelInstanceHierarchy);
+	ActorHierarchy->SetShowingUnloadedActors(!bHideUnloadedActors);
+	ActorHierarchy->SetShowingEmptyFolders(!bHideEmptyFolders);
+	
+	return ActorHierarchy;
+}
+
 void FChaosVDWorldOutlinerMode::EnqueueAndCombineHierarchyEvent(const FSceneOutlinerTreeItemID& ItemID, const FSceneOutlinerHierarchyChangedData& EnventToProcess)
 {
 	if (FSceneOutlinerHierarchyChangedData* EventData = PendingOutlinerEventsMap.Find(ItemID))
@@ -116,7 +220,7 @@ void FChaosVDWorldOutlinerMode::HandleActorLabelChanged(AActor* ChangedActor)
 	if (IsActorDisplayable(ChangedActor) && RepresentingWorld.Get() == ChangedActor->GetWorld())
 	{
 		// Force create the item otherwise the outliner may not be notified of a change to the item if it is filtered out
-		if (FSceneOutlinerTreeItemPtr Item = CreateItemFor<FActorTreeItem>(ChangedActor, true))
+		if (FSceneOutlinerTreeItemPtr Item = CreateItemFor<FChaosVDActorTreeItem>(ChangedActor, true))
 		{
 			SceneOutliner->OnItemLabelChanged(Item);
 		}
@@ -137,7 +241,7 @@ void FChaosVDWorldOutlinerMode::HandleActorActiveStateChanged(AChaosVDParticleAc
 		if (ChangedActor->IsActive())
 		{
 			EventData.Type = FSceneOutlinerHierarchyChangedData::Added;
-			EventData.Items.Emplace(CreateItemFor<FActorTreeItem>(ChangedActor));
+			EventData.Items.Emplace(CreateItemFor<FChaosVDActorTreeItem>(ChangedActor));
 		}
 		else
 		{
@@ -181,3 +285,5 @@ void FChaosVDWorldOutlinerMode::HandlePostSelectionChange(const UTypedElementSel
 		}
 	}
 }
+
+#undef LOCTEXT_NAMESPACE
diff --git a/Engine/Plugins/Experimental/ChaosVD/Source/ChaosVD/Private/ChaosVDWorldOutlinerMode.h b/Engine/Plugins/Experimental/ChaosVD/Source/ChaosVD/Private/ChaosVDWorldOutlinerMode.h
--- a/Engine/Plugins/Experimental/ChaosVD/Source/ChaosVD/Private/ChaosVDWorldOutlinerMode.h
+++ b/Engine/Plugins/Experimental/ChaosVD/Source/ChaosVD/Private/ChaosVDWorldOutlinerMode.h
@@ -2,13 +2,70 @@
 
 #pragma once
 #include "ActorBrowsingMode.h"
+#include "ActorHierarchy.h"
+#include "ActorTreeItem.h"
 #include "ChaosVDScene.h"
 #include "ChaosVDSceneSelectionObserver.h"
+#include "SceneOutlinerGutter.h"
 #include "Containers/Map.h"
 #include "Templates/SharedPointer.h"
 
 class FChaosVDScene;
 
+/** Outliner Gutter used to override the behaviour of the visibility widget in the Scene outliner  */
+class FChaosVDSceneOutlinerGutter final : public FSceneOutlinerGutter
+{
+public:
+	explicit FChaosVDSceneOutlinerGutter(ISceneOutliner& Outliner)
+		: FSceneOutlinerGutter(Outliner)
+	{
+	}
+
+	virtual const TSharedRef<SWidget> ConstructRowWidget(FSceneOutlinerTreeItemRef TreeItem, const STableRow<FSceneOutlinerTreeItemPtr>& Row) override;
+	
+	FText GetVisibilityTooltip(TWeakPtr<ISceneOutlinerTreeItem> WeakTreeItem) const;
+
+	bool IsEnabled(TWeakPtr<ISceneOutlinerTreeItem> WeakTreeItem) const;
+};
+
+/** Scene outliner Actor Tree Item used to override the behaviour of the visibility changes for CVD Actors in the Scene outliner  */
+class FChaosVDActorTreeItem final : public FActorTreeItem
+{
+public:
+	
+	FChaosVDActorTreeItem(AActor* InActor) : FActorTreeItem(InActor)
+	{
+		
+	};
+
+	FChaosVDActorTreeItem(const FSceneOutlinerTreeItemType& TypeIn, AActor* InActor)
+		: FActorTreeItem(TypeIn, InActor)
+	{
+	}
+
+	virtual bool GetVisibility() const override;
+
+	virtual void OnVisibilityChanged(const bool bNewVisibility) override;
+
+	static const FSceneOutlinerTreeItemType Type;
+};
+
+/** Actor Hierarchy used to override the what Tree items will be used for actors in CVD's Scene outliner  */
+class FChaosVDOutlinerHierarchy final : public FActorHierarchy
+{
+public:
+
+	static TUniquePtr<FChaosVDOutlinerHierarchy> Create(ISceneOutlinerMode* Mode, const TWeakObjectPtr<UWorld>& World);
+
+protected:
+
+	FChaosVDOutlinerHierarchy(ISceneOutlinerMode* Mode, const TWeakObjectPtr<UWorld>& Worlds)
+	: FActorHierarchy(Mode, Worlds)
+	{
+	}
+	virtual FSceneOutlinerTreeItemPtr CreateItemForActor(AActor* InActor, bool bForce) const override;
+};
+
 /**
  * Scene outliner mode used to represent a CVD (Chaos Visual Debugger) world
  * It has a more limited view compared to the normal outliner, hiding features we don't support,
@@ -36,6 +93,8 @@ public:
 
 	virtual bool Tick(float DeltaTime) override;
 
+	virtual TUniquePtr<ISceneOutlinerHierarchy> CreateHierarchy() override;
+
 private:
 
 	void EnqueueAndCombineHierarchyEvent(const FSceneOutlinerTreeItemID& ItemID, const FSceneOutlinerHierarchyChangedData& EnventToProcess);
diff --git a/Engine/Plugins/Experimental/ChaosVD/Source/ChaosVD/Private/ChaosVDWorldOutlinerTab.cpp b/Engine/Plugins/Experimental/ChaosVD/Source/ChaosVD/Private/ChaosVDWorldOutlinerTab.cpp
--- a/Engine/Plugins/Experimental/ChaosVD/Source/ChaosVD/Private/ChaosVDWorldOutlinerTab.cpp
+++ b/Engine/Plugins/Experimental/ChaosVD/Source/ChaosVD/Private/ChaosVDWorldOutlinerTab.cpp
@@ -24,6 +24,7 @@ void FChaosVDWorldOutlinerTab::CreateWorldOutlinerWidget()
 	InitOptions.UseDefaultColumns();
 	InitOptions.ColumnMap.Add(FSceneOutlinerBuiltInColumnTypes::Gutter(), FSceneOutlinerColumnInfo(ESceneOutlinerColumnVisibility::Visible, 0, FCreateSceneOutlinerColumn(), true, TOptional<float>(), FSceneOutlinerBuiltInColumnTypes::Gutter_Localized()));
 	InitOptions.ColumnMap.Add(FSceneOutlinerBuiltInColumnTypes::ActorInfo(), FSceneOutlinerColumnInfo(ESceneOutlinerColumnVisibility::Visible, 20, FCreateSceneOutlinerColumn(), true, TOptional<float>(), FSceneOutlinerBuiltInColumnTypes::ActorInfo_Localized()));
+	InitOptions.ColumnMap.Add(FChaosVDSceneOutlinerGutter::GetID(), FSceneOutlinerColumnInfo(ESceneOutlinerColumnVisibility::Visible, 0, FCreateSceneOutlinerColumn::CreateLambda([](ISceneOutliner& InSceneOutliner) { return MakeShareable(new FChaosVDSceneOutlinerGutter(InSceneOutliner)); })));
 		
 	FCreateSceneOutlinerMode ModeFactory = FCreateSceneOutlinerMode::CreateLambda([this](SSceneOutliner* Outliner)
 	{
diff --git a/Engine/Plugins/Experimental/ChaosVD/Source/ChaosVD/Private/Widgets/SChaosVDPlaybackViewport.cpp b/Engine/Plugins/Experimental/ChaosVD/Source/ChaosVD/Private/Widgets/SChaosVDPlaybackViewport.cpp
--- a/Engine/Plugins/Experimental/ChaosVD/Source/ChaosVD/Private/Widgets/SChaosVDPlaybackViewport.cpp
+++ b/Engine/Plugins/Experimental/ChaosVD/Source/ChaosVD/Private/Widgets/SChaosVDPlaybackViewport.cpp
@@ -232,6 +232,11 @@ void SChaosVDPlaybackViewport::OnPlaybackSceneUpdated()
 	PlaybackViewportClient->bNeedsRedraw = true;	
 }
 
+void SChaosVDPlaybackViewport::OnSolverVisibilityUpdated(int32 SolverID, bool bNewVisibility)
+{
+	PlaybackViewportClient->bNeedsRedraw = true;	
+}
+
 void SChaosVDPlaybackViewport::RegisterNewController(TWeakPtr<FChaosVDPlaybackController> NewController)
 {
 	if (PlaybackController != NewController)
@@ -241,6 +246,7 @@ void SChaosVDPlaybackViewport::RegisterNewController(TWeakPtr<FChaosVDPlaybackCo
 			if (TSharedPtr<FChaosVDScene> ScenePtr = CurrentPlaybackControllerPtr->GetControllerScene().Pin())
 			{
 				ScenePtr->OnSceneUpdated().RemoveAll(this);
+				ScenePtr->OnSolverVisibilityUpdated().RemoveAll(this);
 			}
 		}
 
@@ -251,6 +257,7 @@ void SChaosVDPlaybackViewport::RegisterNewController(TWeakPtr<FChaosVDPlaybackCo
 			if (TSharedPtr<FChaosVDScene> ScenePtr = NewPlaybackControllerPtr->GetControllerScene().Pin())
 			{
 				ScenePtr->OnSceneUpdated().AddRaw(this, &SChaosVDPlaybackViewport::OnPlaybackSceneUpdated);
+				ScenePtr->OnSolverVisibilityUpdated().AddRaw(this, &SChaosVDPlaybackViewport::OnSolverVisibilityUpdated);
 			}
 		}
 	}
diff --git a/Engine/Plugins/Experimental/ChaosVD/Source/ChaosVD/Private/Widgets/SChaosVDSolverPlaybackControls.cpp b/Engine/Plugins/Experimental/ChaosVD/Source/ChaosVD/Private/Widgets/SChaosVDSolverPlaybackControls.cpp
--- a/Engine/Plugins/Experimental/ChaosVD/Source/ChaosVD/Private/Widgets/SChaosVDSolverPlaybackControls.cpp
+++ b/Engine/Plugins/Experimental/ChaosVD/Source/ChaosVD/Private/Widgets/SChaosVDSolverPlaybackControls.cpp
@@ -5,10 +5,12 @@
 #include "ChaosVDEditorSettings.h"
 #include "ChaosVDModule.h"
 #include "ChaosVDPlaybackController.h"
+#include "ChaosVDScene.h"
 #include "Widgets/ChaosVDPlaybackControlsHelper.h"
 #include "Widgets/SBoxPanel.h"
 #include "Widgets/SChaosVDPlaybackViewport.h"
 #include "Widgets/SChaosVDTimelineWidget.h"
+#include "Widgets/Images/SImage.h"
 #include "Widgets/Input/SButton.h"
 #include "Widgets/Text/STextBlock.h"
 
@@ -18,6 +20,12 @@ void SChaosVDSolverPlaybackControls::Construct(const FArguments& InArgs, int32 I
 {
 	SolverID = InSolverID;
 
+	static const FName NAME_VisibleNotHoveredBrush = TEXT("Level.VisibleIcon16x");
+	static const FName NAME_NotVisibleNotHoveredBrush = TEXT("Level.NotVisibleIcon16x");
+
+	SolverVisibleIconBrush = FAppStyle::Get().GetBrush(NAME_VisibleNotHoveredBrush);
+	SolverHiddenIconBrush = FAppStyle::Get().GetBrush(NAME_NotVisibleNotHoveredBrush);
+
 	ChildSlot
 	[
 		SNew(SVerticalBox)
@@ -84,11 +92,22 @@ void SChaosVDSolverPlaybackControls::Construct(const FArguments& InArgs, int32 I
 				]
 				+SVerticalBox::Slot()
 				[
-					SAssignNew(StepsTimelineWidget, SChaosVDTimelineWidget)
-					.ButtonVisibilityFlags(static_cast<uint16>(EChaosVDTimelineElementIDFlags::AllManualStepping))
-					.OnFrameLockStateChanged_Raw(this, &SChaosVDSolverPlaybackControls::HandleLockStateChanged)
-					.OnFrameChanged_Raw(this, &SChaosVDSolverPlaybackControls::OnStepSelectionUpdated)
-					.MaxFrames(0)
+					SNew(SHorizontalBox)
+					+SHorizontalBox::Slot()
+					.FillWidth(0.6f)
+					[
+						
+						SAssignNew(StepsTimelineWidget, SChaosVDTimelineWidget)
+						.ButtonVisibilityFlags(static_cast<uint16>(EChaosVDTimelineElementIDFlags::AllManualStepping))
+						.OnFrameLockStateChanged_Raw(this, &SChaosVDSolverPlaybackControls::HandleLockStateChanged)
+						.OnFrameChanged_Raw(this, &SChaosVDSolverPlaybackControls::OnStepSelectionUpdated)
+						.MaxFrames(0)
+					]
+					+SHorizontalBox::Slot()
+					.AutoWidth()
+					[
+						CreateVisibilityWidget().ToSharedRef()
+					]
 				]
 			]
 		]
@@ -123,6 +142,31 @@ void SChaosVDSolverPlaybackControls::ConditionallyLockPlaybackControl(const TSha
 	StepsTimelineWidget->SetIsLocked(InControllerSharedRef->IsPlayingLiveSession() || !bUserCanControlPlayback);
 }
 
+void SChaosVDSolverPlaybackControls::HandleSolverVisibilityChanged(int32 InSolverID, bool bNewVisibility)
+{
+	if (SolverID != InSolverID)
+	{
+		return;
+	}
+
+	bIsVisible = bNewVisibility;
+}
+
+FReply SChaosVDSolverPlaybackControls::ToggleSolverVisibility() const
+{
+	if (const TSharedPtr<FChaosVDPlaybackController> CurrentPlaybackControllerPtr = PlaybackController.Pin())
+	{
+		CurrentPlaybackControllerPtr->UpdateTrackVisibility(EChaosVDTrackType::Solver, SolverID, !bIsVisible);
+	}
+
+	return FReply::Handled();
+}
+
+const FSlateBrush* SChaosVDSolverPlaybackControls::GetBrushForCurrentVisibility() const
+{
+	return bIsVisible ? SolverVisibleIconBrush : SolverHiddenIconBrush;
+}
+
 void SChaosVDSolverPlaybackControls::HandlePlaybackControllerDataUpdated(TWeakPtr<FChaosVDPlaybackController> InController)
 {
 	if (PlaybackController != InController)
@@ -238,6 +282,18 @@ const FSlateBrush* SChaosVDSolverPlaybackControls::GetFrameTypeBadgeBrush() cons
 	return bIsReSimFrame ? &ButtonStyle.Pressed : FCoreStyle::Get().GetBrush("Border");
 }
 
+TSharedPtr<SWidget> SChaosVDSolverPlaybackControls::CreateVisibilityWidget()
+{
+	return SNew(SButton)
+	.OnClicked_Raw(this, &SChaosVDSolverPlaybackControls::ToggleSolverVisibility)
+	[
+		SNew( SImage )
+		.Image_Raw(this,&SChaosVDSolverPlaybackControls::GetBrushForCurrentVisibility)
+		.DesiredSizeOverride(FVector2D(16.0f,16.0f))
+		.ColorAndOpacity(FSlateColor::UseForeground())
+	];	
+}
+
 void SChaosVDSolverPlaybackControls::OnFrameSelectionUpdated(int32 NewFrameIndex)
 {
 	if (const TSharedPtr<FChaosVDPlaybackController> PlaybackControllerPtr = PlaybackController.Pin())
@@ -268,4 +324,25 @@ void SChaosVDSolverPlaybackControls::OnStepSelectionUpdated(int32 NewStepIndex)
 	}
 }
 
-#undef LOCTEXT_NAMESPACE
\ No newline at end of file
+void SChaosVDSolverPlaybackControls::RegisterNewController(TWeakPtr<FChaosVDPlaybackController> NewController)
+{
+	if (const TSharedPtr<FChaosVDPlaybackController> OldPlaybackControllerPtr = PlaybackController.Pin())
+	{
+		if (TSharedPtr<FChaosVDScene> Scene = OldPlaybackControllerPtr->GetControllerScene().Pin())
+		{
+			Scene->OnSolverVisibilityUpdated().RemoveAll(this);
+		}
+	}
+
+	FChaosVDPlaybackControllerObserver::RegisterNewController(NewController);
+
+	if (const TSharedPtr<FChaosVDPlaybackController> NewPlaybackControllerPtr = PlaybackController.Pin())
+	{
+		if (TSharedPtr<FChaosVDScene> Scene = NewPlaybackControllerPtr->GetControllerScene().Pin())
+		{
+			Scene->OnSolverVisibilityUpdated().AddRaw(this, &SChaosVDSolverPlaybackControls::HandleSolverVisibilityChanged);
+		}
+	}
+}
+
+#undef LOCTEXT_NAMESPACE
diff --git a/Engine/Plugins/Experimental/ChaosVD/Source/ChaosVD/Public/Actors/ChaosVDSolverInfoActor.h b/Engine/Plugins/Experimental/ChaosVD/Source/ChaosVD/Public/Actors/ChaosVDSolverInfoActor.h
--- a/Engine/Plugins/Experimental/ChaosVD/Source/ChaosVD/Public/Actors/ChaosVDSolverInfoActor.h
+++ b/Engine/Plugins/Experimental/ChaosVD/Source/ChaosVD/Public/Actors/ChaosVDSolverInfoActor.h
@@ -28,7 +28,7 @@ public:
 	void SetSolverID(int32 InSolverID) { SolverID = InSolverID; }
 	int32 GetSolverID() const { return SolverID; }
 
-	void SetSolverName(const FString& InSolverName) { SolverName = InSolverName; }
+	void SetSolverName(const FString& InSolverName);
 	const FString& GetSolverName() { return SolverName; }
 
 	void SetIsServer(bool bInIsServer) { bIsServer = bInIsServer; }
@@ -64,10 +64,15 @@ public:
 	void HandleColorsSettingsUpdated();
 	void RemoveSolverFolders(UWorld* World);
 
+#if WITH_EDITOR
+	void SetIsTemporarilyHiddenInEditor(bool bIsHidden) override;
+#endif
 	virtual void Destroyed() override;
 
 protected:
 
+	void ApplySolverVisibilityToParticle(AChaosVDParticleActor* ParticleActor, bool bIsHidden);
+
 	virtual void HandlePostSelectionChange(const UTypedElementSelectionSet* ChangesSelectionSet) override;
 
 	FName GetFolderPathForParticleType(EChaosVDParticleType ParticleType);
@@ -100,6 +105,7 @@ protected:
 
 	UPROPERTY()
 	TObjectPtr<UChaosVDSolverJointConstraintDataComponent> JointsDataComponent;
+
 };
 
 template <typename TCallback>
diff --git a/Engine/Plugins/Experimental/ChaosVD/Source/ChaosVD/Public/ChaosVDParticleActor.h b/Engine/Plugins/Experimental/ChaosVD/Source/ChaosVD/Public/ChaosVDParticleActor.h
--- a/Engine/Plugins/Experimental/ChaosVD/Source/ChaosVD/Public/ChaosVDParticleActor.h
+++ b/Engine/Plugins/Experimental/ChaosVD/Source/ChaosVD/Public/ChaosVDParticleActor.h
@@ -37,6 +37,18 @@ ENUM_CLASS_FLAGS(EChaosVDActorGeometryUpdateFlags)
 
 DECLARE_DELEGATE(FChaosVDParticleDataUpdatedDelegate)
 
+UENUM()
+enum class EChaosVDHideParticleFlags
+{
+	None = 0,
+	HiddenByVisualizationFlags = 1 << 0,
+	HiddenBySceneOutliner = 1 << 1,
+	HiddenByActiveState = 1 << 2,
+	HiddenBySolverVisibility = 1 << 3,
+	
+};
+ENUM_CLASS_FLAGS(EChaosVDHideParticleFlags)
+
 /** Actor used to represent a Chaos Particle in the Visual Debugger's world */
 UCLASS(HideCategories=(Transform))
 class AChaosVDParticleActor : public AActor, public IChaosVDParticleVisualizationDataProvider,
@@ -70,6 +82,17 @@ public:
 	/** Changes the active state of this CVD Particle Actor */
 	void SetIsActive(bool bNewActive);
 
+	void AddHiddenFlag(EChaosVDHideParticleFlags Flag);
+
+	void RemoveHiddenFlag(EChaosVDHideParticleFlags Flag);
+
+	EChaosVDHideParticleFlags GetHideFlags() const { return HideParticleFlags; }
+
+	bool IsVisible() const
+	{
+		return HideParticleFlags == EChaosVDHideParticleFlags::None;
+	}
+
 	/** Returns true if this particle actor is active - Inactive Particle actors are still in the world but with outdated data
 	 * and hidden from the viewport and outliner. They represent particles that were destroyed.
 	 */
@@ -103,9 +126,6 @@ protected:
 	template<typename TTaskCallback>
 	void VisitGeometryInstances(const TTaskCallback& VisitorCallback);
 
-	UPROPERTY(EditAnywhere, Category = "Viewport Visualization Flags")
-	bool bShowDebugText = false;
-
 	TSharedPtr<FChaosVDParticleDataWrapper> ParticleDataPtr;
 
 	FTransform CachedSimulationTransform;
@@ -122,6 +142,8 @@ protected:
 
 	bool bIsServer = false;
 
+	EChaosVDHideParticleFlags HideParticleFlags;
+
 	friend FChaosVDParticleActorCustomization;
 };
 
diff --git a/Engine/Plugins/Experimental/ChaosVD/Source/ChaosVD/Public/ChaosVDPlaybackController.h b/Engine/Plugins/Experimental/ChaosVD/Source/ChaosVD/Public/ChaosVDPlaybackController.h
--- a/Engine/Plugins/Experimental/ChaosVD/Source/ChaosVD/Public/ChaosVDPlaybackController.h
+++ b/Engine/Plugins/Experimental/ChaosVD/Source/ChaosVD/Public/ChaosVDPlaybackController.h
@@ -238,6 +238,8 @@ public:
 	float GetFrameTimeOverride() const;
 	float GetFrameTimeForTrack(EChaosVDTrackType TrackType, int32 TrackID, const FChaosVDTrackInfo& TrackInfo) const;
 
+	void UpdateTrackVisibility(EChaosVDTrackType Type, int32 TrackID, bool bNewVisibility);
+
 protected:
 
 	/** Updates (or adds) solvers data from the loaded recording to the solver tracks */
diff --git a/Engine/Plugins/Experimental/ChaosVD/Source/ChaosVD/Public/ChaosVDScene.h b/Engine/Plugins/Experimental/ChaosVD/Source/ChaosVD/Public/ChaosVDScene.h
--- a/Engine/Plugins/Experimental/ChaosVD/Source/ChaosVD/Public/ChaosVDScene.h
+++ b/Engine/Plugins/Experimental/ChaosVD/Source/ChaosVD/Public/ChaosVDScene.h
@@ -32,6 +32,8 @@ DECLARE_MULTICAST_DELEGATE_OneParam(FChaosVDActorActiveStateUpdateDelegate, ACha
 DECLARE_MULTICAST_DELEGATE_OneParam(FChaosVDOnObjectSelectedDelegate, UObject*)
 DECLARE_MULTICAST_DELEGATE_OneParam(FChaosVDSolverInfoActorCreatedDelegate, AChaosVDSolverInfoActor*)
 
+DECLARE_MULTICAST_DELEGATE_TwoParams(FChaosVDSolverVisibilityChangedDelegate, int32 SolverID, bool bNewVisibility)
+
 /** Recreates a UWorld from a recorded Chaos VD Frame */
 class FChaosVDScene : public FGCObject , public TSharedFromThis<FChaosVDScene>
 {
@@ -115,6 +117,8 @@ public:
 
 	FChaosVDSolverInfoActorCreatedDelegate& OnSolverInfoActorCreated() { return SolverInfoActorCreatedDelegate; }
 
+	FChaosVDSolverVisibilityChangedDelegate& OnSolverVisibilityUpdated() { return SolverVisibilityChangedDelegate; }
+
 	TSharedPtr<FChaosVDRecording> LoadedRecording;
 
 private:
@@ -195,4 +199,6 @@ private:
 	FDelegateHandle ActorDestroyedHandle;
 
 	FChaosVDSolverInfoActorCreatedDelegate SolverInfoActorCreatedDelegate;
+
+	FChaosVDSolverVisibilityChangedDelegate SolverVisibilityChangedDelegate;
 };
diff --git a/Engine/Plugins/Experimental/ChaosVD/Source/ChaosVD/Public/Widgets/SChaosVDPlaybackViewport.h b/Engine/Plugins/Experimental/ChaosVD/Source/ChaosVD/Public/Widgets/SChaosVDPlaybackViewport.h
--- a/Engine/Plugins/Experimental/ChaosVD/Source/ChaosVD/Public/Widgets/SChaosVDPlaybackViewport.h
+++ b/Engine/Plugins/Experimental/ChaosVD/Source/ChaosVD/Public/Widgets/SChaosVDPlaybackViewport.h
@@ -52,6 +52,7 @@ protected:
 	virtual void HandlePostSelectionChange(const UTypedElementSelectionSet* ChangesSelectionSet) override;
 
 	void OnPlaybackSceneUpdated();
+	void OnSolverVisibilityUpdated(int32 SolverID, bool bNewVisibility);
 
 	void OnFrameSelectionUpdated(int32 NewFrameIndex) const;
 
diff --git a/Engine/Plugins/Experimental/ChaosVD/Source/ChaosVD/Public/Widgets/SChaosVDSolverPlaybackControls.h b/Engine/Plugins/Experimental/ChaosVD/Source/ChaosVD/Public/Widgets/SChaosVDSolverPlaybackControls.h
--- a/Engine/Plugins/Experimental/ChaosVD/Source/ChaosVD/Public/Widgets/SChaosVDSolverPlaybackControls.h
+++ b/Engine/Plugins/Experimental/ChaosVD/Source/ChaosVD/Public/Widgets/SChaosVDSolverPlaybackControls.h
@@ -4,6 +4,7 @@
 
 #include "ChaosVDPlaybackControllerInstigator.h"
 #include "ChaosVDPlaybackControllerObserver.h"
+#include "Styling/ISlateStyle.h"
 #include "Templates/SharedPointer.h"
 #include "Widgets/SCompoundWidget.h"
 
@@ -42,6 +43,8 @@ private:
 	void OnFrameSelectionUpdated(int32 NewFrameIndex);
 	void OnStepSelectionUpdated(int32 NewStepIndex);
 
+	virtual void RegisterNewController(TWeakPtr<FChaosVDPlaybackController> NewController) override;
+
 	virtual void HandlePlaybackControllerDataUpdated(TWeakPtr<FChaosVDPlaybackController> InController) override;
 	void UpdateStepsWidgetForFrame(const FChaosVDPlaybackController& InCurrentPlaybackController, int32 FrameNumber, int32 StepNumber, EChaosVDStepsWidgetUpdateFlags OptionsFlags = EChaosVDStepsWidgetUpdateFlags::Default);
 	virtual void HandleControllerTrackFrameUpdated(TWeakPtr<FChaosVDPlaybackController> InController, const FChaosVDTrackInfo* UpdatedTrackInfo, FGuid InstigatorGuid) override;
@@ -51,12 +54,24 @@ private:
 
 	void ConditionallyLockPlaybackControl(const TSharedRef<FChaosVDPlaybackController>& InControllerSharedRef);
 
+	void HandleSolverVisibilityChanged(int32 InSolverID, bool bNewVisibility);
+
+	FReply ToggleSolverVisibility() const;
+
+	const FSlateBrush* GetBrushForCurrentVisibility() const;
+
 	const FSlateBrush* GetFrameTypeBadgeBrush() const;
 
+	TSharedPtr<SWidget> CreateVisibilityWidget();
+
 	int32 SolverID = INDEX_NONE;
 	bool bIsReSimFrame = false;
 	FString CurrentStepName;
 	TSharedPtr<SChaosVDTimelineWidget> FramesTimelineWidget;
 	TSharedPtr<SChaosVDTimelineWidget> StepsTimelineWidget;
 	bool bStepsLocked = false;
+	bool bIsVisible = true;
+
+	const FSlateBrush* SolverVisibleIconBrush = nullptr;
+	const FSlateBrush* SolverHiddenIconBrush = nullptr;
 };
