[ChaosVD] Fix for ensure when loading some CVD recordings in the tool

The ensure was triggered because some of the recorded geometry containts non-manifold geometry or duplicated triagles. In a large recording with ten of thoudands of geometry, this is the case for just a handful of geometries.

To fix this, similar to how it is done in other tools, I implemented in the CVD side custom logic to generate a dynamic mesh instead of using the default one

CVD dynamic mesh creation logic does the following

 - Only accepts mesh generators with attributes (Normal,UVs, etc, this is true for all CVD geometry).
 - If a non-manifold triangle is detected, we attempt to "repair" visually it by just duplicating the vertices and re-adding the triangle with these
 - If a duplicated triangle is detected, we just skip it and calculate an index offset to continue using the correct attributes from the mesh generator in the following triangles instead of discarding them
 - Discards UV data, as currently CVD shades everything with flat colors

This change can be disabled by setting the cvar  p.Chaos.VD.Tool.UseCVDDynamicMeshGenerator to false

#tests Tested in PIE that geometry is still generated when a CVD recording is loaded, and that the new code can be disabled by setting the old cvar to false

diff --git a/Engine/Plugins/Experimental/ChaosVD/Source/ChaosVD/Private/ChaosVDGeometryBuilder.cpp b/Engine/Plugins/Experimental/ChaosVD/Source/ChaosVD/Private/ChaosVDGeometryBuilder.cpp
--- a/Engine/Plugins/Experimental/ChaosVD/Source/ChaosVD/Private/ChaosVDGeometryBuilder.cpp
+++ b/Engine/Plugins/Experimental/ChaosVD/Source/ChaosVD/Private/ChaosVDGeometryBuilder.cpp
@@ -18,6 +18,144 @@
 #include "Chaos/HeightField.h"
 #include "UObject/UObjectGlobals.h"
 
+namespace Chaos::VisualDebugger
+{
+	namespace Cvars
+	{
+		static bool bUseCVDDynamicMeshGenerator = true;
+		static FAutoConsoleVariableRef CVarUseCVDDynamicMeshGenerator(
+			TEXT("p.Chaos.VD.Tool.UseCVDDynamicMeshGenerator"),
+			bUseCVDDynamicMeshGenerator,
+			TEXT("If true, when creating a dynamic mesh from a mesh generator, CVD will use it's own mesh creation logic which included error handling that tries to repair broken geometry"));
+
+		static bool bDisableUVsSupport = true;
+		static FAutoConsoleVariableRef CVarDisableUVsSupport(
+			TEXT("p.Chaos.VD.Tool.DisableUVsSupport"),
+			bDisableUVsSupport,
+			TEXT("If true, the generated meshes will not have UV data"));
+	}
+
+	void SetTriangleAttributes(const UE::Geometry::FMeshShapeGenerator& Generator, FDynamicMesh3& OutDynamicMesh, int32 AppendedTriangleID, int32 GeneratorTriangleIndex)
+	{
+		UE::Geometry::FDynamicMeshUVOverlay* UVOverlay = OutDynamicMesh.Attributes()->PrimaryUV();
+		UE::Geometry::FDynamicMeshNormalOverlay* NormalOverlay = OutDynamicMesh.Attributes()->PrimaryNormals();
+
+		if (UVOverlay &&Generator.TriangleUVs.IsValidIndex(GeneratorTriangleIndex))
+		{
+			UVOverlay->SetTriangle(AppendedTriangleID, Generator.TriangleUVs[GeneratorTriangleIndex]);
+		}
+		
+		if (ensure(NormalOverlay && Generator.TriangleUVs.IsValidIndex(GeneratorTriangleIndex)))
+		{
+			NormalOverlay->SetTriangle(AppendedTriangleID, Generator.TriangleNormals[GeneratorTriangleIndex]);
+		}
+	}
+
+	void HandleTriangleAddedToDynamicMesh(const UE::Geometry::FMeshShapeGenerator& Generator, FDynamicMesh3& OutDynamicMesh, int32 TriangleIDResult, int32 GroupID, int32 GeneratorTriangleIndex, int32& OutSkippedTriangles, bool bAttemptToFixNoManifoldError = true)
+	{
+		// If we get a triangle ID greater than 0 means the add triangle operation didn't generate an error itself
+		// But we still need to take into account skipped triangles to verify that we have valid data for this triangle in the mesh generator
+		const bool bHasUnhandledError = TriangleIDResult < 0 ? true : (TriangleIDResult + OutSkippedTriangles) != GeneratorTriangleIndex;
+
+		if (!bHasUnhandledError)
+		{
+			SetTriangleAttributes(Generator, OutDynamicMesh, TriangleIDResult, GeneratorTriangleIndex);
+			return;
+		}
+
+		if (TriangleIDResult == FDynamicMesh3::NonManifoldID && bAttemptToFixNoManifoldError)
+		{
+			// If we get to here, it means we have more than two triangles sharing the same edge.
+			// So lets try to conserve the original geometry by cloning the vertices and creating a new triangle with these
+			// Visually should be mostly ok, although technically this triangle will be "detached"
+			const UE::Geometry::FIndex3i& TriangleData = Generator.Triangles[GeneratorTriangleIndex];
+			UE::Geometry::FIndex3i DuplicatedVertices(
+				OutDynamicMesh.AppendVertex(OutDynamicMesh.GetVertex(TriangleData.A)),
+				OutDynamicMesh.AppendVertex(OutDynamicMesh.GetVertex(TriangleData.B)),
+				OutDynamicMesh.AppendVertex(OutDynamicMesh.GetVertex(TriangleData.C))
+			);
+
+			const int32 RepairedTriangleID = OutDynamicMesh.AppendTriangle(DuplicatedVertices, GroupID);
+
+			UE_LOG(LogChaosVDEditor, Verbose, TEXT("Failed to add triangle | [%d] but expected [%d] | Attempting to fix it ... Repaired triangle ID [%d]"), TriangleIDResult, GeneratorTriangleIndex, RepairedTriangleID);
+
+			// Only attempt to fix once
+			constexpr bool bShouldAttemptToFixNoManifoldError = false;
+			HandleTriangleAddedToDynamicMesh(Generator, OutDynamicMesh, RepairedTriangleID, GroupID, GeneratorTriangleIndex, OutSkippedTriangles, bShouldAttemptToFixNoManifoldError);
+			return;
+		}
+
+		if (TriangleIDResult == FDynamicMesh3::DuplicateTriangleID)
+		{
+			OutSkippedTriangles++;
+			UE_LOG(LogTemp, Verbose, TEXT("Failed to add triangle | [%d] but expected [%d] | Ignoring Duplicated triangle."), TriangleIDResult, GeneratorTriangleIndex);
+			return;
+		}
+
+		OutSkippedTriangles++;
+		UE_LOG(LogTemp, Error, TEXT("Failed to add triangle | [%d] but expected [%d]. This geometry will have missing triangles."), TriangleIDResult, GeneratorTriangleIndex);
+
+		ensure(!bHasUnhandledError);
+	}
+
+	void GenerateDynamicMeshFromGenerator(const UE::Geometry::FMeshShapeGenerator& Generator, FDynamicMesh3& OutDynamicMesh)
+	{
+		OutDynamicMesh.Clear();
+
+		OutDynamicMesh.EnableTriangleGroups();
+
+		if (ensure(Generator.HasAttributes()))
+		{
+			OutDynamicMesh.EnableAttributes();
+		}
+		else
+		{
+			UE_LOG(LogChaosVDEditor, Warning, TEXT("[%s] Attempted to created a mesh using a generator without attributes. CVD Meshes requiere attributes, this should have not happened."), ANSI_TO_TCHAR(__FUNCTION__));
+			return;
+		}
+
+		const int32 NumVerts = Generator.Vertices.Num();
+		for (int32 VertexIndex = 0; VertexIndex < NumVerts; ++VertexIndex)
+		{
+			OutDynamicMesh.AppendVertex(Generator.Vertices[VertexIndex]);
+		}
+
+		if (Cvars::bDisableUVsSupport)
+		{
+			// Remove the default UV Layer
+			OutDynamicMesh.Attributes()->SetNumUVLayers(0);
+		}
+		else if (UE::Geometry::FDynamicMeshUVOverlay* UVOverlay = OutDynamicMesh.Attributes()->PrimaryUV())
+		{
+			const int32 NumUVs = Generator.UVs.Num();
+			for (int32 UVIndex = 0; UVIndex < NumUVs; ++UVIndex)
+			{
+				UVOverlay->AppendElement(Generator.UVs[UVIndex]);
+			}
+		}
+
+		if (UE::Geometry::FDynamicMeshNormalOverlay* NormalOverlay = OutDynamicMesh.Attributes()->PrimaryNormals())
+		{
+			const int32 NumNormals = Generator.Normals.Num();
+            for (int32 NormalIndex = 0; NormalIndex < NumNormals; ++NormalIndex)
+            {
+            	NormalOverlay->AppendElement(Generator.Normals[NormalIndex]);
+            }
+		}
+
+		int32 SkippedTriangles = 0;
+		const int32 NumTris = Generator.Triangles.Num();
+		for (int32 GeneratorTriangleIndex = 0; GeneratorTriangleIndex < NumTris; ++GeneratorTriangleIndex)
+		{
+			const int32 PolygonGroupID = Generator.TrianglePolygonIDs.Num() > 0 ? 1 + Generator.TrianglePolygonIDs[GeneratorTriangleIndex] : 0;
+			const int32 ResultingTriangleID = OutDynamicMesh.AppendTriangle(Generator.Triangles[GeneratorTriangleIndex], PolygonGroupID);
+		
+			constexpr bool bShouldAttemptToFixNoManifoldError = true;
+			HandleTriangleAddedToDynamicMesh(Generator, OutDynamicMesh, ResultingTriangleID, PolygonGroupID, GeneratorTriangleIndex, SkippedTriangles, bShouldAttemptToFixNoManifoldError);	
+		}
+	}
+}
+
 void FChaosVDGeometryBuilder::Initialize(const TWeakPtr<FChaosVDScene>& ChaosVDScene)
 {
 	SceneWeakPtr = ChaosVDScene;
@@ -128,7 +266,19 @@ UDynamicMesh* FChaosVDGeometryBuilder::CreateAndCacheDynamicMesh(const uint32 Ge
 	TRACE_CPUPROFILER_EVENT_SCOPE(FChaosVDGeometryBuilder::CreateAndCacheDynamicMesh_BUILD);
 
 	UDynamicMesh* Mesh = NewObject<UDynamicMesh>();
-	Mesh->SetMesh(&MeshGenerator.Generate());
+
+	FDynamicMesh3 DynamicMesh;
+
+	if (Chaos::VisualDebugger::Cvars::bUseCVDDynamicMeshGenerator)
+	{
+		Chaos::VisualDebugger::GenerateDynamicMeshFromGenerator(MeshGenerator.Generate(), DynamicMesh);
+	}
+	else
+	{
+		DynamicMesh.Copy(&MeshGenerator.Generate());
+	}
+
+	Mesh->SetMesh(DynamicMesh);
 
 	{
 		FWriteScopeLock WriteLock(GeometryCacheRWLock);
@@ -161,7 +311,16 @@ UStaticMesh* FChaosVDGeometryBuilder::CreateAndCacheStaticMesh(const uint32 Geom
 
 	MainStaticMesh->SetNumSourceModels(MeshDescriptionsToGenerate);
 
-	FDynamicMesh3 DynamicMesh(&MeshGenerator.Generate());
+	FDynamicMesh3 DynamicMesh;
+
+	if (Chaos::VisualDebugger::Cvars::bUseCVDDynamicMeshGenerator)
+	{
+		Chaos::VisualDebugger::GenerateDynamicMeshFromGenerator(MeshGenerator.Generate(), DynamicMesh);
+	}
+	else
+	{
+		DynamicMesh.Copy(&MeshGenerator.Generate());
+	}
 
 	for (int32 i = 0; i < MeshDescriptionsToGenerate; i++)
 	{
