[ChaosVD] Allow translucent geometry to be selected

This issue is caused because by default hit proxies are disabled for translucent geometry and there is not a simple way to enable them per viewport.

This CL is a temporary workaround, but we need a better long term solution

diff --git a/Engine/Plugins/Experimental/ChaosVD/Source/ChaosVD/Private/ChaosVDPlaybackViewportClient.cpp b/Engine/Plugins/Experimental/ChaosVD/Source/ChaosVD/Private/ChaosVDPlaybackViewportClient.cpp
--- a/Engine/Plugins/Experimental/ChaosVD/Source/ChaosVD/Private/ChaosVDPlaybackViewportClient.cpp
+++ b/Engine/Plugins/Experimental/ChaosVD/Source/ChaosVD/Private/ChaosVDPlaybackViewportClient.cpp
@@ -14,6 +14,7 @@
 #include "Elements/Framework/TypedElementSelectionSet.h"
 #include "Engine/DirectionalLight.h"
 #include "EngineUtils.h"
+#include "SceneView.h"
 #include "SEditorViewport.h"
 #include "Selection.h"
 #include "UnrealWidget.h"
@@ -347,6 +348,14 @@ bool FChaosVDPlaybackViewportClient::InputKey(const FInputKeyEventArgs& EventArg
 
 void FChaosVDPlaybackViewportClient::Draw(const FSceneView* View, FPrimitiveDrawInterface* PDI)
 {
+	if (View)
+	{
+		// Hack to allow selection of translucent objects (for CVD is all geometry set a Query Only)
+		// The current setting to allow this behaviour is project wide or on custom hitproxies implementations which we can't use
+		// A proper fix would be have a way to override this per viewport, which could be done by adding a new method to FViewElementDrawer
+		const_cast<FSceneView*>(View)->bAllowTranslucentPrimitivesInHitProxy = true;
+	}
+
 	const TSharedPtr<SChaosVDMainTab> MainTabToolkitHost = ModeTools.IsValid() ? StaticCastSharedPtr<SChaosVDMainTab>(ModeTools->GetToolkitHost()) : nullptr;
 	if (!MainTabToolkitHost.IsValid())
 	{
