[ChaosVD] Implemented the UI icons created for the tool + layout fixes

- Now several buttons in the UI use UE Menu system. There is still work to do to fully migrate to that system, but this change allows us to fix a few layout issues we had
- Solver timeline controls no longer deform when resizing the window. The only flexible part is the timeline slider itself. Elements will still overlap when sizing the window too small, further work is needed to make the widgets handle that case
- Limited the output log window in CVD to only show CVD errors by default. More categories can be enabled if desired from the UI.

diff --git a/Engine/Plugins/Experimental/ChaosVD/Source/ChaosVD/Private/ChaosVDCollisionDataDetailsTab.cpp b/Engine/Plugins/Experimental/ChaosVD/Source/ChaosVD/Private/ChaosVDCollisionDataDetailsTab.cpp
--- a/Engine/Plugins/Experimental/ChaosVD/Source/ChaosVD/Private/ChaosVDCollisionDataDetailsTab.cpp
+++ b/Engine/Plugins/Experimental/ChaosVD/Source/ChaosVD/Private/ChaosVDCollisionDataDetailsTab.cpp
@@ -26,6 +26,8 @@ TSharedRef<SDockTab> FChaosVDCollisionDataDetailsTab::HandleTabSpawnRequest(cons
 		SAssignNew(CollisionDataInspector, SChaosVDCollisionDataInspector, GetChaosVDScene())
 	);
 
+	DetailsPanelTab->SetTabIcon(FSlateIcon(FAppStyle::GetAppStyleSetName(), "CollisionAnalyzer.TabIcon").GetIcon());
+
 	HandleTabSpawned(DetailsPanelTab);
 
 	return DetailsPanelTab;
diff --git a/Engine/Plugins/Experimental/ChaosVD/Source/ChaosVD/Private/ChaosVDModule.cpp b/Engine/Plugins/Experimental/ChaosVD/Source/ChaosVD/Private/ChaosVDModule.cpp
--- a/Engine/Plugins/Experimental/ChaosVD/Source/ChaosVD/Private/ChaosVDModule.cpp
+++ b/Engine/Plugins/Experimental/ChaosVD/Source/ChaosVD/Private/ChaosVDModule.cpp
@@ -43,8 +43,7 @@ void FChaosVDModule::StartupModule()
 	FGlobalTabmanager::Get()->RegisterNomadTabSpawner(FChaosVDTabID::ChaosVisualDebuggerTab, FOnSpawnTab::CreateRaw(this, &FChaosVDModule::SpawnMainTab))
 								.SetDisplayName(LOCTEXT("VisualDebuggerTabTitle", "Chaos Visual Debugger"))
 								.SetTooltipText(LOCTEXT("VisualDebuggerTabDesc", "Opens the Chaos Visual Debugger window"))
-								//TODO: Hook up the final icon
-								.SetIcon(FSlateIcon(FAppStyle::GetAppStyleSetName(), "CollisionAnalyzer.TabIcon"))
+								.SetIcon(FSlateIcon(FChaosVDStyle::GetStyleSetName(), "ChaosVisualDebugger"))
 								.SetGroup(WorkspaceMenu::GetMenuStructure().GetDeveloperToolsDebugCategory());
 
 	ChaosVDTraceManager = MakeShared<FChaosVDTraceManager>();
@@ -105,7 +104,7 @@ TSharedRef<SDockTab> FChaosVDModule::SpawnMainTab(const FSpawnTabArgs& Args)
 		SNew(SDockTab)
 		.TabRole(ETabRole::MajorTab)
 		.Label(LOCTEXT("MainTabLabel", "Chaos Visual Debugger"))
-		.ToolTipText(LOCTEXT("MainTabToolTip", "The Chaos Visual debugger is under development"));
+		.ToolTipText(LOCTEXT("MainTabToolTip", "Chaos Visual Debugger is an experimental tool and it can be unstable"));
 
 	// Initialize the Chaos VD Engine instance this tab will represent
 	// For now its lifetime will be controlled by this tab
@@ -118,8 +117,6 @@ TSharedRef<SDockTab> FChaosVDModule::SpawnMainTab(const FSpawnTabArgs& Args)
 			.OwnerTab(MainTabInstance.ToSharedPtr())
 	);
 
-	MainTabInstance->SetTabIcon(FChaosVDStyle::Get().GetBrush("TabIconPlaybackViewport"));
-
 	const FGuid InstanceGuid = ChaosVDEngineInstance->GetInstanceGuid();
 	RegisterChaosVDEngineInstance(InstanceGuid, ChaosVDEngineInstance);
 
diff --git a/Engine/Plugins/Experimental/ChaosVD/Source/ChaosVD/Private/ChaosVDOutputLogTab.cpp b/Engine/Plugins/Experimental/ChaosVD/Source/ChaosVD/Private/ChaosVDOutputLogTab.cpp
--- a/Engine/Plugins/Experimental/ChaosVD/Source/ChaosVD/Private/ChaosVDOutputLogTab.cpp
+++ b/Engine/Plugins/Experimental/ChaosVD/Source/ChaosVD/Private/ChaosVDOutputLogTab.cpp
@@ -2,6 +2,9 @@
 
 #include "ChaosVDOutputLogTab.h"
 
+#include "ChaosLog.h"
+#include "ChaosVDModule.h"
+#include "ChaosVDStyle.h"
 #include "OutputLogCreationParams.h"
 #include "OutputLogModule.h"
 #include "Framework/Docking/TabManager.h"
@@ -16,10 +19,18 @@ TSharedRef<SDockTab> FChaosVDOutputLogTab::HandleTabSpawnRequest(const FSpawnTab
 	Params.SettingsMenuCreationFlags = EOutputLogSettingsMenuFlags::SkipClearOnPie
 		| EOutputLogSettingsMenuFlags::SkipOpenSourceButton
 		| EOutputLogSettingsMenuFlags::SkipEnableWordWrapping; // Checkbox relies on saving an editor config file and does not work correctly
+
+	Params.AllowAsInitialLogCategory = FAllowLogCategoryCallback::CreateLambda([](const FName LogCategory) {
+		return LogCategory == LogChaosVDEditor.GetCategoryName() || LogCategory == LogChaos.GetCategoryName() || LogCategory == FName("Cmd");
+		});
+
+	Params.DefaultCategorySelection.Emplace(LogChaosVDEditor.GetCategoryName(), true);
+	Params.DefaultCategorySelection.Emplace(LogChaos.GetCategoryName(), true);
+	Params.DefaultCategorySelection.Emplace(FName("Cmd"), true);
 	
 	TSharedRef<SDockTab> OutputLogTab =
 		SNew(SDockTab)
-		.TabRole(ETabRole::MajorTab)
+		.TabRole(ETabRole::PanelTab)
 		.Label(LOCTEXT("OutputLogTabLabel", "Output Log"));
 	
 	OutputLogTab->SetContent
@@ -27,6 +38,8 @@ TSharedRef<SDockTab> FChaosVDOutputLogTab::HandleTabSpawnRequest(const FSpawnTab
 		FOutputLogModule::Get().MakeOutputLogWidget(Params)
 	);
 
+	OutputLogTab->SetTabIcon(FChaosVDStyle::Get().GetBrush("TabIconOutputLog"));
+
 	HandleTabSpawned(OutputLogTab);
 
 	return OutputLogTab;
diff --git a/Engine/Plugins/Experimental/ChaosVD/Source/ChaosVD/Private/ChaosVDPlaybackViewportTab.cpp b/Engine/Plugins/Experimental/ChaosVD/Source/ChaosVD/Private/ChaosVDPlaybackViewportTab.cpp
--- a/Engine/Plugins/Experimental/ChaosVD/Source/ChaosVD/Private/ChaosVDPlaybackViewportTab.cpp
+++ b/Engine/Plugins/Experimental/ChaosVD/Source/ChaosVD/Private/ChaosVDPlaybackViewportTab.cpp
@@ -18,7 +18,7 @@ TSharedRef<SDockTab> FChaosVDPlaybackViewportTab::HandleTabSpawnRequest(const FS
 		SNew(SDockTab)
 		.TabRole(ETabRole::PanelTab)
 		.Label(LOCTEXT("ViewportTabLabel", "Playback Viewport"))
-		.ToolTipText(LOCTEXT("ViewportTabToolTip", "The Chaos Visual debugger Viewport is under development"));
+		.ToolTipText(LOCTEXT("ViewportTabToolTip", "Contains the viewport where the recorded physics scene will be rendered"));
 	
 	if (TSharedPtr<SChaosVDMainTab> MainTabPtr = OwningTabWidget.Pin())
 	{
@@ -32,7 +32,7 @@ TSharedRef<SDockTab> FChaosVDPlaybackViewportTab::HandleTabSpawnRequest(const FS
 		ViewportTab->SetContent(GenerateErrorWidget());
 	}
 	
-	ViewportTab->SetTabIcon(FChaosVDStyle::Get().GetBrush("TabIconPlaybackViewport"));
+	ViewportTab->SetTabIcon(FChaosVDStyle::Get().GetBrush("ChaosVisualDebugger"));
 	
 	HandleTabSpawned(ViewportTab);
 
diff --git a/Engine/Plugins/Experimental/ChaosVD/Source/ChaosVD/Private/ChaosVDSceneQueryDataInspectorTab.cpp b/Engine/Plugins/Experimental/ChaosVD/Source/ChaosVD/Private/ChaosVDSceneQueryDataInspectorTab.cpp
--- a/Engine/Plugins/Experimental/ChaosVD/Source/ChaosVD/Private/ChaosVDSceneQueryDataInspectorTab.cpp
+++ b/Engine/Plugins/Experimental/ChaosVD/Source/ChaosVD/Private/ChaosVDSceneQueryDataInspectorTab.cpp
@@ -2,6 +2,7 @@
 
 #include "ChaosVDSceneQueryDataInspectorTab.h"
 
+#include "ChaosVDStyle.h"
 #include "EditorModeManager.h"
 #include "Widgets/SChaosVDMainTab.h"
 #include "Widgets/SChaosVDSceneQueryDataInspector.h"
@@ -32,6 +33,8 @@ TSharedRef<SDockTab> FChaosVDSceneQueryDataInspectorTab::HandleTabSpawnRequest(c
 		DetailsPanelTab->SetContent(GenerateErrorWidget());
 	}
 
+	DetailsPanelTab->SetTabIcon(FChaosVDStyle::Get().GetBrush("SceneQueriesInspectorIcon"));
+
 	HandleTabSpawned(DetailsPanelTab);
 
 	return DetailsPanelTab;
diff --git a/Engine/Plugins/Experimental/ChaosVD/Source/ChaosVD/Private/ChaosVDSolversTracksTab.cpp b/Engine/Plugins/Experimental/ChaosVD/Source/ChaosVD/Private/ChaosVDSolversTracksTab.cpp
--- a/Engine/Plugins/Experimental/ChaosVD/Source/ChaosVD/Private/ChaosVDSolversTracksTab.cpp
+++ b/Engine/Plugins/Experimental/ChaosVD/Source/ChaosVD/Private/ChaosVDSolversTracksTab.cpp
@@ -15,7 +15,7 @@ TSharedRef<SDockTab> FChaosVDSolversTracksTab::HandleTabSpawnRequest(const FSpaw
 	TSharedRef<SDockTab> SolverTracksTab =
 	SNew(SDockTab)
 	.TabRole(ETabRole::PanelTab)
-	.Label(LOCTEXT("SolverTracksTabLabel", "Available Solvers"))
+	.Label(LOCTEXT("SolverTracksTabLabel", "Solver Tracks"))
 	.ToolTipText(LOCTEXT("SolverTracksTabToolTip", "Playback controls for the available solvers on the current Frame"));
 
 	if (const TSharedPtr<SChaosVDMainTab> MainTabPtr = OwningTabWidget.Pin())
@@ -30,7 +30,7 @@ TSharedRef<SDockTab> FChaosVDSolversTracksTab::HandleTabSpawnRequest(const FSpaw
 		SolverTracksTab->SetContent(GenerateErrorWidget());
 	}
 
-	SolverTracksTab->SetTabIcon(FChaosVDStyle::Get().GetBrush("TabIconPlaybackViewport"));
+	SolverTracksTab->SetTabIcon(FChaosVDStyle::Get().GetBrush("TabIconSolverTracks"));
 
 	HandleTabSpawned(SolverTracksTab);
 
diff --git a/Engine/Plugins/Experimental/ChaosVD/Source/ChaosVD/Private/ChaosVDStyle.cpp b/Engine/Plugins/Experimental/ChaosVD/Source/ChaosVD/Private/ChaosVDStyle.cpp
--- a/Engine/Plugins/Experimental/ChaosVD/Source/ChaosVD/Private/ChaosVDStyle.cpp
+++ b/Engine/Plugins/Experimental/ChaosVD/Source/ChaosVD/Private/ChaosVDStyle.cpp
@@ -46,19 +46,26 @@ TSharedRef< FSlateStyleSet > FChaosVDStyle::Create()
 
 	Style->Set("ChaosVD.OpenPluginWindow", new IMAGE_BRUSH_SVG(TEXT("PlaceholderButtonIcon"), Icon20x20));
 
-	Style->Set("PlayIcon", new IMAGE_BRUSH(TEXT("Play_16x"), Icon16x16));
-	Style->Set("StopIcon",  new IMAGE_BRUSH(TEXT("Stop_16x"), Icon16x16));
-	Style->Set("NextIcon",  new IMAGE_BRUSH(TEXT("Next_16x"), Icon16x16));
-	Style->Set("PrevIcon",  new IMAGE_BRUSH( TEXT("Prev_16x"), Icon16x16));
-	Style->Set("PauseIcon",  new IMAGE_BRUSH( TEXT("generic_pause_16x"), Icon16x16));
+	Style->Set("PlayIcon", new IMAGE_BRUSH_SVG(TEXT("PlayControlsPlayForward"), Icon16x16));
+	Style->Set("StopIcon",  new IMAGE_BRUSH_SVG(TEXT("PlayControlsStop"), Icon16x16));
+	Style->Set("NextIcon",  new IMAGE_BRUSH_SVG(TEXT("PlayControlsToNext"), Icon16x16));
+	Style->Set("PrevIcon",  new IMAGE_BRUSH_SVG( TEXT("PlayControlsToPrevious"), Icon16x16));
+	Style->Set("PauseIcon",  new IMAGE_BRUSH_SVG( TEXT("PlayControlsPause"), Icon16x16));
 	
 	Style->Set("TabIconDetailsPanel",  new IMAGE_BRUSH_SVG(TEXT("Details"), Icon16x16));
-	Style->Set("TabIconPlaybackViewport",  new IMAGE_BRUSH(TEXT("icon_tab_PlaybackViewport_16x"), Icon16x16));
+	Style->Set("TabIconPlaybackViewport",  new IMAGE_BRUSH_SVG(TEXT("ChaosVisualDebugger_16"), Icon16x16));
 	Style->Set("TabIconWorldOutliner", new IMAGE_BRUSH_SVG(TEXT("WorldOutliner"), Icon16x16));
+	Style->Set("TabIconOutputLog", new IMAGE_BRUSH_SVG(TEXT("OutputLog"), Icon16x16));
+	Style->Set("TabIconSolverTracks", new IMAGE_BRUSH_SVG(TEXT("ChaosSolverTrack_16"), Icon16x16));
+
+	Style->Set("ChaosVisualDebugger", new IMAGE_BRUSH_SVG(TEXT("ChaosVisualDebugger_16"), Icon16x16));
+	Style->Set("OpenFileIcon", new IMAGE_BRUSH_SVG(TEXT("File_20"), Icon20x20));
+	Style->Set("OpenSessionIcon", new IMAGE_BRUSH_SVG(TEXT("Session_20"), Icon20x20));
+	Style->Set("SceneQueriesInspectorIcon", new IMAGE_BRUSH_SVG(TEXT("ChaosSceneQueries_16"), Icon16x16));
 	
 	Style->Set("LockIcon", new IMAGE_BRUSH_SVG(TEXT("lock"), Icon16x16));
 	Style->Set("UnlockedIcon", new IMAGE_BRUSH_SVG(TEXT("lock-unlocked"), Icon16x16));
-	Style->Set("RecordIcon", new IMAGE_BRUSH(TEXT("CircleBox2"), Icon16x16));
+	Style->Set("RecordIcon", new IMAGE_BRUSH_SVG(TEXT("PlayControlsRecord"), Icon16x16));
 	Style->Set("RecordToFileIcon", new IMAGE_BRUSH(TEXT("RecordToFile"), Icon16x16));
 	Style->Set("RecordToLiveIcon", new IMAGE_BRUSH(TEXT("RecordLiveSession"), Icon16x16));
 	Style->Set("ConnectionIcon", new IMAGE_BRUSH_SVG(TEXT("Connection"), Icon16x16));
diff --git a/Engine/Plugins/Experimental/ChaosVD/Source/ChaosVD/Private/Widgets/SChaosVDMainTab.cpp b/Engine/Plugins/Experimental/ChaosVD/Source/ChaosVD/Private/Widgets/SChaosVDMainTab.cpp
--- a/Engine/Plugins/Experimental/ChaosVD/Source/ChaosVD/Private/Widgets/SChaosVDMainTab.cpp
+++ b/Engine/Plugins/Experimental/ChaosVD/Source/ChaosVD/Private/Widgets/SChaosVDMainTab.cpp
@@ -28,6 +28,8 @@
 #include "IDesktopPlatform.h"
 #include "Misc/MessageDialog.h"
 #include "StatusBarSubsystem.h"
+#include "ToolMenu.h"
+#include "ToolMenus.h"
 #include "Components/ChaosVDParticleDataComponent.h"
 #include "Components/ChaosVDSolverCollisionDataComponent.h"
 #include "Components/ChaosVDSolverJointConstraintDataComponent.h"
@@ -109,156 +111,33 @@ void SChaosVDMainTab::Construct(const FArguments& InArgs, TSharedPtr<FChaosVDEng
 		+ SVerticalBox::Slot()
 		.AutoHeight()
 		[
-			SNew( SBorder )
-			.BorderImage(FAppStyle::Get().GetBrush("Brushes.Panel"))
-			.Padding(FMargin(0.f, 0.f, 0.f, 0.f))
+			// Create the Main Toolbar
+			SNew(SOverlay)
+			+SOverlay::Slot()
 			[
-				// Open Button
-				SNew(SHorizontalBox)
-				+SHorizontalBox::Slot()
-				.VAlign(VAlign_Center)
-				.HAlign(HAlign_Left)
-				.Padding(FMargin(12, 7, 6, 7))
-				.AutoWidth()
-				[
-					SNew(SButton)
-					.ToolTip(SNew(SToolTip).Text(LOCTEXT("OpenFileDesc", "Click here to open a Chaos Visual Debugger file.")))
-					.ContentPadding(FMargin(0, 5.f, 0, 4.f))
-					.OnClicked_Lambda([this]()
-					{
-						BrowseAndOpenChaosVDRecording();
-						return FReply::Handled();
-					})
-					.Content()
-					[
-						SNew(SHorizontalBox)
-						+SHorizontalBox::Slot()
-						.HAlign(HAlign_Center)
-						.VAlign(VAlign_Center)
-						[
-							SNew(SImage)
-							.Image(FAppStyle::Get().GetBrush("Icons.Plus"))
-							.ColorAndOpacity(FStyleColors::AccentGreen)
-						]
-						+SHorizontalBox::Slot()
-						.Padding(FMargin(3, 0, 0, 0))
-						.VAlign(VAlign_Center)
-						.AutoWidth()
-						[
-							SNew(STextBlock)
-							.TextStyle(FAppStyle::Get(), "SmallButtonText")
-							.Text(LOCTEXT("OpenFile", "Open File"))
-						]
-					]
-				]
-
-				// Remote Connection Button
-				+SHorizontalBox::Slot()
-				.VAlign(VAlign_Center)
-				.HAlign(HAlign_Left)
-				.Padding(FMargin(6, 7, 18, 7))
-				.AutoWidth()
+				SNew(SImage)
+				.Image(&FAppStyle::Get().GetWidgetStyle<FToolBarStyle>("SlimToolBar").BackgroundBrush)
+			]
+			+ SOverlay::Slot()
+			[
+				SNew(SBorder)
+				.Padding(0)
+				.BorderImage(FAppStyle::Get().GetBrush("NoBorder"))
 				[
-					SNew(SButton)
-					.ContentPadding(FMargin(0, 5.f, 0, 4.f))
-					.OnClicked_Raw(this, &SChaosVDMainTab::HandleSessionConnectionClicked)
-					.Content()
+					SNew(SHorizontalBox)
+					+SHorizontalBox::Slot()
+					.AutoWidth()
 					[
-						SNew(SHorizontalBox)
-						+SHorizontalBox::Slot()
-						.HAlign(HAlign_Center)
-						.VAlign(VAlign_Center)
-						[
-							SNew(SImage)
-							.Image(FChaosVDStyle::Get().GetBrush("ConnectionIcon"))
-							.ColorAndOpacity(FStyleColors::AccentGreen)
-						]
-						+SHorizontalBox::Slot()
-						.Padding(FMargin(3, 0, 0, 0))
-						.VAlign(VAlign_Center)
-						.AutoWidth()
-						[
-							SNew(STextBlock)
-							.TextStyle(FAppStyle::Get(), "SmallButtonText")
-							.Text_Raw(this, &SChaosVDMainTab::GetConnectButtonText)
-						]
+						GenerateMainToolbarWidget()
 					]
-				]
-
-
-				+SHorizontalBox::Slot()
-				.AutoWidth()
-				.Padding(1.0f)
-				[
-					SNew(SSeparator)
-						.Orientation(Orient_Vertical)
-						.Thickness(2.0f)
-						.ColorAndOpacity(FColor::Black)
-						.SeparatorImage(FAppStyle::Get().GetBrush("Menu.Separator"))
-				]
-				
-				+SHorizontalBox::Slot()
-				[
-					
-					SNew(SChaosVDRecordingControls, StaticCastSharedRef<SChaosVDMainTab>(AsShared()))
-				]
-				
-				+SHorizontalBox::Slot()
-				.AutoWidth()
-				.Padding(1.0f, 5.0f, 1.0f, 5.0f)
-				[
-					SNew(SSeparator)
-						.Orientation(Orient_Vertical)
-						.Thickness(2.0f)
-						.ColorAndOpacity(FColor::Black)
-						.SeparatorImage(FAppStyle::Get().GetBrush("Menu.Separator"))
-				]
-
-				// Settings button
-				+ SHorizontalBox::Slot()
-				.HAlign(HAlign_Right)
-				.VAlign(VAlign_Center)
-				.AutoWidth()
-				.Padding(FMargin(14.f, 0.f, 14.f, 0.f))
-				[
-					SNew(SComboButton)
-					.ContentPadding(0)
-					.HasDownArrow(false)
-					.IsEnabled(false) // Disabled until we create content for it
-					.ForegroundColor(FSlateColor::UseForeground())
-					.ComboButtonStyle(FAppStyle::Get(), "SimpleComboButton")
-					.AddMetaData<FTagMetaData>(FTagMetaData(TEXT("ViewOptions")))
-					/*.MenuContent()
-					[
-						//TODO: Implement Settings menu
-					]*/
-					.ButtonContent()
+					+SHorizontalBox::Slot()
+					.AutoWidth()
 					[
-						SNew(SHorizontalBox)
-						+ SHorizontalBox::Slot()
-						.AutoWidth()
-						.HAlign(HAlign_Center)
-						[
-							SNew(SImage)
-							.Image(FAppStyle::Get().GetBrush("Icons.Toolbar.Settings"))
-							.ColorAndOpacity(FSlateColor::UseForeground())
-						]
-						+ SHorizontalBox::Slot()
-						.Padding(FMargin(5, 0, 0, 0))
-						.VAlign(VAlign_Center)
-						.AutoWidth()
-						[
-							SNew(STextBlock)
-							.TextStyle(FAppStyle::Get(), "NormalText")
-							.Text(LOCTEXT("SettingsLabel", "Settings"))
-							.ColorAndOpacity(FSlateColor::UseForeground())
-						]
+						SNew(SChaosVDRecordingControls, StaticCastSharedRef<SChaosVDMainTab>(AsShared()))
 					]
 				]
 			]
-				
 		]
-		
 		// Main Visual Debugger Interface content
 		+SVerticalBox::Slot()
 		.Padding(FMargin(0.0f,5.0f,0.0f,0.0f))
@@ -433,7 +312,7 @@ void SChaosVDMainTab::GenerateMainWindowMenu()
 	TabManager->SetMenuMultiBox(MenuBarBuilder.GetMultiBox(), MenuBarBuilder.MakeWidget());
 }
 
-void SChaosVDMainTab::BrowseAndOpenChaosVDRecording()
+FReply SChaosVDMainTab::BrowseAndOpenChaosVDRecording()
 {
 	const TSharedRef<SChaosBrowseTraceFileSourceModal> SessionBrowserModal = SNew(SChaosBrowseTraceFileSourceModal);
 
@@ -469,6 +348,56 @@ void SChaosVDMainTab::BrowseAndOpenChaosVDRecording()
 				ensureMsgf(false, TEXT("Invalid responce received"));
 			break;
 	}
+
+	return FReply::Handled();
+}
+
+TSharedRef<SButton> SChaosVDMainTab::CreateSimpleButton(TFunction<FText()>&& GetTextDelegate, TFunction<FText()>&& ToolTipTextDelegate, const FSlateBrush* ButtonIcon, const UChaosVDMainToolbarMenuContext* MenuContext, const FOnClicked& InButtonClickedCallback)
+{
+	const TSharedRef<SChaosVDMainTab> MainTab = MenuContext->MainTab.Pin().ToSharedRef();
+	
+	TSharedRef<SButton> Button = SNew(SButton)
+								.ButtonStyle(FAppStyle::Get(), "SimpleButton")
+								.ToolTipText_Lambda(MoveTemp(ToolTipTextDelegate))
+								.ContentPadding(FMargin(6.0f, 0.0f))
+								.OnClicked(InButtonClickedCallback)
+								[
+									SNew(SHorizontalBox)
+									+ SHorizontalBox::Slot()
+									.AutoWidth()
+									.HAlign(HAlign_Center)
+									.VAlign(VAlign_Center)
+									[
+										SNew(SImage)
+										.Image(ButtonIcon)
+										.ColorAndOpacity(FSlateColor::UseForeground())
+									]
+									+SHorizontalBox::Slot()
+									.Padding(FMargin(4, 0, 0, 0))
+									.VAlign(VAlign_Center)
+									.AutoWidth()
+									[
+										SNew(STextBlock)
+										.TextStyle(FAppStyle::Get(), "NormalText")
+										.Text_Lambda(MoveTemp(GetTextDelegate))
+									]
+								];
+
+	return Button;
+}
+
+TSharedRef<SWidget> SChaosVDMainTab::GenerateMainToolbarWidget()
+{
+	RegisterMainTabMenu();
+
+	FToolMenuContext MenuContext;
+
+	UChaosVDMainToolbarMenuContext* CommonContextObject = NewObject<UChaosVDMainToolbarMenuContext>();
+	CommonContextObject->MainTab = SharedThis(this);
+
+	MenuContext.AddObject(CommonContextObject);
+
+	return UToolMenus::Get()->GenerateWidget(MainToolBarName, MenuContext);
 }
 
 void SChaosVDMainTab::BrowseChaosVDRecordingFromFolder(FStringView FolderPath)
@@ -523,6 +452,104 @@ bool SChaosVDMainTab::ConnectToLiveSession(int32 SessionID, const FString Sessio
 	return bSuccess;
 }
 
+void SChaosVDMainTab::RegisterMainTabMenu()
+{
+	const UToolMenus* ToolMenus = UToolMenus::Get();
+	if (ToolMenus->IsMenuRegistered(MainToolBarName))
+	{
+		return;
+	}
+
+	UToolMenu* ToolBar = UToolMenus::Get()->RegisterMenu(MainToolBarName, NAME_None, EMultiBoxType::SlimHorizontalToolBar);
+
+	FToolMenuSection& Section = ToolBar->AddSection("LoadRecording");
+	Section.AddDynamicEntry("OpenFile", FNewToolMenuSectionDelegate::CreateLambda([](FToolMenuSection& InSection)
+	{
+		const UChaosVDMainToolbarMenuContext* Context = InSection.FindContext<UChaosVDMainToolbarMenuContext>();
+		TSharedPtr<SChaosVDMainTab> MainTabPtr = Context->MainTab.Pin();
+		if (!MainTabPtr)
+		{
+			return;
+		}
+
+		TSharedRef<SButton> OpenFileButton = MainTabPtr->CreateSimpleButton(
+			[](){ return LOCTEXT("OpenFile", "Open File"); },
+			[](){ return LOCTEXT("OpenFileDesc", "Click here to open a Chaos Visual Debugger file."); },
+			FChaosVDStyle::Get().GetBrush("OpenFileIcon"),
+			Context, FOnClicked::CreateLambda([WeakTab = MainTabPtr.ToWeakPtr()]()
+			{
+				if (TSharedPtr<SChaosVDMainTab> TabPtr = WeakTab.Pin())
+				{
+					return TabPtr->BrowseAndOpenChaosVDRecording();
+				}
+
+				return FReply::Handled();
+			}));
+
+		InSection.AddEntry(
+			FToolMenuEntry::InitWidget(
+				"OpenFileButton",
+				OpenFileButton,
+				FText::GetEmpty(),
+				true,
+				false
+			));
+	}));
+
+	Section.AddDynamicEntry("ConnectToSession", FNewToolMenuSectionDelegate::CreateLambda([](FToolMenuSection& InSection)
+	{
+		const UChaosVDMainToolbarMenuContext* Context = InSection.FindContext<UChaosVDMainToolbarMenuContext>();
+		TSharedPtr<SChaosVDMainTab> MainTabPtr = Context->MainTab.Pin();
+		if (!MainTabPtr)
+		{
+			return;
+		}
+
+		TFunction<FText()> GetTextDelegate = [WeakTab = MainTabPtr.ToWeakPtr()]()
+		{
+			if (TSharedPtr<SChaosVDMainTab> TabPtr = WeakTab.Pin())
+			{
+				return TabPtr->GetConnectButtonText();
+			}
+
+			return FText();
+		};
+
+		TFunction<FText()> GetTooltipTextDelegate = [WeakTab = MainTabPtr.ToWeakPtr()]()
+		{
+			if (TSharedPtr<SChaosVDMainTab> TabPtr = WeakTab.Pin())
+			{
+				return TabPtr->GetConnectButtonTooltipText();
+			}
+
+			return FText();
+		};
+		
+		FOnClicked OnClickedDelegate = FOnClicked::CreateLambda([WeakTab = StaticCastWeakPtr<SChaosVDMainTab>(MainTabPtr->AsWeak())]()
+		{
+			if (TSharedPtr<SChaosVDMainTab> TabPtr = WeakTab.Pin())
+			{
+				return TabPtr->HandleSessionConnectionClicked();
+			}
+
+			return FReply::Handled();
+		});
+
+		TSharedRef<SButton> ConnectToSessionButton = MainTabPtr->CreateSimpleButton(MoveTemp(GetTextDelegate), MoveTemp(GetTooltipTextDelegate), FChaosVDStyle::Get().GetBrush("OpenSessionIcon"), Context, MoveTemp(OnClickedDelegate));
+
+		InSection.AddEntry(
+			FToolMenuEntry::InitWidget(
+				"ConnectToSession",
+				ConnectToSessionButton,
+				FText::GetEmpty(),
+				true,
+				false
+			));
+	}));
+	
+	Section.AddSeparator(NAME_None);
+}
+
 void SChaosVDMainTab::BrowseLiveSessionsFromTraceStore() const
 {
 	const TSharedRef<SChaosVDBrowseSessionsModal> SessionBrowserModal = SNew(SChaosVDBrowseSessionsModal);
@@ -570,4 +597,10 @@ FText SChaosVDMainTab::GetConnectButtonText() const
 	return bIsAlreadyInLiveSession ? LOCTEXT("DisconnectFromSession", "Disconnect from Session") : LOCTEXT("ConnectToSession", "Connect to Session");
 }
 
+FText SChaosVDMainTab::GetConnectButtonTooltipText() const
+{
+	const bool bIsAlreadyInLiveSession = GetChaosVDEngineInstance()->GetPlaybackController()->IsPlayingLiveSession();
+	return bIsAlreadyInLiveSession ? LOCTEXT("DisconnectFromSessionTooltip", "Disconnects from the current live session but it does not stop it") : LOCTEXT("ConnectToSessionTooltip", "Opens a panel where you can browse active live sessions and connect to one");
+}
+
 #undef LOCTEXT_NAMESPACE
diff --git a/Engine/Plugins/Experimental/ChaosVD/Source/ChaosVD/Private/Widgets/SChaosVDRecordingControls.cpp b/Engine/Plugins/Experimental/ChaosVD/Source/ChaosVD/Private/Widgets/SChaosVDRecordingControls.cpp
--- a/Engine/Plugins/Experimental/ChaosVD/Source/ChaosVD/Private/Widgets/SChaosVDRecordingControls.cpp
+++ b/Engine/Plugins/Experimental/ChaosVD/Source/ChaosVD/Private/Widgets/SChaosVDRecordingControls.cpp
@@ -11,8 +11,9 @@
 #include "Input/Reply.h"
 #include "Misc/MessageDialog.h"
 #include "StatusBarSubsystem.h"
+#include "ToolMenu.h"
+#include "ToolMenus.h"
 #include "Framework/MultiBox/MultiBoxBuilder.h"
-#include "Styling/SlateBrush.h"
 #include "Widgets/Images/SImage.h"
 #include "Widgets/Input/SButton.h"
 #include "Widgets/SChaosVDMainTab.h"
@@ -25,90 +26,126 @@ void SChaosVDRecordingControls::Construct(const FArguments& InArgs, const TShare
 {
 	MainTabWeakPtr = InMainTabSharedRef;
 	StatusBarID = InMainTabSharedRef->GetStatusBarName();
+	
+	RecordingAnimation = FCurveSequence();
+	RecordingAnimation.AddCurve(0.f, 1.5f, ECurveEaseFunction::Linear);
 
 	ChildSlot
 	[
-		SNew(SHorizontalBox)
-		+SHorizontalBox::Slot()
-		.HAlign(HAlign_Left)
-		.Padding(FMargin(12, 7, 2, 7))
-		[
-			SNew(SHorizontalBox)
-			+SHorizontalBox::Slot()
-			.HAlign(HAlign_Left)
-			[
-				GenerateToggleRecordingStateButton(EChaosVDRecordingMode::File, LOCTEXT("RecordToFileButtonDesc", "Starts a recording for the current session, saving it directly to file"))
-			]
-			+SHorizontalBox::Slot()
-			.Padding(5.0f,  0.0f, 0.0f, 0.0f)
-			.HAlign(HAlign_Left)
-			[
-				GenerateToggleRecordingStateButton(EChaosVDRecordingMode::Live, LOCTEXT("RecordLiveButtonDesc", "Starts a recording and automatically connects to it playing it back in real time"))
-			]
-			+SHorizontalBox::Slot()
-			.HAlign(HAlign_Left)
+		GenerateToolbarWidget()
+	];
+
+	RecordingStartedHandle = FChaosVDRuntimeModule::Get().RegisterRecordingStartedCallback(FChaosVDRecordingStateChangedDelegate::FDelegate::CreateRaw(this, &SChaosVDRecordingControls::HandleRecordingStart));
+	RecordingStoppedHandle = FChaosVDRuntimeModule::Get().RegisterRecordingStopCallback(FChaosVDRecordingStateChangedDelegate::FDelegate::CreateRaw(this, &SChaosVDRecordingControls::HandleRecordingStop));
+}
+
+TSharedRef<SWidget> SChaosVDRecordingControls::GenerateToggleRecordingStateButton(EChaosVDRecordingMode RecordingMode, const FText& StartRecordingTooltip)
+{
+	return SNew(SButton)
+			.OnClicked(FOnClicked::CreateRaw(this, &SChaosVDRecordingControls::ToggleRecordingState, RecordingMode))
+			.ForegroundColor(FSlateColor::UseForeground())
+			.IsFocusable(false)
+			.IsEnabled_Raw(this, &SChaosVDRecordingControls::IsRecordingToggleButtonEnabled, RecordingMode)
+			.Visibility_Raw(this, &SChaosVDRecordingControls::IsRecordingToggleButtonVisible, RecordingMode)
 			.VAlign(VAlign_Center)
-			.AutoWidth()
-			.Padding(12.0f, 0.0f, 0.0f, 0.0f)
+			.HAlign(HAlign_Center)
+			.OnHovered_Lambda([this](){ bRecordingButtonHovered = true;})
+			.OnUnhovered_Lambda([this](){ bRecordingButtonHovered = false;})
+			.ToolTipText_Lambda([this, StartRecordingTooltip]()
+			{
+				return IsRecording() ? LOCTEXT("StopRecordButtonDesc", "Stop the current recording ") : StartRecordingTooltip;
+			})
 			[
-				SNew( STextBlock )
+				SNew(SHorizontalBox)
+				+SHorizontalBox::Slot()
+				.Padding(FMargin(0, 0, 0, 0))
+				.VAlign(VAlign_Center)
+				.AutoWidth()
+				[
+					SNew(SImage)
+					.Image_Raw(this, &SChaosVDRecordingControls::GetRecordOrStopButton, RecordingMode)
+					.ColorAndOpacity_Lambda([this]()
+					{
+						if (IsRecording())
+						{
+							if (!RecordingAnimation.IsPlaying())
+							{
+								RecordingAnimation.Play(AsShared(), true);
+							}
+
+							const FLinearColor Color = bRecordingButtonHovered ? FLinearColor::Red : FLinearColor::White;
+							return FSlateColor(bRecordingButtonHovered ? Color : Color.CopyWithNewOpacity(0.2f + 0.8f * RecordingAnimation.GetLerp()));
+						}
+
+						RecordingAnimation.Pause();
+						return FSlateColor::UseSubduedForeground();
+					})
+				]
+				+SHorizontalBox::Slot()
+				.Padding(FMargin(4, 0, 0, 0))
+				.VAlign(VAlign_Center)
+				.AutoWidth()
+				[
+					SNew(STextBlock)
+					.Visibility_Lambda([this](){return IsRecording() ? EVisibility::Collapsed : EVisibility::Visible;})
 					.TextStyle(FAppStyle::Get(), "SmallButtonText")
-					.Text_Raw(this, &SChaosVDRecordingControls::GetRecordingTimeText)
-					.ColorAndOpacity(FColor::White)
-			]
+					.Text_Lambda( [RecordingMode]()
+					{
+						return RecordingMode == EChaosVDRecordingMode::File ? LOCTEXT("RecordToFileButtonLabel", "Record To File") : LOCTEXT("RecordToLiveButtonLabel", "Record Live Session");
+					})
+				]
+			];
+}
 
-			+SHorizontalBox::Slot()
-			.HAlign(HAlign_Left)
+TSharedRef<SWidget> SChaosVDRecordingControls::GenerateRecordingTimeTextBlock()
+{
+	return SNew(SBox)
 			.VAlign(VAlign_Center)
-			.AutoWidth()
+			.Visibility_Lambda([this]() { return IsRecording() ? EVisibility::Visible : EVisibility::Collapsed; })
 			.Padding(12.0f, 0.0f, 0.0f, 0.0f)
 			[
-				SNew(SComboButton)
-				.ContentPadding(FMargin(6.0f, 0.0f))
-				.IsEnabled_Raw(this, &SChaosVDRecordingControls::HasDataChannelsSupport)
-				.MenuPlacement(MenuPlacement_AboveAnchor)
-				.ComboButtonStyle(&FAppStyle::Get().GetWidgetStyle<FComboButtonStyle>("SimpleComboButton"))
-				.OnGetMenuContent(this, &SChaosVDRecordingControls::GenerateDataChannelsMenu)
-				.HasDownArrow(true)
-				.ButtonContent()
-				[
-					SNew(SHorizontalBox)
-					+ SHorizontalBox::Slot()
-					.VAlign(VAlign_Center)
-					.AutoWidth()
-					[
-						SNew(STextBlock)
-						.Text(LOCTEXT("DataChannelsButton", "Data Channels"))
-						.TextStyle(&FAppStyle::Get().GetWidgetStyle<FTextBlockStyle>("DialogButtonText"))
-					]
-				]
-			]
-		]
-	];
+				SNew(STextBlock)
+				.TextStyle(FAppStyle::Get(), "SmallButtonText")
+				.Text_Raw(this, &SChaosVDRecordingControls::GetRecordingTimeText)
+				.ColorAndOpacity(FColor::White)
+			];
+}
 
-	RecordingStartedHandle = FChaosVDRuntimeModule::Get().RegisterRecordingStartedCallback(FChaosVDRecordingStateChangedDelegate::FDelegate::CreateRaw(this, &SChaosVDRecordingControls::HandleRecordingStart));
-	RecordingStoppedHandle = FChaosVDRuntimeModule::Get().RegisterRecordingStopCallback(FChaosVDRecordingStateChangedDelegate::FDelegate::CreateRaw(this, &SChaosVDRecordingControls::HandleRecordingStop));
+TSharedRef<SWidget> SChaosVDRecordingControls::GenerateToolbarWidget()
+{
+	RegisterMenus();
+
+	FToolMenuContext MenuContext;
+
+	UChaosVDRecordingToolbarMenuContext* CommonContextObject = NewObject<UChaosVDRecordingToolbarMenuContext>();
+	CommonContextObject->RecordingControlsWidget = SharedThis(this);
+
+	MenuContext.AddObject(CommonContextObject);
+
+	return UToolMenus::Get()->GenerateWidget(RecordingControlsToolbarName, MenuContext);
 }
 
-TSharedRef<SButton> SChaosVDRecordingControls::GenerateToggleRecordingStateButton(EChaosVDRecordingMode RecordingMode, const FText& StartRecordingTooltip)
+TSharedRef<SWidget> SChaosVDRecordingControls::GenerateDataChannelsButton()
 {
-	return SNew(SButton)
-		.OnClicked(FOnClicked::CreateRaw(this, &SChaosVDRecordingControls::ToggleRecordingState, RecordingMode))
-		.ForegroundColor(FSlateColor::UseForeground())
-		.IsFocusable(false)
-		.IsEnabled_Raw(this, &SChaosVDRecordingControls::IsRecordingToggleButtonEnabled, RecordingMode)
-		.Visibility_Raw(this, &SChaosVDRecordingControls::IsRecordingToggleButtonVisible, RecordingMode)
-		.VAlign(VAlign_Center)
-		.HAlign(HAlign_Center)
-		.ToolTipText_Lambda([this, StartRecordingTooltip]()
-		{
-			return IsRecording() ? LOCTEXT("StopRecordButtonDesc", "Stop the current recording ") : StartRecordingTooltip;
-		})
-		[
-			SNew(SImage)
-			.Image_Raw(this, &SChaosVDRecordingControls::GetRecordOrStopButton, RecordingMode)
-			.ColorAndOpacity_Lambda([this](){ return IsRecording() ? FColor::Red : FColor::White; })
-		];
+	return SNew(SComboButton)
+			.ContentPadding(FMargin(6.0f, 0.0f))
+			.IsEnabled_Raw(this, &SChaosVDRecordingControls::HasDataChannelsSupport)
+			.MenuPlacement(MenuPlacement_AboveAnchor).ComboButtonStyle(&FAppStyle::Get()
+			.GetWidgetStyle<FComboButtonStyle>("SimpleComboButton"))
+			.OnGetMenuContent(this, &SChaosVDRecordingControls::GenerateDataChannelsMenu)
+			.HasDownArrow(true)
+			.ButtonContent()
+			[
+				SNew(SHorizontalBox)
+				+SHorizontalBox::Slot()
+				.VAlign(VAlign_Center)
+				.AutoWidth()
+				[
+					SNew(STextBlock)
+					.Text(LOCTEXT("DataChannelsButton", "Data Channels"))
+					.TextStyle(&FAppStyle::Get().GetWidgetStyle<FTextBlockStyle>("DialogButtonText"))
+				]
+			];
 }
 
 TSharedRef<SWidget> SChaosVDRecordingControls::GenerateDataChannelsMenu()
@@ -192,8 +229,8 @@ SChaosVDRecordingControls::~SChaosVDRecordingControls()
 
 const FSlateBrush* SChaosVDRecordingControls::GetRecordOrStopButton(EChaosVDRecordingMode RecordingMode) const
 {
-	const FSlateBrush* RecordIconBrush = RecordingMode == EChaosVDRecordingMode::File ? FChaosVDStyle::Get().GetBrush("RecordToFileIcon") : FChaosVDStyle::Get().GetBrush("RecordToLiveIcon");
-	return IsRecording() ? FChaosVDStyle::Get().GetBrush("StopIcon") : RecordIconBrush;
+	const FSlateBrush* RecordIconBrush = FChaosVDStyle::Get().GetBrush("RecordIcon");
+	return bRecordingButtonHovered && IsRecording() ? FChaosVDStyle::Get().GetBrush("StopIcon") : RecordIconBrush;
 }
 
 void SChaosVDRecordingControls::HandleRecordingStop()
@@ -366,6 +403,66 @@ EVisibility SChaosVDRecordingControls::IsRecordingToggleButtonVisible(EChaosVDRe
 	return bShouldButtonBeVisible ? EVisibility::Visible : EVisibility::Collapsed;
 }
 
+void SChaosVDRecordingControls::RegisterMenus()
+{
+	const UToolMenus* ToolMenus = UToolMenus::Get();
+	if (ToolMenus->IsMenuRegistered(RecordingControlsToolbarName))
+	{
+		return;
+	}
+
+	UToolMenu* ToolBar = UToolMenus::Get()->RegisterMenu(RecordingControlsToolbarName, NAME_None, EMultiBoxType::SlimHorizontalToolBar);
+
+	FToolMenuSection& Section = ToolBar->AddSection("LoadRecording");
+	Section.AddDynamicEntry("OpenFile", FNewToolMenuSectionDelegate::CreateLambda([](FToolMenuSection& InSection)
+	{
+		const UChaosVDRecordingToolbarMenuContext* Context = InSection.FindContext<UChaosVDRecordingToolbarMenuContext>();
+		TSharedRef<SChaosVDRecordingControls> RecordingControls = Context->RecordingControlsWidget.Pin().ToSharedRef();
+
+		TSharedRef<SWidget> RecordToFileButton = SNew(SBox).Padding(4.0f,0.0f)[ RecordingControls->GenerateToggleRecordingStateButton(EChaosVDRecordingMode::File, LOCTEXT("RecordToFileButtonDesc", "Starts a recording for the current session, saving it directly to file")) ];
+		TSharedRef<SWidget> RecordToLiveButton = SNew(SBox).Padding(4.0f,0.0f)[ RecordingControls->GenerateToggleRecordingStateButton(EChaosVDRecordingMode::Live, LOCTEXT("RecordLiveButtonDesc", "Starts a recording and automatically connects to it playing it back in real time")) ];
+		TSharedRef<SWidget> RecordingTime = RecordingControls->GenerateRecordingTimeTextBlock();
+		TSharedRef<SWidget> DataChannelsButton = SNew(SBox).Padding(4.0f,0.0f)[ RecordingControls->GenerateDataChannelsButton() ];
+
+		InSection.AddEntry(
+			FToolMenuEntry::InitWidget(
+				"RecordToFileButton",
+				RecordToFileButton,
+				FText::GetEmpty(),
+				true,
+				false
+			));
+
+		InSection.AddEntry(
+			FToolMenuEntry::InitWidget(
+				"RecordToLiveButton",
+				RecordToLiveButton,
+				FText::GetEmpty(),
+				false,
+				false
+			));
+
+		InSection.AddEntry(
+			FToolMenuEntry::InitWidget(
+				"RecordingTime",
+				RecordingTime,
+				FText::GetEmpty(),
+				false,
+				false
+			));
+
+		InSection.AddEntry(
+			FToolMenuEntry::InitWidget(
+				"DataChannelsButton",
+				DataChannelsButton,
+				FText::GetEmpty(),
+				false,
+				false
+			));
+	}));
+}
+
+
 bool SChaosVDRecordingControls::IsRecording() const
 {
 #if WITH_CHAOS_VISUAL_DEBUGGER
diff --git a/Engine/Plugins/Experimental/ChaosVD/Source/ChaosVD/Private/Widgets/SChaosVDSolverPlaybackControls.cpp b/Engine/Plugins/Experimental/ChaosVD/Source/ChaosVD/Private/Widgets/SChaosVDSolverPlaybackControls.cpp
--- a/Engine/Plugins/Experimental/ChaosVD/Source/ChaosVD/Private/Widgets/SChaosVDSolverPlaybackControls.cpp
+++ b/Engine/Plugins/Experimental/ChaosVD/Source/ChaosVD/Private/Widgets/SChaosVDSolverPlaybackControls.cpp
@@ -28,86 +28,81 @@ void SChaosVDSolverPlaybackControls::Construct(const FArguments& InArgs, int32 I
 
 	ChildSlot
 	[
-		SNew(SVerticalBox)
-		// Playback controls
-		// TODO: Now that the tool is In-Editor, see if we can/is worth use the Sequencer widgets
-		// instead of these custom ones
-		+SVerticalBox::Slot()
+		SNew(SHorizontalBox)
+		+SHorizontalBox::Slot()
+		.AutoWidth()
+		.FillWidth(0.8f)
 		[
-			SNew(SHorizontalBox)
-			+SHorizontalBox::Slot()
-			.FillWidth(0.7f)
+			SNew(SVerticalBox)
+			+SVerticalBox::Slot()
+			.AutoHeight()
+			.Padding(0.0f, 0.0f, 0.0f, 2.0f)
 			[
-				SNew(SVerticalBox)
-				+SVerticalBox::Slot()
-				.AutoHeight()
-				.Padding(0.0f, 0.0f, 0.0f, 2.0f)
+				SNew(STextBlock)
+				.Justification(ETextJustify::Center)
+				.Text(LOCTEXT("PlaybackViewportWidgetPhysicsFramesLabel", "Solver Frames" ))
+			]
+			+SVerticalBox::Slot()
+			[
+				SNew(SHorizontalBox)
+				+SHorizontalBox::Slot()
+				.FillWidth(0.9f)
 				[
-					SNew(STextBlock)
-					.Justification(ETextJustify::Center)
-					.Text(LOCTEXT("PlaybackViewportWidgetPhysicsFramesLabel", "Solver Frames" ))
+					SAssignNew(FramesTimelineWidget, SChaosVDTimelineWidget)
+					.ButtonVisibilityFlags(static_cast<uint16>(EChaosVDTimelineElementIDFlags::AllPlayback))
+					.OnFrameChanged_Raw(this, &SChaosVDSolverPlaybackControls::OnFrameSelectionUpdated)
+					.OnButtonClicked_Raw(this, &SChaosVDSolverPlaybackControls::HandlePlaybackButtonClicked)
+					.MaxFrames(0)
 				]
-				+SVerticalBox::Slot()
+				+SHorizontalBox::Slot()
+				.Padding(6.0f,0.0f)
+				.AutoWidth()
 				[
-					SNew(SHorizontalBox)
-					+SHorizontalBox::Slot()
-					.FillWidth(0.95f)
-					[
-						SAssignNew(FramesTimelineWidget, SChaosVDTimelineWidget)
-						.ButtonVisibilityFlags(static_cast<uint16>(EChaosVDTimelineElementIDFlags::AllPlayback))
-						.OnFrameChanged_Raw(this, &SChaosVDSolverPlaybackControls::OnFrameSelectionUpdated)
-						.OnButtonClicked_Raw(this, &SChaosVDSolverPlaybackControls::HandlePlaybackButtonClicked)
-						.MaxFrames(0)
-					]
-					+SHorizontalBox::Slot()
-					.FillWidth(0.05f)
+					SNew(SBorder)
+					.BorderImage_Raw(this, &SChaosVDSolverPlaybackControls::GetFrameTypeBadgeBrush)
+					.Padding(2.0f)
+					.Content()
 					[
-						SNew(SBorder)
-						.BorderImage_Raw(this, &SChaosVDSolverPlaybackControls::GetFrameTypeBadgeBrush)
-						.Padding(1.5f)
-						.Content()
+						SNew(SBox)
+						.Padding(4.0f,0.0f)
+						.VAlign(VAlign_Center)
 						[
-							SNew(SBox)
-							.VAlign(VAlign_Center)
-							[
-								SNew(STextBlock)
-								.Justification(ETextJustify::Center)
-								.Text_Lambda([this]()->FText{ return bIsReSimFrame ? LOCTEXT("PlaybackViewportWidgetPhysicsFramesResimLabel", "ReSim" ) : LOCTEXT("PlaybackViewportWidgetPhysicsFramesNormalLabel", "Normal" );})
-							]
+							SNew(STextBlock)
+							.Justification(ETextJustify::Center)
+							.Text_Lambda([this]()->FText{ return bIsReSimFrame ? LOCTEXT("PlaybackViewportWidgetPhysicsFramesResimLabel", "ReSim" ) : LOCTEXT("PlaybackViewportWidgetPhysicsFramesNormalLabel", "Normal" );})
 						]
 					]
 				]
 			]
-			+SHorizontalBox::Slot()
-			.FillWidth(0.3f)
+		]
+		+SHorizontalBox::Slot()
+		.FillWidth(0.2f)
+		[
+			SNew(SVerticalBox)
+			+SVerticalBox::Slot()
+			.AutoHeight()
+			.Padding(0.0f, 0.0f, 0.0f, 2.0f)
+			[
+				SNew(STextBlock)
+				.Justification(ETextJustify::Center)
+				.Text_Lambda([this]()->FText{ return FText::Format(LOCTEXT("PlaybackViewportWidgetStepsLabel","Solver Stage: {0}"), FText::AsCultureInvariant(CurrentStepName));})
+			]
+			+SVerticalBox::Slot()
 			[
-				SNew(SVerticalBox)
-				+SVerticalBox::Slot()
-				.AutoHeight()
-				.Padding(0.0f, 0.0f, 0.0f, 2.0f)
+				SNew(SHorizontalBox)
+				+SHorizontalBox::Slot()
+				.FillWidth(0.9f)
 				[
-					SNew(STextBlock)
-					.Justification(ETextJustify::Center)
-					.Text_Lambda([this]()->FText{ return FText::Format(LOCTEXT("PlaybackViewportWidgetStepsLabel","Solver Stage: {0}"), FText::AsCultureInvariant(CurrentStepName));})
+					SAssignNew(StepsTimelineWidget, SChaosVDTimelineWidget)
+					.ButtonVisibilityFlags(static_cast<uint16>(EChaosVDTimelineElementIDFlags::AllManualStepping))
+					.OnFrameLockStateChanged_Raw(this, &SChaosVDSolverPlaybackControls::HandleLockStateChanged)
+					.OnFrameChanged_Raw(this, &SChaosVDSolverPlaybackControls::OnStepSelectionUpdated)
+					.MaxFrames(0)
 				]
-				+SVerticalBox::Slot()
+				+SHorizontalBox::Slot()
+				.AutoWidth()
 				[
-					SNew(SHorizontalBox)
-					+SHorizontalBox::Slot()
-					.FillWidth(0.6f)
-					[
-						
-						SAssignNew(StepsTimelineWidget, SChaosVDTimelineWidget)
-						.ButtonVisibilityFlags(static_cast<uint16>(EChaosVDTimelineElementIDFlags::AllManualStepping))
-						.OnFrameLockStateChanged_Raw(this, &SChaosVDSolverPlaybackControls::HandleLockStateChanged)
-						.OnFrameChanged_Raw(this, &SChaosVDSolverPlaybackControls::OnStepSelectionUpdated)
-						.MaxFrames(0)
-					]
-					+SHorizontalBox::Slot()
-					.AutoWidth()
-					[
-						CreateVisibilityWidget().ToSharedRef()
-					]
+					CreateVisibilityWidget().ToSharedRef()
 				]
 			]
 		]
@@ -285,13 +280,13 @@ const FSlateBrush* SChaosVDSolverPlaybackControls::GetFrameTypeBadgeBrush() cons
 TSharedPtr<SWidget> SChaosVDSolverPlaybackControls::CreateVisibilityWidget()
 {
 	return SNew(SButton)
-	.OnClicked_Raw(this, &SChaosVDSolverPlaybackControls::ToggleSolverVisibility)
-	[
-		SNew( SImage )
-		.Image_Raw(this,&SChaosVDSolverPlaybackControls::GetBrushForCurrentVisibility)
-		.DesiredSizeOverride(FVector2D(16.0f,16.0f))
-		.ColorAndOpacity(FSlateColor::UseForeground())
-	];	
+			.OnClicked_Raw(this, &SChaosVDSolverPlaybackControls::ToggleSolverVisibility)
+			[
+				SNew(SImage)
+				.Image_Raw(this,&SChaosVDSolverPlaybackControls::GetBrushForCurrentVisibility)
+				.DesiredSizeOverride(FVector2D(16.0f,16.0f))
+				.ColorAndOpacity(FSlateColor::UseForeground())
+			];	
 }
 
 void SChaosVDSolverPlaybackControls::OnFrameSelectionUpdated(int32 NewFrameIndex)
diff --git a/Engine/Plugins/Experimental/ChaosVD/Source/ChaosVD/Private/Widgets/SChaosVDSolverTracks.cpp b/Engine/Plugins/Experimental/ChaosVD/Source/ChaosVD/Private/Widgets/SChaosVDSolverTracks.cpp
--- a/Engine/Plugins/Experimental/ChaosVD/Source/ChaosVD/Private/Widgets/SChaosVDSolverTracks.cpp
+++ b/Engine/Plugins/Experimental/ChaosVD/Source/ChaosVD/Private/Widgets/SChaosVDSolverTracks.cpp
@@ -13,10 +13,10 @@ void SChaosVDSolverTracks::Construct(const FArguments& InArgs, TWeakPtr<FChaosVD
 	ChildSlot
 	[
 		SAssignNew(SolverTracksListWidget, SListView<TSharedPtr<FChaosVDTrackInfo>>)
-			.ListItemsSource(&CachedTrackInfoArray)
-			.SelectionMode( ESelectionMode::None )
-			.ListViewStyle(&FAppStyle::Get().GetWidgetStyle<FTableViewStyle>("SimpleListView"))
-			.OnGenerateRow(this, &SChaosVDSolverTracks::MakeSolverTrackControlsFromTrackInfo)
+		.ListItemsSource(&CachedTrackInfoArray)
+		.SelectionMode( ESelectionMode::None )
+		.ListViewStyle(&FAppStyle::Get().GetWidgetStyle<FTableViewStyle>("SimpleListView"))
+		.OnGenerateRow(this, &SChaosVDSolverTracks::MakeSolverTrackControlsFromTrackInfo)
 	];
 
 	ensure(InPlaybackController.IsValid());
@@ -99,7 +99,7 @@ TSharedRef<ITableRow> SChaosVDSolverTracks::MakeSolverTrackControlsFromTrackInfo
 	SAssignNew(PlaybackControlsContainer, SVerticalBox)
 			+SVerticalBox::Slot()
 			.AutoHeight()
-			.Padding(20.0f, 1.0f, 20.0f, 1.0f)
+			.Padding(10.0f, 2.0f, 10.0f, 0.0f)
 			[
 				SNew(SExpandableArea)
 					.InitiallyCollapsed(false)
diff --git a/Engine/Plugins/Experimental/ChaosVD/Source/ChaosVD/Private/Widgets/SChaosVDTimelineWidget.cpp b/Engine/Plugins/Experimental/ChaosVD/Source/ChaosVD/Private/Widgets/SChaosVDTimelineWidget.cpp
--- a/Engine/Plugins/Experimental/ChaosVD/Source/ChaosVD/Private/Widgets/SChaosVDTimelineWidget.cpp
+++ b/Engine/Plugins/Experimental/ChaosVD/Source/ChaosVD/Private/Widgets/SChaosVDTimelineWidget.cpp
@@ -29,10 +29,12 @@ void SChaosVDTimelineWidget::Construct(const FArguments& InArgs)
 			+SHorizontalBox::Slot()
 			.HAlign(HAlign_Center)
 			.VAlign(VAlign_Center)
-			.FillWidth(0.25f)
+			.AutoWidth()
+			.Padding(4.0f,0.0f)
 			[
 				SNew(SHorizontalBox)
 				+SHorizontalBox::Slot()
+				.AutoWidth()
 				[
 					SNew(SButton)
 					.Visibility_Raw(this, &SChaosVDTimelineWidget::GetElementVisibility, EChaosVDTimelineElementIDFlags::Play)
@@ -49,6 +51,7 @@ void SChaosVDTimelineWidget::Construct(const FArguments& InArgs)
 					]
 				]
 				+SHorizontalBox::Slot()
+				.AutoWidth()
 				[
 					SNew(SButton)
 					.Visibility_Raw(this, &SChaosVDTimelineWidget::GetElementVisibility, EChaosVDTimelineElementIDFlags::Stop)
@@ -65,6 +68,7 @@ void SChaosVDTimelineWidget::Construct(const FArguments& InArgs)
 					]
 				]
 				+SHorizontalBox::Slot()
+				.AutoWidth()
 				[
 					SNew(SButton)
 					.Visibility_Raw(this, &SChaosVDTimelineWidget::GetElementVisibility, EChaosVDTimelineElementIDFlags::Prev)
@@ -81,6 +85,7 @@ void SChaosVDTimelineWidget::Construct(const FArguments& InArgs)
 					]
 				]
 				+SHorizontalBox::Slot()
+				.AutoWidth()
 				[
 					SNew(SButton)
 					.Visibility_Raw(this, &SChaosVDTimelineWidget::GetElementVisibility, EChaosVDTimelineElementIDFlags::Next)
@@ -97,6 +102,7 @@ void SChaosVDTimelineWidget::Construct(const FArguments& InArgs)
 					]
 				]
 				+SHorizontalBox::Slot()
+				.AutoWidth()
 				[
 					SNew(SButton)
 					.Visibility_Raw(this, &SChaosVDTimelineWidget::GetElementVisibility, EChaosVDTimelineElementIDFlags::Lock)
@@ -115,7 +121,8 @@ void SChaosVDTimelineWidget::Construct(const FArguments& InArgs)
 			]
 			+SHorizontalBox::Slot()
 			.VAlign(VAlign_Center)
-			.FillWidth(0.65f)
+			.Padding(4.0f,0.0f)
+			.FillWidth(1.0f)
 			[
 			  SAssignNew(TimelineSlider, SSlider)
 			  .Visibility_Raw(this, &SChaosVDTimelineWidget::GetElementVisibility, EChaosVDTimelineElementIDFlags::Timeline)
@@ -129,7 +136,8 @@ void SChaosVDTimelineWidget::Construct(const FArguments& InArgs)
 			  .MinValue(0)
 			]
 			+SHorizontalBox::Slot()
-			.FillWidth(0.1f)
+			.Padding(4.0f,0.0f)
+			.AutoWidth()
 			.HAlign(HAlign_Center)
 			.VAlign(VAlign_Center)
 			[
diff --git a/Engine/Plugins/Experimental/ChaosVD/Source/ChaosVD/Private/Widgets/SChaosVDViewportToolbar.cpp b/Engine/Plugins/Experimental/ChaosVD/Source/ChaosVD/Private/Widgets/SChaosVDViewportToolbar.cpp
--- a/Engine/Plugins/Experimental/ChaosVD/Source/ChaosVD/Private/Widgets/SChaosVDViewportToolbar.cpp
+++ b/Engine/Plugins/Experimental/ChaosVD/Source/ChaosVD/Private/Widgets/SChaosVDViewportToolbar.cpp
@@ -38,7 +38,7 @@ void SChaosVDViewportToolbar::ExtendOptionsMenu(FMenuBuilder& OptionsMenuBuilder
 													.Padding(5.0f)
 													[
 														SNew(STextBlock)
-														.Text(LOCTEXT("Default Options Message", "Nothing so see yet. An options menu will be implemented soon"))
+														.Text(LOCTEXT("Default Options Message", "Nothing to see yet. An options menu will be implemented soon"))
 													];
 
 	// Intentionally leaving the label empty as we only want to show the temp message defined above
diff --git a/Engine/Plugins/Experimental/ChaosVD/Source/ChaosVD/Public/Widgets/SChaosVDMainTab.h b/Engine/Plugins/Experimental/ChaosVD/Source/ChaosVD/Public/Widgets/SChaosVDMainTab.h
--- a/Engine/Plugins/Experimental/ChaosVD/Source/ChaosVD/Public/Widgets/SChaosVDMainTab.h
+++ b/Engine/Plugins/Experimental/ChaosVD/Source/ChaosVD/Public/Widgets/SChaosVDMainTab.h
@@ -2,11 +2,14 @@
 
 #pragma once
 
+#include "Framework/SlateDelegates.h"
 #include "Framework/Docking/TabManager.h"
 #include "Templates/SharedPointer.h"
 #include "Toolkits/IToolkitHost.h"
 #include "Widgets/SCompoundWidget.h"
 
+#include "SChaosVDMainTab.generated.h"
+
 class FComponentVisualizer;
 class FChaosVDEditorModeTools;
 class FChaosVDTabSpawnerBase;
@@ -17,8 +20,19 @@ class FChaosVDOutputLogTab;
 class FChaosVDPlaybackViewportTab;
 class FChaosVDObjectDetailsTab;
 class FChaosVDWorldOutlinerTab;
+class SButton;
 class SDockTab;
 
+UCLASS()
+class CHAOSVD_API UChaosVDMainToolbarMenuContext : public UObject
+{
+	GENERATED_BODY()
+
+public:
+
+	TWeakPtr<class SChaosVDMainTab> MainTab;
+};
+
 /** The main widget containing the Chaos Visual Debugger interface */
 class SChaosVDMainTab : public SCompoundWidget, public IToolkitHost
 {
@@ -58,6 +72,8 @@ public:
 
 private:
 
+	void RegisterMainTabMenu();
+
 	template<typename TabType>
 	void RegisterTabSpawner(FName TabID);
 
@@ -70,7 +86,11 @@ private:
 
 	void GenerateMainWindowMenu();
 
-	void BrowseAndOpenChaosVDRecording();
+	FReply BrowseAndOpenChaosVDRecording();
+
+	TSharedRef<SButton> CreateSimpleButton(TFunction<FText()>&& GetTextDelegate, TFunction<FText()>&& ToolTipTextDelegate, const FSlateBrush* ButtonIcon, const UChaosVDMainToolbarMenuContext* MenuContext, const FOnClicked& InButtonClickedCallback);
+
+	TSharedRef<SWidget> GenerateMainToolbarWidget();
 
 	void BrowseChaosVDRecordingFromFolder(FStringView FolderPath = TEXT(""));
 
@@ -95,6 +115,9 @@ private:
 
 	FReply HandleSessionConnectionClicked();
 	FText GetConnectButtonText() const;
+	FText GetConnectButtonTooltipText() const;
+
+	static inline const FName MainToolBarName = FName("ChaosVD.MainToolBar");
 };
 
 
diff --git a/Engine/Plugins/Experimental/ChaosVD/Source/ChaosVD/Public/Widgets/SChaosVDRecordingControls.h b/Engine/Plugins/Experimental/ChaosVD/Source/ChaosVD/Public/Widgets/SChaosVDRecordingControls.h
--- a/Engine/Plugins/Experimental/ChaosVD/Source/ChaosVD/Public/Widgets/SChaosVDRecordingControls.h
+++ b/Engine/Plugins/Experimental/ChaosVD/Source/ChaosVD/Public/Widgets/SChaosVDRecordingControls.h
@@ -6,6 +6,8 @@
 #include "Templates/SharedPointer.h"
 #include "Widgets/SCompoundWidget.h"
 
+#include "SChaosVDRecordingControls.generated.h"
+
 namespace Chaos::VisualDebugger
 {
 	struct FChaosVDOptionalDataChannel;
@@ -16,6 +18,16 @@ class SChaosVDMainTab;
 class FReply;
 struct FSlateBrush;
 
+UCLASS()
+class CHAOSVD_API UChaosVDRecordingToolbarMenuContext : public UObject
+{
+	GENERATED_BODY()
+
+public:
+
+	TWeakPtr<class SChaosVDRecordingControls> RecordingControlsWidget;
+};
+
 UENUM()
 enum class EChaosVDRecordingMode
 {
@@ -37,8 +49,11 @@ public:
 
 protected:
 
-	TSharedRef<SButton> GenerateToggleRecordingStateButton(EChaosVDRecordingMode RecordingMode, const FText& StartRecordingTooltip);
+	TSharedRef<SWidget> GenerateToggleRecordingStateButton(EChaosVDRecordingMode RecordingMode, const FText& StartRecordingTooltip);
 	TSharedRef<SWidget> GenerateDataChannelsMenu();
+	TSharedRef<SWidget> GenerateDataChannelsButton();
+	TSharedRef<SWidget> GenerateRecordingTimeTextBlock();
+	TSharedRef<SWidget> GenerateToolbarWidget();
 
 #if WITH_CHAOS_VISUAL_DEBUGGER
 	void ToggleChannelEnabledState(TWeakPtr<FCVDDataChannel> Channel);
@@ -59,6 +74,8 @@ protected:
 	bool IsRecordingToggleButtonEnabled(EChaosVDRecordingMode RecordingMode) const;
 	EVisibility IsRecordingToggleButtonVisible(EChaosVDRecordingMode RecordingMode) const;
 
+	void RegisterMenus();
+
 	bool IsRecording() const;
 
 	FText GetRecordingTimeText() const;
@@ -77,4 +94,10 @@ protected:
 	float IntervalBetweenAutoplayConnectionAttemptsSeconds = 0.1f;
 	bool bAutoConnectionAttemptInProgress = false;
 	int32 CurrentConnectionAttempts = 0;
+
+	bool bRecordingButtonHovered = false;
+
+	FCurveSequence RecordingAnimation;
+
+	static inline const FName RecordingControlsToolbarName = FName("ChaosVD.MainToolBar.RecordingControls");
 };
