[ChaosVD] Better handling cases where the CVD recording instance might fail to load or be unloaded before the UI its updated

Added missing null checks, and some ensures to catch cases where this is happening to confirm the access call when the recording it is not valid it is coming from the UI.

Also, added dedicated message dialogs to inform the user when a recording failed to load instead of fail silently.

diff --git a/Engine/Plugins/Experimental/ChaosVD/Source/ChaosVD/Private/ChaosVDPlaybackController.cpp b/Engine/Plugins/Experimental/ChaosVD/Source/ChaosVD/Private/ChaosVDPlaybackController.cpp
--- a/Engine/Plugins/Experimental/ChaosVD/Source/ChaosVD/Private/ChaosVDPlaybackController.cpp
+++ b/Engine/Plugins/Experimental/ChaosVD/Source/ChaosVD/Private/ChaosVDPlaybackController.cpp
@@ -9,10 +9,13 @@
 #include "ChaosVDRuntimeModule.h"
 #include "ChaosVDScene.h"
 #include "Actors/ChaosVDSolverInfoActor.h"
+#include "Misc/MessageDialog.h"
 #include "Trace/ChaosVDTraceManager.h"
 #include "Trace/ChaosVDTraceProvider.h"
 #include "TraceServices/Model/AnalysisSession.h"
 
+#define LOCTEXT_NAMESPACE "ChaosVisualDebugger"
+
 FChaosVDPlaybackController::FChaosVDPlaybackController(const TWeakPtr<FChaosVDScene>& InSceneToControl)
 {
 	SceneToControl = InSceneToControl;
@@ -46,8 +49,9 @@ FChaosVDPlaybackController::~FChaosVDPlaybackController()
 
 bool FChaosVDPlaybackController::LoadChaosVDRecordingFromTraceSession(const FChaosVDTraceSessionDescriptor& InSessionDescriptor)
 {
-	if (InSessionDescriptor.SessionName.IsEmpty())
+	if (!ensure(!InSessionDescriptor.SessionName.IsEmpty()))
 	{
+		FMessageDialog::Open(EAppMsgType::Ok, LOCTEXT("LoadRecordingSessionFailedMessage", "Failed to create a session from the selected CVD recording. Please see the logs for more details... "));
 		return false;
 	}
 
@@ -64,8 +68,10 @@ bool FChaosVDPlaybackController::LoadChaosVDRecordingFromTraceSession(const FCha
 		}
 	}
 
-	if (!LoadedRecording.IsValid())
+	if (!ensure(LoadedRecording.IsValid()))
 	{
+		FMessageDialog::Open(EAppMsgType::Ok, LOCTEXT("LoadRecordingFailedMessage", "Failed to load the selected CVD recording. Please see the logs for more details... "));
+
 		return false;
 	}
 
@@ -117,6 +123,11 @@ void FChaosVDPlaybackController::UnloadCurrentRecording(EChaosVDUnloadRecordingF
 
 void FChaosVDPlaybackController::PlayFromClosestKeyFrame_AssumesLocked(const int32 InTrackID, const int32 FrameNumber, FChaosVDScene& InSceneToControl) const
 {
+	if (!LoadedRecording.IsValid())
+	{
+		return;
+	}
+
 	const int32 KeyFrameNumber = LoadedRecording->FindFirstSolverKeyFrameNumberFromFrame_AssumesLocked(InTrackID, FrameNumber);
 	if (KeyFrameNumber < 0)
 	{
@@ -289,6 +300,11 @@ void FChaosVDPlaybackController::GoToRecordedGameFrame_AssumesLocked(const int32
 
 void FChaosVDPlaybackController::GoToTrackFrame(FGuid InstigatorID, EChaosVDTrackType TrackType, int32 InTrackID, int32 FrameNumber, int32 Step)
 {
+	if (!ensure(LoadedRecording.IsValid()))
+	{
+		return;
+	}
+
 	FReadScopeLock ReadLock(LoadedRecording->GetRecordingDataLock());
 	GoToTrackFrame_AssumesLocked(InstigatorID, TrackType, InTrackID, FrameNumber, Step);
 }
@@ -412,6 +428,11 @@ int32 FChaosVDPlaybackController::ConvertCurrentFrameToOtherTrackFrame(const FCh
 		ensureMsgf(false, TEXT("One of provided track infos is not valid"));
 		return INDEX_NONE;
 	}
+	
+	if (!ensure(LoadedRecording.IsValid()))
+    {
+    	return INDEX_NONE;
+    }
 
 	const bool bBothTracksHaveSameID = FromTrack->TrackID == ToTrack->TrackID;
 	const bool bBothTracksHaveSameType = FromTrack->TrackType == ToTrack->TrackType;
@@ -611,31 +632,34 @@ bool FChaosVDPlaybackController::Tick(float DeltaTime)
 		}
 	}
 
-	// Load at least the first frame
-	if (!bPlayedFirstFrame)
+	if (LoadedRecording.IsValid())
 	{
-		if (LoadedRecording.IsValid() && LoadedRecording->GetAvailableSolversNumber_AssumesLocked() > 0)
+		// Load at least the first frame
+		if (!bPlayedFirstFrame)
 		{
-			constexpr int32 GameFrameToLoad = 0;
-			constexpr int32 Step = 0;
-			GoToTrackFrame(IChaosVDPlaybackControllerInstigator::InvalidGuid, EChaosVDTrackType::Game, GameTrackID, GameFrameToLoad, Step);
-			bPlayedFirstFrame = true;
+			if (LoadedRecording->GetAvailableSolversNumber_AssumesLocked() > 0)
+			{
+				constexpr int32 GameFrameToLoad = 0;
+				constexpr int32 Step = 0;
+				GoToTrackFrame(IChaosVDPlaybackControllerInstigator::InvalidGuid, EChaosVDTrackType::Game, GameTrackID, GameFrameToLoad, Step);
+				bPlayedFirstFrame = true;
+			}
 		}
-	}
 
-	// If we are live, make sure we don't lag too much behind
-	if (!bPauseRequested && IsPlayingLiveSession())
-	{
-		if (const FChaosVDTrackInfo* GameTrackInfo = GetTrackInfo(EChaosVDTrackType::Game, GameTrackID))
+		// If we are live, make sure we don't lag too much behind
+		if (!bPauseRequested && IsPlayingLiveSession())
 		{
-			const int32 CurrentFrameDeltaFromLast = FMath::Abs(GameTrackInfo->MaxFrames - GameTrackInfo->CurrentFrame);
-			if (CurrentFrameDeltaFromLast > MaxFramesLaggingBehindDuringLiveSession)
+			if (const FChaosVDTrackInfo* GameTrackInfo = GetTrackInfo(EChaosVDTrackType::Game, GameTrackID))
 			{
-				// Playing the middle point between last and the threshold. We don't want to play the last available frame as it could be incomplete,
-				// and we don't want to go to close to the threshold.
-				const int32 GameFrameToLoad = LoadedRecording->GetAvailableGameFramesNumber() - MinFramesLaggingBehindDuringLiveSession;
-				constexpr int32 Step = 0;
-				GoToTrackFrame(IChaosVDPlaybackControllerInstigator::InvalidGuid, EChaosVDTrackType::Game, GameTrackID, GameFrameToLoad, Step);
+				const int32 CurrentFrameDeltaFromLast = FMath::Abs(GameTrackInfo->MaxFrames - GameTrackInfo->CurrentFrame);
+				if (CurrentFrameDeltaFromLast > MaxFramesLaggingBehindDuringLiveSession)
+				{
+					// Playing the middle point between last and the threshold. We don't want to play the last available frame as it could be incomplete,
+					// and we don't want to go to close to the threshold.
+					const int32 GameFrameToLoad = LoadedRecording->GetAvailableGameFramesNumber() - MinFramesLaggingBehindDuringLiveSession;
+					constexpr int32 Step = 0;
+					GoToTrackFrame(IChaosVDPlaybackControllerInstigator::InvalidGuid, EChaosVDTrackType::Game, GameTrackID, GameFrameToLoad, Step);
+				}
 			}
 		}
 	}
@@ -761,6 +785,17 @@ void FChaosVDPlaybackController::UpdateTrackVisibility(EChaosVDTrackType Type, i
 
 void FChaosVDPlaybackController::UpdateSolverTracksData()
 {
+	if (!LoadedRecording.IsValid())
+	{
+		// If the recording is no longer valid, clear any existing solver track info data so the UI can be updated accordingly
+		if (TrackInfoByIDMap* SolverTracks = TrackInfoPerType.Find(EChaosVDTrackType::Solver))
+		{
+			SolverTracks->Empty();
+		}
+
+		return;
+	}
+
 	const TMap<int32, TArray<FChaosVDSolverFrameData>>& SolversByID = LoadedRecording->GetAvailableSolvers_AssumesLocked();
 	for (const TPair<int32, TArray<FChaosVDSolverFrameData>>& SolverIDPair : SolversByID)
 	{
@@ -795,7 +830,7 @@ void FChaosVDPlaybackController::HandleCurrentRecordingUpdated()
 		GameTrackInfo->CurrentStep = 0;
 	}
 
-	GameTrackInfo->MaxFrames = LoadedRecording->GetAvailableGameFrames_AssumesLocked().Num();
+	GameTrackInfo->MaxFrames = LoadedRecording.IsValid() ? LoadedRecording->GetAvailableGameFrames_AssumesLocked().Num() : INDEX_NONE;
 	GameTrackInfo->TrackType = EChaosVDTrackType::Game;
 
 	// Each time the recording is updated, populate or update the existing solver tracks data
@@ -803,3 +838,5 @@ void FChaosVDPlaybackController::HandleCurrentRecordingUpdated()
 
 	bHasPendingGTUpdateBroadcast = true;
 }
+
+#undef LOCTEXT_NAMESPACE
