[ChaosVD] Implemented a simple Archive Header so we can ensure backwards compatibility using the existing Custom Versionin system

Until now, as CVD was not specifying any custom versions on load, everything defaulted to the available versions on the versions registry. This means old files would be de-serialized with a potential wrong version and failing in the process.

Now we record all available custom versions at the start of a recording, so then we can use these versions in any Archive we create during load of a CVD file, ensuring we use the correct versions for the recorded file

diff --git a/Engine/Plugins/Experimental/ChaosVD/Source/ChaosVD/Private/Trace/ChaosVDTraceAnalyzer.cpp b/Engine/Plugins/Experimental/ChaosVD/Source/ChaosVD/Private/Trace/ChaosVDTraceAnalyzer.cpp
--- a/Engine/Plugins/Experimental/ChaosVD/Source/ChaosVD/Private/Trace/ChaosVDTraceAnalyzer.cpp
+++ b/Engine/Plugins/Experimental/ChaosVD/Source/ChaosVD/Private/Trace/ChaosVDTraceAnalyzer.cpp
@@ -182,7 +182,7 @@ bool FChaosVDTraceAnalyzer::OnEvent(uint16 RouteId, EStyle Style, const FOnEvent
 			{
 				// This can happen during live debugging as we miss some of the events at the beginning.
 				// Loading a trace file that was recorded as part of a live session, will have the same issue.
-				UE_LOG(LogChaosVDEditor, Warning, TEXT("[%s] FailedToProcess Binary Data with ID [%d]"), ANSI_TO_TCHAR(__FUNCTION__), DataID);
+				UE_LOG(LogChaosVDEditor, Verbose, TEXT("[%s] FailedToProcess Binary Data with ID [%d]"), ANSI_TO_TCHAR(__FUNCTION__), DataID);
 			}
 
 			break;
diff --git a/Engine/Plugins/Experimental/ChaosVD/Source/ChaosVD/Private/Trace/ChaosVDTraceProvider.cpp b/Engine/Plugins/Experimental/ChaosVD/Source/ChaosVD/Private/Trace/ChaosVDTraceProvider.cpp
--- a/Engine/Plugins/Experimental/ChaosVD/Source/ChaosVD/Private/Trace/ChaosVDTraceProvider.cpp
+++ b/Engine/Plugins/Experimental/ChaosVD/Source/ChaosVD/Private/Trace/ChaosVDTraceProvider.cpp
@@ -5,12 +5,10 @@
 #include "ChaosVDModule.h"
 #include "ChaosVDRecording.h"
 
-#include "Chaos/ChaosArchive.h"
 #include "ChaosVisualDebugger/ChaosVDSerializedNameTable.h"
 #include "ChaosVisualDebugger/ChaosVisualDebuggerTrace.h"
 
 #include "Compression/OodleDataCompressionUtil.h"
-#include "Serialization/MemoryReader.h"
 #include "Trace/DataProcessors/ChaosVDConstraintDataProcessor.h"
 #include "Trace/DataProcessors/ChaosVDJointConstraintDataProcessor.h"
 #include "Trace/DataProcessors/ChaosVDMidPhaseDataProcessor.h"
@@ -19,11 +17,14 @@
 #include "Trace/DataProcessors/ChaosVDSerializedNameEntryDataProcessor.h"
 #include "Trace/DataProcessors/ChaosVDTraceImplicitObjectProcessor.h"
 #include "Trace/DataProcessors/ChaosVDTraceParticleDataProcessor.h"
+#include "Trace/DataProcessors/ChaosVDArchiveHeaderProcessor.h"
+#include "Trace/DataProcessors/IChaosVDDataProcessor.h"
 
 FName FChaosVDTraceProvider::ProviderName("ChaosVDProvider");
 
 FChaosVDTraceProvider::FChaosVDTraceProvider(TraceServices::IAnalysisSession& InSession): Session(InSession)
 {
+	DefaultHeaderData = Chaos::VisualDebugger::FChaosVDArchiveHeader::Current();
 }
 
 void FChaosVDTraceProvider::CreateRecordingInstanceForSession(const FString& InSessionName)
@@ -163,6 +164,10 @@ bool FChaosVDTraceProvider::ProcessBinaryData(const int32 DataID)
 					{
 						return true;
 					}
+					else
+					{
+						UE_LOG(LogChaosVDEditor, Warning, TEXT("[%s] Failed to serialize Binary Data with ID [%d] | Type [%s]"), ANSI_TO_TCHAR(__FUNCTION__), DataID, *UnprocessedData->TypeName);
+					}
 				}
 			}
 			else
@@ -185,11 +190,6 @@ void FChaosVDTraceProvider::RegisterDataProcessor(TSharedPtr<IChaosVDDataProcess
 	RegisteredDataProcessors.Add(InDataProcessor->GetCompatibleTypeName(), InDataProcessor);
 }
 
-TSharedPtr<Chaos::VisualDebugger::FChaosVDSerializableNameTable> FChaosVDTraceProvider::GetNameTable() const
-{
-	return InternalRecording.IsValid() ? InternalRecording->GetNameTableInstance() : nullptr; 
-}
-
 void FChaosVDTraceProvider::RegisterDefaultDataProcessorsIfNeeded()
 {
 	if (bDefaultDataProcessorsRegistered)
@@ -229,6 +229,10 @@ void FChaosVDTraceProvider::RegisterDefaultDataProcessorsIfNeeded()
 	JointConstraintDataProcessor->SetTraceProvider(AsShared());
 	RegisterDataProcessor(JointConstraintDataProcessor);
 
+	TSharedPtr<FChaosVDArchiveHeaderProcessor> ArchiveHeaderDataProcessor = MakeShared<FChaosVDArchiveHeaderProcessor>();
+	ArchiveHeaderDataProcessor->SetTraceProvider(AsShared());
+	RegisterDataProcessor(ArchiveHeaderDataProcessor);
+
 	bDefaultDataProcessorsRegistered = true;
 }
 
diff --git a/Engine/Plugins/Experimental/ChaosVD/Source/ChaosVD/Private/Trace/DataProcessors/ChaosVDArchiveHeaderProcessor.cpp b/Engine/Plugins/Experimental/ChaosVD/Source/ChaosVD/Private/Trace/DataProcessors/ChaosVDArchiveHeaderProcessor.cpp
new file mode 100644
--- /dev/null
+++ b/Engine/Plugins/Experimental/ChaosVD/Source/ChaosVD/Private/Trace/DataProcessors/ChaosVDArchiveHeaderProcessor.cpp
@@ -0,0 +1,43 @@
+// Copyright Epic Games, Inc. All Rights Reserved.
+
+#include "Trace/DataProcessors/ChaosVDArchiveHeaderProcessor.h"
+
+#include "ChaosVDModule.h"
+
+
+FChaosVDArchiveHeaderProcessor::FChaosVDArchiveHeaderProcessor() : IChaosVDDataProcessor(Chaos::VisualDebugger::FChaosVDArchiveHeader::WrapperTypeName)
+{
+}
+
+bool FChaosVDArchiveHeaderProcessor::ProcessRawData(const TArray<uint8>& InData)
+{
+	TSharedPtr<FChaosVDTraceProvider> ProviderSharedPtr = TraceProvider.Pin();
+	if (!ensure(ProviderSharedPtr.IsValid()))
+	{
+		return false;
+	}
+
+	Chaos::VisualDebugger::FChaosVDArchiveHeader RecordedHeaderData;
+	const bool bSuccess = ReadDataFromBuffer(InData, RecordedHeaderData, ProviderSharedPtr.ToSharedRef());
+
+	if (const TSharedPtr<FChaosVDRecording> RecordingInstance = ProviderSharedPtr->GetRecordingForSession())
+	{
+		// This works under the assumption the header is the first thing written and therefore the first thing read
+		// If that didn't happen, we need to know to investigate further
+
+		// Note: This is not a fatal error, unless pretty drastic serialization changes were made and the loaded file is old.
+		// CVD can gracefully handle serialization errors (as long the types serializers themselves don't crash and properly error out as expected).
+		if (!ensure(RecordingInstance->IsEmpty()))
+		{
+			UE_LOG(LogChaosVDEditor, Warning, TEXT("[%s] Processed an archive header while the recording already had data loaded. That initially loaded data used the default header and serialization errors might have occured | This should not happen..."), ANSI_TO_TCHAR(__FUNCTION__));
+		}
+
+		RecordingInstance->SetHeaderData(RecordedHeaderData);
+	}
+	else
+	{
+		UE_LOG(LogChaosVDEditor, Error, TEXT("[%s] Processed an archive header without a valid recording instace. Any follwoing archive reads will use the default header and serialization errors might occur..."), ANSI_TO_TCHAR(__FUNCTION__));	
+	}
+
+	return bSuccess;
+}
diff --git a/Engine/Plugins/Experimental/ChaosVD/Source/ChaosVD/Private/Trace/DataProcessors/ChaosVDConstraintDataProcessor.cpp b/Engine/Plugins/Experimental/ChaosVD/Source/ChaosVD/Private/Trace/DataProcessors/ChaosVDConstraintDataProcessor.cpp
--- a/Engine/Plugins/Experimental/ChaosVD/Source/ChaosVD/Private/Trace/DataProcessors/ChaosVDConstraintDataProcessor.cpp
+++ b/Engine/Plugins/Experimental/ChaosVD/Source/ChaosVD/Private/Trace/DataProcessors/ChaosVDConstraintDataProcessor.cpp
@@ -5,7 +5,6 @@
 #include "ChaosVDRecording.h"
 #include "ChaosVisualDebugger/ChaosVDMemWriterReader.h"
 #include "DataWrappers/ChaosVDCollisionDataWrappers.h"
-#include "Serialization/MemoryReader.h"
 #include "Trace/ChaosVDTraceProvider.h"
 
 FChaosVDConstraintDataProcessor::FChaosVDConstraintDataProcessor() : IChaosVDDataProcessor(FChaosVDConstraint::WrapperTypeName)
@@ -20,13 +19,8 @@ bool FChaosVDConstraintDataProcessor::ProcessRawData(const TArray<uint8>& InData
 		return false;
 	}
 
-	if (!ensure(ProviderSharedPtr->GetNameTable().IsValid()))
-	{
-		return false;
-	}
-
 	FChaosVDConstraint RecordedConstraint;
-	const bool bSuccess = ReadDataFromBuffer(InData, RecordedConstraint, ProviderSharedPtr->GetNameTable().ToSharedRef());
+	const bool bSuccess = Chaos::VisualDebugger::ReadDataFromBuffer(InData, RecordedConstraint, ProviderSharedPtr.ToSharedRef());
 
 	if (bSuccess)
 	{
diff --git a/Engine/Plugins/Experimental/ChaosVD/Source/ChaosVD/Private/Trace/DataProcessors/ChaosVDJointConstraintDataProcessor.cpp b/Engine/Plugins/Experimental/ChaosVD/Source/ChaosVD/Private/Trace/DataProcessors/ChaosVDJointConstraintDataProcessor.cpp
--- a/Engine/Plugins/Experimental/ChaosVD/Source/ChaosVD/Private/Trace/DataProcessors/ChaosVDJointConstraintDataProcessor.cpp
+++ b/Engine/Plugins/Experimental/ChaosVD/Source/ChaosVD/Private/Trace/DataProcessors/ChaosVDJointConstraintDataProcessor.cpp
@@ -19,13 +19,8 @@ bool FChaosVDJointConstraintDataProcessor::ProcessRawData(const TArray<uint8>& I
 		return false;
 	}
 
-	if (!ensure(ProviderSharedPtr->GetNameTable().IsValid()))
-	{
-		return false;
-	}
-
 	TSharedPtr<FChaosVDJointConstraint> JointConstraint = MakeShared<FChaosVDJointConstraint>();
-	const bool bSuccess = ReadDataFromBuffer(InData, *JointConstraint, ProviderSharedPtr->GetNameTable().ToSharedRef());
+	const bool bSuccess = Chaos::VisualDebugger::ReadDataFromBuffer(InData, *JointConstraint, ProviderSharedPtr.ToSharedRef());
 
 	if (bSuccess)
 	{
diff --git a/Engine/Plugins/Experimental/ChaosVD/Source/ChaosVD/Private/Trace/DataProcessors/ChaosVDMidPhaseDataProcessor.cpp b/Engine/Plugins/Experimental/ChaosVD/Source/ChaosVD/Private/Trace/DataProcessors/ChaosVDMidPhaseDataProcessor.cpp
--- a/Engine/Plugins/Experimental/ChaosVD/Source/ChaosVD/Private/Trace/DataProcessors/ChaosVDMidPhaseDataProcessor.cpp
+++ b/Engine/Plugins/Experimental/ChaosVD/Source/ChaosVD/Private/Trace/DataProcessors/ChaosVDMidPhaseDataProcessor.cpp
@@ -17,16 +17,10 @@ bool FChaosVDMidPhaseDataProcessor::ProcessRawData(const TArray<uint8>& InData)
 	{
 		return false;
 	}
-	
-	if (!ensure(ProviderSharedPtr->GetNameTable().IsValid()))
-	{
-		return false;
-	}
-
 
 	TSharedPtr<FChaosVDParticlePairMidPhase> MidPhase = MakeShared<FChaosVDParticlePairMidPhase>();
 
-	const bool bSuccess = ReadDataFromBuffer(InData, *MidPhase, ProviderSharedPtr->GetNameTable().ToSharedRef());
+	const bool bSuccess = Chaos::VisualDebugger::ReadDataFromBuffer(InData, *MidPhase, ProviderSharedPtr.ToSharedRef());
 
 	if (bSuccess)
 	{
diff --git a/Engine/Plugins/Experimental/ChaosVD/Source/ChaosVD/Private/Trace/DataProcessors/ChaosVDSceneQueryDataProcessor.cpp b/Engine/Plugins/Experimental/ChaosVD/Source/ChaosVD/Private/Trace/DataProcessors/ChaosVDSceneQueryDataProcessor.cpp
--- a/Engine/Plugins/Experimental/ChaosVD/Source/ChaosVD/Private/Trace/DataProcessors/ChaosVDSceneQueryDataProcessor.cpp
+++ b/Engine/Plugins/Experimental/ChaosVD/Source/ChaosVD/Private/Trace/DataProcessors/ChaosVDSceneQueryDataProcessor.cpp
@@ -20,13 +20,9 @@ bool FChaosVDSceneQueryDataProcessor::ProcessRawData(const TArray<uint8>& InData
 	{
 		return false;
 	}
-	if (!ensure(ProviderSharedPtr->GetNameTable().IsValid()))
-	{
-		return false;
-	}
 
 	const TSharedPtr<FChaosVDQueryDataWrapper> QueryData = MakeShared<FChaosVDQueryDataWrapper>();
-	const bool bSuccess = ReadDataFromBuffer(InData, *QueryData, ProviderSharedPtr->GetNameTable().ToSharedRef());
+	const bool bSuccess = Chaos::VisualDebugger::ReadDataFromBuffer(InData, *QueryData, ProviderSharedPtr.ToSharedRef());
 
 	if (bSuccess)
 	{
diff --git a/Engine/Plugins/Experimental/ChaosVD/Source/ChaosVD/Private/Trace/DataProcessors/ChaosVDSceneQueryVisitDataProcessor.cpp b/Engine/Plugins/Experimental/ChaosVD/Source/ChaosVD/Private/Trace/DataProcessors/ChaosVDSceneQueryVisitDataProcessor.cpp
--- a/Engine/Plugins/Experimental/ChaosVD/Source/ChaosVD/Private/Trace/DataProcessors/ChaosVDSceneQueryVisitDataProcessor.cpp
+++ b/Engine/Plugins/Experimental/ChaosVD/Source/ChaosVD/Private/Trace/DataProcessors/ChaosVDSceneQueryVisitDataProcessor.cpp
@@ -4,7 +4,6 @@
 
 #include "ChaosVDRecording.h"
 #include "ChaosVisualDebugger/ChaosVDMemWriterReader.h"
-#include "ChaosVisualDebugger/ChaosVDSerializedNameTable.h"
 #include "ChaosVisualDebugger/ChaosVisualDebuggerTrace.h"
 #include "DataWrappers/ChaosVDQueryDataWrappers.h"
 #include "Trace/ChaosVDTraceProvider.h"
@@ -21,13 +20,8 @@ bool FChaosVDSceneQueryVisitDataProcessor::ProcessRawData(const TArray<uint8>& I
 		return false;
 	}
 
-	if (!ensure(ProviderSharedPtr->GetNameTable().IsValid()))
-	{
-		return false;
-	}
-
 	FChaosVDQueryVisitStep VisitStepData;
-	const bool bSuccess = ReadDataFromBuffer(InData, VisitStepData, ProviderSharedPtr->GetNameTable().ToSharedRef());
+	const bool bSuccess = Chaos::VisualDebugger::ReadDataFromBuffer(InData, VisitStepData, ProviderSharedPtr.ToSharedRef());
 
 	if (bSuccess)
 	{
@@ -43,7 +37,7 @@ bool FChaosVDSceneQueryVisitDataProcessor::ProcessRawData(const TArray<uint8>& I
 						// Quick and dirty way of show the hits in the details panel. If copying this data around becomes a bottle neck we can write a customization layout for it
 						QueryDataPtr->Hits.Add(VisitStepData);
 					}
-	
+
 					QueryDataPtr->SQVisitData.Emplace(MoveTemp(VisitStepData));
 				}
 			}
diff --git a/Engine/Plugins/Experimental/ChaosVD/Source/ChaosVD/Private/Trace/DataProcessors/ChaosVDSerializedNameEntryDataProcessor.cpp b/Engine/Plugins/Experimental/ChaosVD/Source/ChaosVD/Private/Trace/DataProcessors/ChaosVDSerializedNameEntryDataProcessor.cpp
--- a/Engine/Plugins/Experimental/ChaosVD/Source/ChaosVD/Private/Trace/DataProcessors/ChaosVDSerializedNameEntryDataProcessor.cpp
+++ b/Engine/Plugins/Experimental/ChaosVD/Source/ChaosVD/Private/Trace/DataProcessors/ChaosVDSerializedNameEntryDataProcessor.cpp
@@ -20,12 +20,25 @@ bool FChaosVDSerializedNameEntryDataProcessor::ProcessRawData(const TArray<uint8
 		return false;
 	}
 
+	TSharedPtr<FChaosVDRecording> RecordingInstance = ProviderSharedPtr->GetRecordingForSession();
+
+	TSharedPtr<Chaos::VisualDebugger::FChaosVDSerializableNameTable> NameTableInstance = RecordingInstance ? RecordingInstance->GetNameTableInstance() : nullptr;
+
+	if (!ensure(NameTableInstance.IsValid()))
+	{
+		return false;
+	}
+
 	Chaos::VisualDebugger::FChaosVDSerializedNameEntry NameEntry;
 	
 	FMemoryReader MemReader(InData);
+
+	const Chaos::VisualDebugger::FChaosVDArchiveHeader& RecordedHeader = RecordingInstance->GetHeaderData();
+	ApplyHeaderDataToArchive(MemReader, RecordedHeader);
+	
 	MemReader << NameEntry;
 
-	ProviderSharedPtr->GetNameTable()->AddNameToTable(NameEntry);
+	NameTableInstance->AddNameToTable(NameEntry);
 
 	return true;
 }
diff --git a/Engine/Plugins/Experimental/ChaosVD/Source/ChaosVD/Private/Trace/DataProcessors/ChaosVDTraceImplicitObjectProcessor.cpp b/Engine/Plugins/Experimental/ChaosVD/Source/ChaosVD/Private/Trace/DataProcessors/ChaosVDTraceImplicitObjectProcessor.cpp
--- a/Engine/Plugins/Experimental/ChaosVD/Source/ChaosVD/Private/Trace/DataProcessors/ChaosVDTraceImplicitObjectProcessor.cpp
+++ b/Engine/Plugins/Experimental/ChaosVD/Source/ChaosVD/Private/Trace/DataProcessors/ChaosVDTraceImplicitObjectProcessor.cpp
@@ -5,10 +5,8 @@
 #include "Chaos/ImplicitObject.h"
 #include "ChaosVDRecording.h"
 #include "ChaosVisualDebugger/ChaosVDMemWriterReader.h"
-#include "ChaosVisualDebugger/ChaosVDSerializedNameTable.h"
 #include "ChaosVisualDebugger/ChaosVisualDebuggerTrace.h"
 #include "DataWrappers/ChaosVDImplicitObjectDataWrapper.h"
-#include "Serialization/MemoryReader.h"
 #include "Trace/ChaosVDTraceProvider.h"
 
 using FChaosVDImplicitObjectWrapper = FChaosVDImplicitObjectDataWrapper<Chaos::FImplicitObjectPtr, Chaos::FChaosArchive>;
@@ -25,13 +23,9 @@ bool FChaosVDTraceImplicitObjectProcessor::ProcessRawData(const TArray<uint8>& I
 		return false;
 	}
 
-	if (!ensure(ProviderSharedPtr->GetNameTable().IsValid()))
-	{
-		return false;
-	}
-
 	FChaosVDImplicitObjectWrapper WrappedGeometryData;
-	const bool bSuccess = ReadDataFromBuffer<FChaosVDImplicitObjectWrapper, Chaos::FChaosArchive>(InData, WrappedGeometryData, ProviderSharedPtr->GetNameTable().ToSharedRef());
+
+	const bool bSuccess = Chaos::VisualDebugger::ReadDataFromBuffer<FChaosVDImplicitObjectWrapper>(InData, WrappedGeometryData, ProviderSharedPtr.ToSharedRef());
 
 	if (bSuccess)
 	{
diff --git a/Engine/Plugins/Experimental/ChaosVD/Source/ChaosVD/Private/Trace/DataProcessors/ChaosVDTraceParticleDataProcessor.cpp b/Engine/Plugins/Experimental/ChaosVD/Source/ChaosVD/Private/Trace/DataProcessors/ChaosVDTraceParticleDataProcessor.cpp
--- a/Engine/Plugins/Experimental/ChaosVD/Source/ChaosVD/Private/Trace/DataProcessors/ChaosVDTraceParticleDataProcessor.cpp
+++ b/Engine/Plugins/Experimental/ChaosVD/Source/ChaosVD/Private/Trace/DataProcessors/ChaosVDTraceParticleDataProcessor.cpp
@@ -4,7 +4,6 @@
 
 #include "ChaosVDRecording.h"
 #include "ChaosVisualDebugger/ChaosVDMemWriterReader.h"
-#include "ChaosVisualDebugger/ChaosVDSerializedNameTable.h"
 #include "ChaosVisualDebugger/ChaosVisualDebuggerTrace.h"
 #include "DataWrappers/ChaosVDParticleDataWrapper.h"
 #include "Trace/ChaosVDTraceProvider.h"
@@ -21,13 +20,8 @@ bool FChaosVDTraceParticleDataProcessor::ProcessRawData(const TArray<uint8>& InD
 		return false;
 	}
 
-	if (!ensure(ProviderSharedPtr->GetNameTable().IsValid()))
-	{
-		return false;
-	}
-
 	TSharedPtr<FChaosVDParticleDataWrapper> ParticleData = MakeShared<FChaosVDParticleDataWrapper>();
-	const bool bSuccess = ReadDataFromBuffer(InData, *ParticleData, ProviderSharedPtr->GetNameTable().ToSharedRef());
+	const bool bSuccess = Chaos::VisualDebugger::ReadDataFromBuffer(InData, *ParticleData, ProviderSharedPtr.ToSharedRef());
 
 	if (bSuccess)
 	{
diff --git a/Engine/Plugins/Experimental/ChaosVD/Source/ChaosVD/Public/Trace/ChaosVDTraceProvider.h b/Engine/Plugins/Experimental/ChaosVD/Source/ChaosVD/Public/Trace/ChaosVDTraceProvider.h
--- a/Engine/Plugins/Experimental/ChaosVD/Source/ChaosVD/Public/Trace/ChaosVDTraceProvider.h
+++ b/Engine/Plugins/Experimental/ChaosVD/Source/ChaosVD/Public/Trace/ChaosVDTraceProvider.h
@@ -8,7 +8,6 @@
 #include "Containers/Array.h"
 #include "Containers/Queue.h"
 #include "Containers/UnrealString.h"
-#include "DataProcessors/IChaosVDDataProcessor.h"
 #include "Delegates/Delegate.h"
 #include "Templates/SharedPointer.h"
 #include "TraceServices/Model/AnalysisSession.h"
@@ -22,6 +21,7 @@ struct FChaosVDGameFrameData;
 class FChaosVDEngine;
 struct FChaosVDSolverFrameData;
 struct FChaosVDRecording;
+class IChaosVDDataProcessor;
 
 struct FChaosVDBinaryDataContainer
 {
@@ -75,8 +75,6 @@ public:
 
 	void RegisterDataProcessor(TSharedPtr<IChaosVDDataProcessor> InDataProcessor);
 
-	TSharedPtr<Chaos::VisualDebugger::FChaosVDSerializableNameTable> GetNameTable() const;
-
 private:
 
 	void RegisterDefaultDataProcessorsIfNeeded();
@@ -98,6 +96,8 @@ private:
 
 	TWeakPtr<FChaosVDGameFrameData> CurrentGameFrame = nullptr;
 
+	Chaos::VisualDebugger::FChaosVDArchiveHeader DefaultHeaderData;
+
 	int32 CurrentGameFrameQueueSize = 0;
 
 	bool bDefaultDataProcessorsRegistered = false;
diff --git a/Engine/Plugins/Experimental/ChaosVD/Source/ChaosVD/Public/Trace/DataProcessors/ChaosVDArchiveHeaderProcessor.h b/Engine/Plugins/Experimental/ChaosVD/Source/ChaosVD/Public/Trace/DataProcessors/ChaosVDArchiveHeaderProcessor.h
new file mode 100644
--- /dev/null
+++ b/Engine/Plugins/Experimental/ChaosVD/Source/ChaosVD/Public/Trace/DataProcessors/ChaosVDArchiveHeaderProcessor.h
@@ -0,0 +1,15 @@
+// Copyright Epic Games, Inc. All Rights Reserved.
+
+#pragma once
+#include "Trace/DataProcessors/IChaosVDDataProcessor.h"
+
+/**
+ * Data processor implementation that is able to deserialize traced Archive headers
+ */
+class FChaosVDArchiveHeaderProcessor final : public IChaosVDDataProcessor
+{
+public:
+	explicit FChaosVDArchiveHeaderProcessor();
+	
+	virtual bool ProcessRawData(const TArray<uint8>& InData) override;
+};
\ No newline at end of file
diff --git a/Engine/Plugins/Experimental/ChaosVD/Source/ChaosVD/Public/Trace/DataProcessors/IChaosVDDataProcessor.h b/Engine/Plugins/Experimental/ChaosVD/Source/ChaosVD/Public/Trace/DataProcessors/IChaosVDDataProcessor.h
--- a/Engine/Plugins/Experimental/ChaosVD/Source/ChaosVD/Public/Trace/DataProcessors/IChaosVDDataProcessor.h
+++ b/Engine/Plugins/Experimental/ChaosVD/Source/ChaosVD/Public/Trace/DataProcessors/IChaosVDDataProcessor.h
@@ -2,36 +2,57 @@
 
 #pragma once
 #include "ChaosVisualDebugger/ChaosVDMemWriterReader.h"
+#include "ChaosVisualDebugger/ChaosVisualDebuggerTrace.h"
 #include "Containers/Array.h"
 #include "Containers/StringFwd.h"
 #include "HAL/Platform.h"
 #include "Templates/SharedPointer.h"
+#include "Trace/ChaosVDTraceProvider.h"
 
-class FChaosVDTraceProvider;
+#include <type_traits>
+
+struct FFortniteSeasonBranchObjectVersion;
 
 namespace Chaos::VisualDebugger
 {
-	template<typename TDataToSerialize>
-	bool ReadDataFromBuffer(const TArray<uint8>& InDataBuffer, TDataToSerialize& Data, const TSharedRef<FChaosVDSerializableNameTable>& InNameTableInstance)
+	template<typename TArchive>
+	void ApplyHeaderDataToArchive(TArchive& InOutArchive, const FChaosVDArchiveHeader& InRecordedHeader)
 	{
-		FChaosVDMemoryReader MemReader(InDataBuffer, InNameTableInstance);
-		MemReader.SetShouldSkipUpdateCustomVersion(true);
-
-		Data.Serialize(MemReader);
-
-		return !MemReader.IsError() && !MemReader.IsCriticalError();
+		InOutArchive.SetCustomVersions(InRecordedHeader.CustomVersionContainer);
+		InOutArchive.SetEngineVer(InRecordedHeader.EngineVersion);
+		InOutArchive.SetShouldSkipUpdateCustomVersion(true);
 	}
 
-	template<typename TDataToSerialize, typename TArchive>
-	bool ReadDataFromBuffer(const TArray<uint8>& InDataBuffer, TDataToSerialize& Data, const TSharedRef<FChaosVDSerializableNameTable>& InNameTableInstance)
+	template<typename TDataToSerialize>
+	bool ReadDataFromBuffer(const TArray<uint8>& InDataBuffer, TDataToSerialize& Data, const TSharedRef<FChaosVDTraceProvider>& DataProvider)
 	{
-		FChaosVDMemoryReader MemReader(InDataBuffer, InNameTableInstance);
-		TArchive Ar(MemReader);
-		Ar.SetShouldSkipUpdateCustomVersion(true);
+		const TSharedPtr<FChaosVDRecording> RecordingInstance = DataProvider->GetRecordingForSession();
+		const TSharedPtr<FChaosVDSerializableNameTable> NameTableInstance = RecordingInstance ? RecordingInstance->GetNameTableInstance() : nullptr;
+
+		if (!ensure(NameTableInstance.IsValid()))
+		{
+			return false;
+		}
+		
+		FChaosVDMemoryReader MemReader(InDataBuffer, NameTableInstance.ToSharedRef());
+		const FChaosVDArchiveHeader& RecordedHeader =  RecordingInstance->GetHeaderData();
+		ApplyHeaderDataToArchive(MemReader, RecordedHeader);
+
+		bool bSuccess = false;
 
-		Data.Serialize(Ar);
+		// We need to use FChaosArchive as proxy to properly read serialized Implicit objects
+		// Note: I don't expect we will need a proxy archive for other types, but if we end up in that situation, we should use to switch to use traits 
+		if constexpr (std::is_same_v<TDataToSerialize, FChaosVDImplicitObjectWrapper>)
+		{
+			FChaosArchive Ar(MemReader);
+			bSuccess = Data.Serialize(Ar);
+		}
+		else
+		{
+			bSuccess = Data.Serialize(MemReader);
+		}
 
-		return !Ar.IsError() && !Ar.IsCriticalError();
+		return bSuccess;
 	}
 }
 
diff --git a/Engine/Source/Runtime/Experimental/Chaos/Private/ChaosVisualDebugger/ChaosVDMemWriterReader.cpp b/Engine/Source/Runtime/Experimental/Chaos/Private/ChaosVisualDebugger/ChaosVDMemWriterReader.cpp
new file mode 100644
--- /dev/null
+++ b/Engine/Source/Runtime/Experimental/Chaos/Private/ChaosVisualDebugger/ChaosVDMemWriterReader.cpp
@@ -0,0 +1,27 @@
+// Copyright Epic Games, Inc. All Rights Reserved.
+
+#include "ChaosVisualDebugger/ChaosVDMemWriterReader.h"
+
+namespace Chaos::VisualDebugger
+{
+	bool FChaosVDArchiveHeader::Serialize(FArchive& Ar)
+	{
+		Ar << EngineVersion;
+
+		uint8 CustomVersionFormat = static_cast<uint8>(ECustomVersionSerializationFormat::Latest);
+		Ar << CustomVersionFormat;
+		CustomVersionContainer.Serialize(Ar, static_cast<ECustomVersionSerializationFormat::Type>(CustomVersionFormat));
+
+		return !Ar.IsError();
+	}
+
+	FChaosVDArchiveHeader FChaosVDArchiveHeader::Current()
+	{
+		FChaosVDArchiveHeader CurrentVersion;
+
+		CurrentVersion.EngineVersion = FEngineVersion::Current();
+		CurrentVersion.CustomVersionContainer = FCurrentCustomVersions::GetAll();
+
+		return CurrentVersion;
+	}
+}
diff --git a/Engine/Source/Runtime/Experimental/Chaos/Private/ChaosVisualDebugger/ChaosVisualDebuggerTrace.cpp b/Engine/Source/Runtime/Experimental/Chaos/Private/ChaosVisualDebugger/ChaosVisualDebuggerTrace.cpp
--- a/Engine/Source/Runtime/Experimental/Chaos/Private/ChaosVisualDebugger/ChaosVisualDebuggerTrace.cpp
+++ b/Engine/Source/Runtime/Experimental/Chaos/Private/ChaosVisualDebugger/ChaosVisualDebuggerTrace.cpp
@@ -546,9 +546,9 @@ void FChaosVisualDebuggerTrace::TraceSolverSimulationSpace(const Chaos::FRigidTr
 		<< CVD_TRACE_ROTATOR_ON_EVENT(ChaosVDSolverSimulationSpace, Rotation, Transform.GetRotation());
 }
 
-void FChaosVisualDebuggerTrace::TraceBinaryData(TConstArrayView<uint8> InData, FStringView TypeName)
+void FChaosVisualDebuggerTrace::TraceBinaryData(TConstArrayView<uint8> InData, FStringView TypeName, EChaosVDTraceBinaryDataOptions Options)
 {
-	if (!IsTracing())
+	if (!IsTracing() && !EnumHasAnyFlags(Options, EChaosVDTraceBinaryDataOptions::ForceTrace))
 	{
 		return;
 	}
@@ -831,6 +831,22 @@ void FChaosVisualDebuggerTrace::HandleRecordingStop()
 	Reset();
 }
 
+void FChaosVisualDebuggerTrace::TraceArchiveHeader()
+{
+	using namespace Chaos::VisualDebugger;
+
+	TArray<uint8> HeaderDataBuffer;
+
+	FMemoryWriter MemWriterAr(HeaderDataBuffer);
+
+	FChaosVDArchiveHeader::Current().Serialize(MemWriterAr);
+
+	// We intentionally trace the header when the recording start was requested but we are not in a tracing state
+	// So we need to force a trace
+	// We do this to ensure the header is traced before any other binary data is generated, as we will need it to be read first on load 
+	TraceBinaryData(HeaderDataBuffer, FChaosVDArchiveHeader::WrapperTypeName, EChaosVDTraceBinaryDataOptions::ForceTrace);
+}
+
 void FChaosVisualDebuggerTrace::HandleRecordingStart()
 {
 	Reset();
@@ -859,7 +875,9 @@ void FChaosVisualDebuggerTrace::HandleRecordingStart()
 			return true;
 		});
 	}
-	
+
+	TraceArchiveHeader();
+
 	bIsTracing = true;
 }
 
diff --git a/Engine/Source/Runtime/Experimental/Chaos/Public/ChaosVisualDebugger/ChaosVDMemWriterReader.h b/Engine/Source/Runtime/Experimental/Chaos/Public/ChaosVisualDebugger/ChaosVDMemWriterReader.h
--- a/Engine/Source/Runtime/Experimental/Chaos/Public/ChaosVisualDebugger/ChaosVDMemWriterReader.h
+++ b/Engine/Source/Runtime/Experimental/Chaos/Public/ChaosVisualDebugger/ChaosVDMemWriterReader.h
@@ -3,6 +3,8 @@
 #pragma once
 #include "ChaosVDSerializedNameTable.h"
 #include "HAL/Platform.h"
+#include "Misc/EngineVersion.h"
+#include "Serialization/CustomVersion.h"
 #include "Serialization/MemoryReader.h"
 #include "Serialization/MemoryWriter.h"
 #include "Templates/SharedPointer.h"
@@ -12,6 +14,27 @@ class FArchive;
 
 namespace Chaos::VisualDebugger
 {
+	struct FChaosVDArchiveHeader
+	{
+		inline static FStringView WrapperTypeName = TEXT("FChaosVDArchiveHeader");
+
+		FChaosVDArchiveHeader()
+		{
+		}
+
+		/** Contains version and branch info from where it was saved */
+		FEngineVersion EngineVersion;
+
+		/** Custom versions */
+		FCustomVersionContainer CustomVersionContainer;
+
+		/** Serialization function, returns false if the archive ends up in an error state */
+		CHAOS_API bool Serialize(FArchive& Ar);
+
+		/** Returns a header that matches current version */
+		CHAOS_API static FChaosVDArchiveHeader Current();
+	};
+
 	FORCEINLINE void NameTableFNameSerializer(FArchive& Ar ,FName& Name, const TSharedRef<FChaosVDSerializableNameTable>& InNameTable)
 	{
 		// Serialize the FName as a CVD Name ID
diff --git a/Engine/Source/Runtime/Experimental/Chaos/Public/ChaosVisualDebugger/ChaosVisualDebuggerTrace.h b/Engine/Source/Runtime/Experimental/Chaos/Public/ChaosVisualDebugger/ChaosVisualDebuggerTrace.h
--- a/Engine/Source/Runtime/Experimental/Chaos/Public/ChaosVisualDebugger/ChaosVisualDebuggerTrace.h
+++ b/Engine/Source/Runtime/Experimental/Chaos/Public/ChaosVisualDebugger/ChaosVisualDebuggerTrace.h
@@ -149,6 +149,13 @@ namespace Chaos
 using FChaosVDImplicitObjectWrapper = FChaosVDImplicitObjectDataWrapper<Chaos::FImplicitObjectPtr, Chaos::FChaosArchive>;
 using FChaosVDSerializableNameTable = Chaos::VisualDebugger::FChaosVDSerializableNameTable;
 
+enum class EChaosVDTraceBinaryDataOptions
+{
+	None = 0,
+	ForceTrace = 1 << 0
+};
+ENUM_CLASS_FLAGS(EChaosVDTraceBinaryDataOptions)
+
 /** Class containing  all the Tracing logic to record data for the Chaos Visual Debugger tool */
 class FChaosVisualDebuggerTrace
 {
@@ -223,7 +230,7 @@ public:
 	 * @param InData Data to trace
 	 * @param TypeName Type name the data represents. It is used during Trace Analysis serialize it back (this is not automatic)
 	 */
-	static CHAOS_API void TraceBinaryData(TConstArrayView<uint8> InData, FStringView TypeName);
+	static CHAOS_API void TraceBinaryData(TConstArrayView<uint8> InData, FStringView TypeName, EChaosVDTraceBinaryDataOptions Options = EChaosVDTraceBinaryDataOptions::None);
 
 	/**
 	 * Serializes the implicit object contained in the wrapper and trace its it as binary data
@@ -299,6 +306,7 @@ private:
 	static void Reset();
 
 	static void HandleRecordingStop();
+	static void TraceArchiveHeader();
 	static void HandleRecordingStart();
 
 	/** Sets up the tracer to perform a full capture in the next solver frame */
diff --git a/Engine/Source/Runtime/Experimental/ChaosVDData/Private/ChaosVDRecording.cpp b/Engine/Source/Runtime/Experimental/ChaosVDData/Private/ChaosVDRecording.cpp
--- a/Engine/Source/Runtime/Experimental/ChaosVDData/Private/ChaosVDRecording.cpp
+++ b/Engine/Source/Runtime/Experimental/ChaosVDData/Private/ChaosVDRecording.cpp
@@ -7,6 +7,9 @@
 FChaosVDRecording::FChaosVDRecording()
 {
 	NameTable = MakeShared<Chaos::VisualDebugger::FChaosVDSerializableNameTable>();
+
+	// Start with a default header data as a fallback
+	HeaderData = Chaos::VisualDebugger::FChaosVDArchiveHeader::Current();
 }
 
 int32 FChaosVDRecording::GetAvailableGameFramesNumber() const
@@ -440,6 +443,12 @@ void FChaosVDRecording::AddImplicitObject(const uint32 ID, const Chaos::FImplici
 	}
 }
 
+bool FChaosVDRecording::IsEmpty() const
+{
+	FReadScopeLock ReadLock(RecordingDataLock);
+	return GetAvailableSolversNumber_AssumesLocked() == 0 && GetAvailableGameFrames_AssumesLocked().Num() == 0 && ImplicitObjects.Num() == 0;
+}
+
 void FChaosVDRecording::AddImplicitObject(const uint32 ID, const Chaos::FImplicitObject* InImplicitObject)
 {
 	FWriteScopeLock WriteLock(RecordingDataLock);
diff --git a/Engine/Source/Runtime/Experimental/ChaosVDData/Public/ChaosVDRecording.h b/Engine/Source/Runtime/Experimental/ChaosVDData/Public/ChaosVDRecording.h
--- a/Engine/Source/Runtime/Experimental/ChaosVDData/Public/ChaosVDRecording.h
+++ b/Engine/Source/Runtime/Experimental/ChaosVDData/Public/ChaosVDRecording.h
@@ -10,6 +10,7 @@
 #include "Containers/UnrealString.h"
 #include "Chaos/ImplicitFwd.h"
 #include "Chaos/ImplicitObject.h"
+#include "ChaosVisualDebugger/ChaosVDMemWriterReader.h"
 #include "DataWrappers/ChaosVDJointDataWrappers.h"
 #include "DataWrappers/ChaosVDQueryDataWrappers.h"
 
@@ -267,8 +268,18 @@ struct CHAOSVDDATA_API FChaosVDRecording
 	/** Sets if this recording is being populated from a live session */
 	void SetIsLive(bool bNewIsLive) { bIsLive = bNewIsLive; }
 
+	/** Returns the name table instances used to de-duplicate strings serialization */
 	TSharedPtr<Chaos::VisualDebugger::FChaosVDSerializableNameTable> GetNameTableInstance() const { return NameTable; }
 
+	/** Returns the FArchive header used to read the serialized binary data */
+	const Chaos::VisualDebugger::FChaosVDArchiveHeader& GetHeaderData() const { return HeaderData; }
+	
+	/** Sets the FArchive header used to read the serialized binary data */
+	void SetHeaderData(const Chaos::VisualDebugger::FChaosVDArchiveHeader& InNewHeader) { HeaderData = InNewHeader; }
+
+	/** Returns true if this recording does not have any usable data */
+	bool IsEmpty() const;
+
 protected:
 
 	/** Adds an Implicit Object to the recording and takes ownership of it */
@@ -303,6 +314,8 @@ protected:
 	/** Map that temporary holds generated particle data during the key frame generation process, keeping its memory allocation between generated frames*/
 	TMap<int32, TSharedPtr<FChaosVDParticleDataWrapper>> ParticlesOnCurrentGeneratedKeyframe;
 
+	Chaos::VisualDebugger::FChaosVDArchiveHeader HeaderData;
+
 	friend class FChaosVDTraceProvider;
 	friend class FChaosVDTraceImplicitObjectProcessor;
 };
diff --git a/Engine/Source/Runtime/Experimental/ChaosVisualDebugger/Public/DataWrappers/ChaosVDImplicitObjectDataWrapper.h b/Engine/Source/Runtime/Experimental/ChaosVisualDebugger/Public/DataWrappers/ChaosVDImplicitObjectDataWrapper.h
--- a/Engine/Source/Runtime/Experimental/ChaosVisualDebugger/Public/DataWrappers/ChaosVDImplicitObjectDataWrapper.h
+++ b/Engine/Source/Runtime/Experimental/ChaosVisualDebugger/Public/DataWrappers/ChaosVDImplicitObjectDataWrapper.h
@@ -17,12 +17,14 @@ struct FChaosVDImplicitObjectDataWrapper
 	uint32 Hash;
 	SerializableImplicitType ImplicitObject;
 
-	void Serialize(ArchiveType& Ar);
+	bool Serialize(ArchiveType& Ar);
 };
 
 template <class SerializableImplicitType, class ArchiveType>
-void FChaosVDImplicitObjectDataWrapper<SerializableImplicitType, ArchiveType>::Serialize(ArchiveType& Ar)
+bool FChaosVDImplicitObjectDataWrapper<SerializableImplicitType, ArchiveType>::Serialize(ArchiveType& Ar)
 {
 	Ar << Hash;
 	Ar << ImplicitObject;
+
+	return !Ar.IsError();
 }
