[ChaosVD] Fix CIS issue in shipping editors

The typedef FChaosVDImplicitObjectWrapper is not available in shipping editors. Switching to manually declaring the type for now

diff --git a/Engine/Plugins/Experimental/ChaosVD/Source/ChaosVD/Public/Trace/DataProcessors/IChaosVDDataProcessor.h b/Engine/Plugins/Experimental/ChaosVD/Source/ChaosVD/Public/Trace/DataProcessors/IChaosVDDataProcessor.h
--- a/Engine/Plugins/Experimental/ChaosVD/Source/ChaosVD/Public/Trace/DataProcessors/IChaosVDDataProcessor.h
+++ b/Engine/Plugins/Experimental/ChaosVD/Source/ChaosVD/Public/Trace/DataProcessors/IChaosVDDataProcessor.h
@@ -42,7 +42,7 @@ namespace Chaos::VisualDebugger
 
 		// We need to use FChaosArchive as proxy to properly read serialized Implicit objects
 		// Note: I don't expect we will need a proxy archive for other types, but if we end up in that situation, we should use to switch to use traits 
-		if constexpr (std::is_same_v<TDataToSerialize, FChaosVDImplicitObjectWrapper>)
+		if constexpr (std::is_same_v<TDataToSerialize, FChaosVDImplicitObjectDataWrapper<Chaos::FImplicitObjectPtr, Chaos::FChaosArchive>>)
 		{
 			FChaosArchive Ar(MemReader);
 			bSuccess = Data.Serialize(Ar);
