[ChaosVD] Fix CIS issue on shipping editors

The previous fix attempt missed two includes

diff --git a/Engine/Plugins/Experimental/ChaosVD/Source/ChaosVD/Public/Trace/DataProcessors/IChaosVDDataProcessor.h b/Engine/Plugins/Experimental/ChaosVD/Source/ChaosVD/Public/Trace/DataProcessors/IChaosVDDataProcessor.h
--- a/Engine/Plugins/Experimental/ChaosVD/Source/ChaosVD/Public/Trace/DataProcessors/IChaosVDDataProcessor.h
+++ b/Engine/Plugins/Experimental/ChaosVD/Source/ChaosVD/Public/Trace/DataProcessors/IChaosVDDataProcessor.h
@@ -1,10 +1,13 @@
 ﻿// Copyright Epic Games, Inc. All Rights Reserved.
 
 #pragma once
+
+#include "Chaos/ImplicitObject.h"
 #include "ChaosVisualDebugger/ChaosVDMemWriterReader.h"
 #include "ChaosVisualDebugger/ChaosVisualDebuggerTrace.h"
 #include "Containers/Array.h"
 #include "Containers/StringFwd.h"
+#include "DataWrappers/ChaosVDImplicitObjectDataWrapper.h"
 #include "HAL/Platform.h"
 #include "Templates/SharedPointer.h"
 #include "Trace/ChaosVDTraceProvider.h"
