[ChaosVD] Implemented a limit to how many debug draw line we can process per frame. This is to prevent a hang when enabling to many debug draw line options in a complex and dense scene.

- I made a change to reserve the lines we expect to draw in each call to reduce re-allocations of the lines array.
- Implemented a Debug text method to draw FText
- Implemented a DebugDrawWarningOnScreen method that allows you to draw text warnings on screen (it is a very primitive implementation though, just the bare minimum need to show the warning about exceeding the debug draw lines limit, we will improve this in future changes).
- Updated the component visualizers to respect the Solver visibility state, so debug draw elements are not drawn if a solver is specifically hidden.
- Made the code that searches for components to visualize for solvers more generic. Now we can add more visualizers without needing to modify CVD's viewport client class.
- Fixed an issue where if you have "DrawOnlySelectedQuery" the visualization for scene queries didn't update between frames if you had a query selected. This is fixed by just clearing out the selection when a frame changes.
- Fixed an issue where the fallback to draw all queries if you have "DrawOnlySelectedQuery" enabled, but no queries selected, was not working.
- Updated the Debug Draw API to take an FText instead of string, so we can draw localized strings. A plain string can be passed by creating a culture invariant text.

diff --git a/Engine/Plugins/Experimental/ChaosVD/Source/ChaosVD/Private/Actors/ChaosVDSolverInfoActor.cpp b/Engine/Plugins/Experimental/ChaosVD/Source/ChaosVD/Private/Actors/ChaosVDSolverInfoActor.cpp
--- a/Engine/Plugins/Experimental/ChaosVD/Source/ChaosVD/Private/Actors/ChaosVDSolverInfoActor.cpp
+++ b/Engine/Plugins/Experimental/ChaosVD/Source/ChaosVD/Private/Actors/ChaosVDSolverInfoActor.cpp
@@ -143,6 +143,11 @@ void AChaosVDSolverInfoActor::RemoveSolverFolders(UWorld* World)
 	}
 }
 
+bool AChaosVDSolverInfoActor::IsVisible() const
+{
+	return !IsTemporarilyHiddenInEditor();
+}
+
 void AChaosVDSolverInfoActor::ApplySolverVisibilityToParticle(AChaosVDParticleActor* ParticleActor, bool bIsHidden)
 {
 	if (!ParticleActor)
diff --git a/Engine/Plugins/Experimental/ChaosVD/Source/ChaosVD/Private/ChaosVDPlaybackViewportClient.cpp b/Engine/Plugins/Experimental/ChaosVD/Source/ChaosVD/Private/ChaosVDPlaybackViewportClient.cpp
--- a/Engine/Plugins/Experimental/ChaosVD/Source/ChaosVD/Private/ChaosVDPlaybackViewportClient.cpp
+++ b/Engine/Plugins/Experimental/ChaosVD/Source/ChaosVD/Private/ChaosVDPlaybackViewportClient.cpp
@@ -364,44 +364,42 @@ void FChaosVDPlaybackViewportClient::Draw(const FSceneView* View, FPrimitiveDraw
 
 	if (const TSharedPtr<FChaosVDScene> ScenePtr = CVDScene.Pin())
 	{
+		//TODO: Currently we can safely assume that any component in these actors is meant to have a visualizer, but we might need a proper interface for these components in the future
+		TInlineComponentArray<const UActorComponent*> ComponentsToVisualize;
 		for (const TPair<int32, AChaosVDSolverInfoActor*>& SolverInfoWithID : ScenePtr->GetSolverInfoActorsMap())
 		{
-			if (const UChaosVDParticleDataComponent* ParticleDataComponent = SolverInfoWithID.Value ? SolverInfoWithID.Value->GetParticleDataComponent() : nullptr)
+			if (SolverInfoWithID.Value)
 			{
-				if (const TSharedPtr<FComponentVisualizer> Visualizer = MainTabToolkitHost->FindComponentVisualizer(ParticleDataComponent->StaticClass()))
+				constexpr bool bIncludeFromChildActors = false;
+				SolverInfoWithID.Value->ForEachComponent(bIncludeFromChildActors, [&ComponentsToVisualize](UActorComponent* Component)
 				{
-					Visualizer->DrawVisualization(ParticleDataComponent, View, PDI);
-				}
-			}
-	
-			if (const UChaosVDSolverCollisionDataComponent* CollisionDataComponent = SolverInfoWithID.Value ? SolverInfoWithID.Value->GetCollisionDataComponent() : nullptr)
-			{
-				if (const TSharedPtr<FComponentVisualizer> Visualizer = MainTabToolkitHost->FindComponentVisualizer(CollisionDataComponent->StaticClass()))
-				{
-					Visualizer->DrawVisualization(CollisionDataComponent, View, PDI);
-				}
-			}
-
-			if (const UChaosVDSolverJointConstraintDataComponent* JointConstraintDataComponent = SolverInfoWithID.Value ? SolverInfoWithID.Value->GetJointsDataComponent() : nullptr)
-			{
-				if (const TSharedPtr<FComponentVisualizer> Visualizer = MainTabToolkitHost->FindComponentVisualizer(JointConstraintDataComponent->StaticClass()))
-				{
-					Visualizer->DrawVisualization(JointConstraintDataComponent, View, PDI);
-				}
+					ComponentsToVisualize.Emplace(Component);
+				});
 			}
 		}
 
 		if (const UChaosVDSceneQueryDataComponent* SceneQueryDataComponent = ScenePtr->GetSceneQueryDataContainerComponent())
 		{
-			if (const TSharedPtr<FComponentVisualizer> Visualizer = MainTabToolkitHost->FindComponentVisualizer(SceneQueryDataComponent->StaticClass()))
+			ComponentsToVisualize.Emplace(SceneQueryDataComponent);
+		}
+
+		for (const UActorComponent* Component : ComponentsToVisualize)
+		{
+			if (!FChaosVDDebugDrawUtils::CanDebugDraw())
+			{
+				break;
+			}
+
+			if (const TSharedPtr<FComponentVisualizer> Visualizer = MainTabToolkitHost->FindComponentVisualizer(Component->GetClass()))
 			{
-				Visualizer->DrawVisualization(SceneQueryDataComponent, View, PDI);
+				Visualizer->DrawVisualization(Component, View, PDI);
 			}
 		}
-		
 	}
 
 	FEditorViewportClient::Draw(View, PDI);
+	
+	FChaosVDDebugDrawUtils::DebugDrawFrameEnd();
 }
 
 void FChaosVDPlaybackViewportClient::DrawCanvas(FViewport& InViewport, FSceneView& View, FCanvas& Canvas)
diff --git a/Engine/Plugins/Experimental/ChaosVD/Source/ChaosVD/Private/Components/ChaosVDSceneQueryDataComponent.cpp b/Engine/Plugins/Experimental/ChaosVD/Source/ChaosVD/Private/Components/ChaosVDSceneQueryDataComponent.cpp
--- a/Engine/Plugins/Experimental/ChaosVD/Source/ChaosVD/Private/Components/ChaosVDSceneQueryDataComponent.cpp
+++ b/Engine/Plugins/Experimental/ChaosVD/Source/ChaosVD/Private/Components/ChaosVDSceneQueryDataComponent.cpp
@@ -51,6 +51,9 @@ void UChaosVDSceneQueryDataComponent::UpdateQueriesFromFrameData(const FChaosVDG
 	RecordedQueriesByID.Empty(RecordedQueriesNum);
 	RecordedQueries.Empty(RecordedQueriesNum);
 
+	// Until we have a way to track and auto-select new queries instances between frames, just clear the selection
+	SelectQuery(FChaosVDSceneQuerySelectionHandle());
+
 	for (const TPair<int32, TSharedPtr<FChaosVDQueryDataWrapper>>& QueryIDPair : InGameFrameData.RecordedSceneQueries)
 	{
 		if (TSharedPtr<FChaosVDQueryDataWrapper> QueryData = QueryIDPair.Value)
diff --git a/Engine/Plugins/Experimental/ChaosVD/Source/ChaosVD/Private/Visualizers/ChaosVDDebugDrawUtils.cpp b/Engine/Plugins/Experimental/ChaosVD/Source/ChaosVD/Private/Visualizers/ChaosVDDebugDrawUtils.cpp
--- a/Engine/Plugins/Experimental/ChaosVD/Source/ChaosVD/Private/Visualizers/ChaosVDDebugDrawUtils.cpp
+++ b/Engine/Plugins/Experimental/ChaosVD/Source/ChaosVD/Private/Visualizers/ChaosVDDebugDrawUtils.cpp
@@ -9,15 +9,34 @@
 #include "Engine/Engine.h"
 #include "SceneView.h"
 
-TQueue<FChaosVDDebugDrawUtils::FChaosVDQueuedTextToDraw> FChaosVDDebugDrawUtils::TexToDrawQueue = TQueue<FChaosVDQueuedTextToDraw>();
+TQueue<FChaosVDDebugDrawUtils::FChaosVDQueuedTextToDraw> FChaosVDDebugDrawUtils::TextToDrawQueue = TQueue<FChaosVDQueuedTextToDraw>();
+int32 FChaosVDDebugDrawUtils::CurrentLinesDrawn = 0;
+int32 FChaosVDDebugDrawUtils::MaxLinesToDrawPerFrame = 300000;
+bool FChaosVDDebugDrawUtils::bIsShowingDebugDrawLimitWarning = false;
+int32 FChaosVDDebugDrawUtils::CurrentWarningsBeingDrawn = 0;
 
-void FChaosVDDebugDrawUtils::DrawArrowVector(FPrimitiveDrawInterface* PDI, const FVector& StartLocation, const FVector& EndLocation, FStringView DebugText, const FColor& Color, ESceneDepthPriorityGroup DepthPriority, float Thickness)
+namespace Chaos::VisualDebugger::Cvars
+{
+	static FAutoConsoleVariableRef CVarChaosVDMaxDebugDrawLinesPerFrame(
+		TEXT("p.Chaos.VD.Tool.MaxDebugDrawLinesPerFrame"),
+		FChaosVDDebugDrawUtils::MaxLinesToDrawPerFrame,
+		TEXT("Sets the max number of lines CVD is allowed to draw between all instances in a single frame."));
+}
+
+#define LOCTEXT_NAMESPACE "ChaosVisualDebugger"
+
+void FChaosVDDebugDrawUtils::DrawArrowVector(FPrimitiveDrawInterface* PDI, const FVector& StartLocation, const FVector& EndLocation, const FText& InDebugText, const FColor& Color, ESceneDepthPriorityGroup DepthPriority, float Thickness)
 {
 	if (!PDI)
 	{
 		return;
 	}
 
+	if (!CanDebugDraw())
+	{
+		return;
+	}
+
 	const FVector LineVectorToDraw = EndLocation - StartLocation;
 
 	FVector ArrowDir;
@@ -37,52 +56,90 @@ void FChaosVDDebugDrawUtils::DrawArrowVector(FPrimitiveDrawInterface* PDI, const
 	
 	DrawDirectionalArrow(PDI, ArrowTransformMatrix, Color, ArrowLength, ArrowSize, DepthPriority, Thickness);
 
-	if (!DebugText.IsEmpty())
+	IncreaseDebugDrawLineCounter();
+
+	if (!InDebugText.IsEmpty())
 	{
 		// Draw the text in the middle of the vector line
 		const FVector TextWorldPosition = StartLocation + LineVectorToDraw  * 0.5f;
-		DrawText(DebugText.GetData(), TextWorldPosition , Color);
+		DrawText(InDebugText, TextWorldPosition , Color);
 	}
 }
 
-void FChaosVDDebugDrawUtils::DrawPoint(FPrimitiveDrawInterface* PDI, const FVector& Location, FStringView DebugText, const FColor& Color, float Size, ESceneDepthPriorityGroup DepthPriority)
+void FChaosVDDebugDrawUtils::DrawPoint(FPrimitiveDrawInterface* PDI, const FVector& Location, const FText& InDebugText, const FColor& Color, float Size, ESceneDepthPriorityGroup DepthPriority)
 {
 	if (!PDI)
 	{
 		return;
 	}
 
-	if (DebugText.IsEmpty())
+	if (InDebugText.IsEmpty())
+	{
+		return;
+	}
+
+	if (!CanDebugDraw())
 	{
 		return;
 	}
 
 	PDI->DrawPoint(Location, Color, Size, DepthPriority);
 
-	if (!DebugText.IsEmpty())
+	if (!InDebugText.IsEmpty())
 	{
-		DrawText(DebugText, Location, Color);
+		DrawText(InDebugText, Location, Color);
 	}
 }
 
-void FChaosVDDebugDrawUtils::DrawText(FStringView StringToDraw, const FVector& Location, const FColor& Color)
+void FChaosVDDebugDrawUtils::DrawString(FStringView StringToDraw, const FVector& Location, const FColor& Color, EChaosVDDebugDrawTextLocationMode LocationMode)
 {
 	if (const UChaosVDEditorSettings* CVDEditorSettings = GetDefault<UChaosVDEditorSettings>())
 	{
 		if (CVDEditorSettings->bShowDebugText)
 		{
-			TexToDrawQueue.Enqueue({StringToDraw.GetData(), Location, Color });
+			TextToDrawQueue.Enqueue({ FText::AsCultureInvariant(StringToDraw.GetData()), LocationMode, Location, Color });
 		}
 	}
 }
 
-void FChaosVDDebugDrawUtils::DrawCircle(FPrimitiveDrawInterface* PDI, const FVector& Origin, float Radius, int32 Segments, const FColor& Color, float Thickness, const FVector& XAxis, const FVector& YAxis, FStringView DebugText, ESceneDepthPriorityGroup DepthPriority)
+void FChaosVDDebugDrawUtils::DrawText(const FText& InText, const FVector& Location, const FColor& Color, EChaosVDDebugDrawTextLocationMode LocationMode)
+{
+	if (const UChaosVDEditorSettings* CVDEditorSettings = GetDefault<UChaosVDEditorSettings>())
+	{
+		if (CVDEditorSettings->bShowDebugText)
+		{
+			TextToDrawQueue.Enqueue({ InText, LocationMode, Location, Color });
+		}
+	}
+}
+
+void FChaosVDDebugDrawUtils::DrawOnScreenWarning(const FText& InText, const FColor& Color)
+{
+	if (!GEngine)
+	{
+		return;
+	}
+
+	static FVector WarningStartPosition = FVector(45.0, 40.0 ,0.0);
+	constexpr int32 LineSpaceY = 20;
+	CurrentWarningsBeingDrawn++;
+
+	const FVector WarningPosition = WarningStartPosition + FVector(0.0, LineSpaceY * CurrentWarningsBeingDrawn,0.0);
+	TextToDrawQueue.Enqueue({ InText, EChaosVDDebugDrawTextLocationMode::Screen, WarningPosition, Color });
+}
+
+void FChaosVDDebugDrawUtils::DrawCircle(FPrimitiveDrawInterface* PDI, const FVector& Origin, float Radius, int32 Segments, const FColor& Color, float Thickness, const FVector& XAxis, const FVector& YAxis, const FText& InDebugText, ESceneDepthPriorityGroup DepthPriority)
 {
 	if (!PDI)
 	{
 		return;
 	}
 
+	if (!CanDebugDraw())
+	{
+		return;
+	}
+
 	constexpr float DepthBias = 0.0f;
 	const bool bScreenSpace = Thickness > 0;
 
@@ -91,6 +148,8 @@ void FChaosVDDebugDrawUtils::DrawCircle(FPrimitiveDrawInterface* PDI, const FVec
 	const float	AngleDelta = 2.0f * UE_PI / Segments;
 	FVector	LastVertex = Origin + XAxis * Radius;
 
+	PDI->AddReserveLines(DepthPriority, Segments, false, Thickness > SMALL_NUMBER);
+
 	for (int32 SideIndex = 0; SideIndex < Segments; SideIndex++)
 	{
 		const FVector Vertex = Origin + (XAxis * FMath::Cos(AngleDelta * (SideIndex + 1)) + YAxis * FMath::Sin(AngleDelta * (SideIndex + 1))) * Radius;
@@ -100,20 +159,27 @@ void FChaosVDDebugDrawUtils::DrawCircle(FPrimitiveDrawInterface* PDI, const FVec
 		LastVertex = Vertex;
 	}
 	
-	if (!DebugText.IsEmpty())
+	if (!InDebugText.IsEmpty())
 	{
-		DrawText(DebugText, Origin, Color);
+		DrawText(InDebugText, Origin, Color);
 	}
 }
 
-void FChaosVDDebugDrawUtils::DrawBox(FPrimitiveDrawInterface* PDI, const FVector& InExtents, const FColor& InColor, const FTransform& InTransform, FStringView DebugText, ESceneDepthPriorityGroup DepthPriority, float Thickness)
+void FChaosVDDebugDrawUtils::DrawBox(FPrimitiveDrawInterface* PDI, const FVector& InExtents, const FColor& InColor, const FTransform& InTransform, const FText& InDebugText, ESceneDepthPriorityGroup DepthPriority, float Thickness)
 {
 	if (!PDI)
 	{
 		return;
 	}
 
+	if (!CanDebugDraw())
+	{
+		return;
+	}
+
 	constexpr int32 MaxBoxLines = 12;
+	
+	PDI->AddReserveLines(DepthPriority, MaxBoxLines, false, Thickness > SMALL_NUMBER);
 
 	// Array for direction offsets for the start/end point of each line
 	static TPair<FVector, FVector> VertexOffsetDirectionFromOrigin[MaxBoxLines] =
@@ -143,39 +209,51 @@ void FChaosVDDebugDrawUtils::DrawBox(FPrimitiveDrawInterface* PDI, const FVector
 		PDI->DrawLine(LineStart, LineEnd, InColor, DepthPriority, Thickness,DepthBias, bScreenSpace);
 	}
 	
-	if (!DebugText.IsEmpty())
+	if (!InDebugText.IsEmpty())
 	{
-		DrawText(DebugText, InTransform.GetLocation(), InColor);
+		DrawText(InDebugText, InTransform.GetLocation(), InColor);
 	}
 }
 
-void FChaosVDDebugDrawUtils::DrawLine(FPrimitiveDrawInterface* PDI, const FVector& InStartPosition, const FVector& InEndPosition, const FColor& InColor, FStringView DebugText, ESceneDepthPriorityGroup DepthPriority, float Thickness)
+void FChaosVDDebugDrawUtils::DrawLine(FPrimitiveDrawInterface* PDI, const FVector& InStartPosition, const FVector& InEndPosition, const FColor& InColor, const FText& InDebugText, ESceneDepthPriorityGroup DepthPriority, float Thickness)
 {
 	if (!PDI)
 	{
 		return;
 	}
 
+	if (!CanDebugDraw())
+	{
+		return;
+	}
+
 	constexpr float DepthBias = 0.0f;
 	const bool bScreenSpace = Thickness > 0;
 
 	PDI->DrawLine(InStartPosition, InEndPosition, InColor, DepthPriority, Thickness, DepthBias, bScreenSpace);
 
-	if (!DebugText.IsEmpty())
+	IncreaseDebugDrawLineCounter();
+
+	if (!InDebugText.IsEmpty())
 	{
 		// Draw the text in the middle of the line
 		const FVector TextWorldPosition = InStartPosition + ((InEndPosition - InStartPosition)  * 0.5f);
-		DrawText(DebugText, TextWorldPosition, InColor);
+		DrawText(InDebugText, TextWorldPosition, InColor);
 	}
 }
 
-void FChaosVDDebugDrawUtils::DrawImplicitObject(FPrimitiveDrawInterface* PDI, const TSharedPtr<FChaosVDGeometryBuilder>& GeometryGenerator, const Chaos::FConstImplicitObjectPtr& ImplicitObject, const FTransform& InWorldTransform, const FColor& InColor, FStringView DebugText, ESceneDepthPriorityGroup DepthPriority, float Thickness)
+void FChaosVDDebugDrawUtils::DrawImplicitObject(FPrimitiveDrawInterface* PDI, const TSharedPtr<FChaosVDGeometryBuilder>& GeometryGenerator, const Chaos::FConstImplicitObjectPtr& ImplicitObject, const FTransform& InWorldTransform, const FColor& InColor, const FText& InDebugText, ESceneDepthPriorityGroup DepthPriority, float Thickness)
 {
 	if (!PDI)
 	{
 		return;
 	}
 
+	if (!CanDebugDraw())
+	{
+		return;
+	}
+
 	if (!ImplicitObject.IsValid())
 	{
 	  return;
@@ -188,21 +266,23 @@ void FChaosVDDebugDrawUtils::DrawImplicitObject(FPrimitiveDrawInterface* PDI, co
 		GeometryGenerator->AdjustedTransformForImplicit(ImplicitObject, AdjustedTransform);
 		MeshGenerator->Generate();
 
+		PDI->AddReserveLines(DepthPriority, MeshGenerator->Triangles.Num() * 3, false, Thickness > SMALL_NUMBER);
+
 		for (const UE::Geometry::FIndex3i& Triangle : MeshGenerator->Triangles)
 		{
 			FVector VertexA = AdjustedTransform.TransformPosition(MeshGenerator->Vertices[Triangle.A]);
 			FVector VertexB = AdjustedTransform.TransformPosition(MeshGenerator->Vertices[Triangle.B]);
 			FVector VertexC = AdjustedTransform.TransformPosition(MeshGenerator->Vertices[Triangle.C]);
 	
-			DrawLine(PDI, VertexA, VertexB, InColor, nullptr, DepthPriority, Thickness);
-			DrawLine(PDI, VertexB, VertexC, InColor, nullptr, DepthPriority, Thickness);
-			DrawLine(PDI, VertexC, VertexA, InColor, nullptr, DepthPriority, Thickness);
+			DrawLine(PDI, VertexA, VertexB, InColor, FText::GetEmpty(), DepthPriority, Thickness);
+			DrawLine(PDI, VertexB, VertexC, InColor, FText::GetEmpty(), DepthPriority, Thickness);
+			DrawLine(PDI, VertexC, VertexA, InColor, FText::GetEmpty(), DepthPriority, Thickness);
 		}	
 	}
 
-	if (!DebugText.IsEmpty())
+	if (!InDebugText.IsEmpty())
 	{
-		DrawText(DebugText, InWorldTransform.GetLocation(), InColor);
+		DrawText(InDebugText, InWorldTransform.GetLocation(), InColor);
 	}
 }
 
@@ -213,20 +293,30 @@ void FChaosVDDebugDrawUtils::DrawCanvas(FViewport& InViewport, FSceneView& View,
 		return;
 	}
 
-	while (!TexToDrawQueue.IsEmpty())
+	while (!TextToDrawQueue.IsEmpty())
 	{
 		FChaosVDQueuedTextToDraw TextToDraw;
-		if (TexToDrawQueue.Dequeue(TextToDraw))
+		if (TextToDrawQueue.Dequeue(TextToDraw))
 		{
-			if (!View.ViewFrustum.IntersectPoint(TextToDraw.WorldPosition))
+			FVector2D LocationToDraw;
+			bool bHasValidLocation = false;
+			switch (TextToDraw.LocationMode)
 			{
-				continue;
+				case EChaosVDDebugDrawTextLocationMode::World:
+					bHasValidLocation = View.ViewFrustum.IntersectPoint(TextToDraw.Location) && View.WorldToPixel(TextToDraw.Location, LocationToDraw);
+					break;
+				case EChaosVDDebugDrawTextLocationMode::Screen:
+					{
+						LocationToDraw = FVector2D(TextToDraw.Location.X,TextToDraw.Location.Y);
+						bHasValidLocation = true;
+					}
+				default:
+					break;
 			}
-
-			FVector2D PixelLocation;
-			if (View.WorldToPixel(TextToDraw.WorldPosition, PixelLocation))
+			
+			if (bHasValidLocation)
 			{
-				FCanvasTextItem TextItem(PixelLocation, FText::AsCultureInvariant(TextToDraw.Text), GEngine->GetSmallFont(), TextToDraw.Color);
+				FCanvasTextItem TextItem(LocationToDraw, TextToDraw.Text, GEngine->GetSmallFont(), TextToDraw.Color);
 				TextItem.Scale = FVector2D::UnitVector;
 				TextItem.EnableShadow(FLinearColor::Black);
 				TextItem.Draw(&Canvas);
@@ -235,7 +325,33 @@ void FChaosVDDebugDrawUtils::DrawCanvas(FViewport& InViewport, FSceneView& View,
 	}
 }
 
+bool FChaosVDDebugDrawUtils::CanDebugDraw()
+{
+	return CurrentLinesDrawn <= MaxLinesToDrawPerFrame;
+}
+
+void FChaosVDDebugDrawUtils::DebugDrawFrameEnd()
+{
+	bIsShowingDebugDrawLimitWarning = false;
+	CurrentLinesDrawn = 0;
+	CurrentWarningsBeingDrawn = 0;
+}
+
+void FChaosVDDebugDrawUtils::IncreaseDebugDrawLineCounter()
+{
+	CurrentLinesDrawn++;
+
+	if (!CanDebugDraw() && !bIsShowingDebugDrawLimitWarning)
+	{
+		bIsShowingDebugDrawLimitWarning = true;
+
+		DrawOnScreenWarning(LOCTEXT("DebugDrawLimitWarning", "Max Debug Draw lines limit reached!. Try selecting fewer debug draw categories of focus the camera in a narrower area."), FColor::Yellow);
+	}
+}
+
 FString Chaos::VisualDebugger::Utils::GenerateDebugTextForVector(const FVector& InVector, const FString& VectorName, const FString& InVectorUnits)
 {
 	return FString::Format(TEXT("{5} : {0} {4} \n	|-- X : {1} {4} \n	|-- Y : {2} {4} \n	|-- Z : {3} {4}"), {InVector.Size(), InVector.X, InVector.Y, InVector.Z, InVectorUnits, VectorName });
 }
+
+#undef LOCTEXT_NAMESPACE
diff --git a/Engine/Plugins/Experimental/ChaosVD/Source/ChaosVD/Private/Visualizers/ChaosVDJointConstraintsDataComponentVisualizer.cpp b/Engine/Plugins/Experimental/ChaosVD/Source/ChaosVD/Private/Visualizers/ChaosVDJointConstraintsDataComponentVisualizer.cpp
--- a/Engine/Plugins/Experimental/ChaosVD/Source/ChaosVD/Private/Visualizers/ChaosVDJointConstraintsDataComponentVisualizer.cpp
+++ b/Engine/Plugins/Experimental/ChaosVD/Source/ChaosVD/Private/Visualizers/ChaosVDJointConstraintsDataComponentVisualizer.cpp
@@ -111,6 +111,11 @@ void FChaosVDJointConstraintsDataComponentVisualizer::DrawVisualization(const UA
 		return;
 	}
 
+	if (!SolverInfoActor->IsVisible())
+	{
+		return;
+	}
+
 	const TSharedPtr<FChaosVDScene> CVDScene = SolverInfoActor->GetScene().Pin();
 	if (!CVDScene)
 	{
@@ -140,14 +145,6 @@ void FChaosVDJointConstraintsDataComponentVisualizer::DrawVisualization(const UA
 		return;
 	}
 
-	bool bCanDrawRecordedSolverInstance = (SolverInfoActor->GetIsServer() && VisualizationContext.IsVisualizationFlagEnabled(EChaosVDJointsDataVisualizationFlags::EnableServerDraw)) ||
-											(!SolverInfoActor->GetIsServer() && VisualizationContext.IsVisualizationFlagEnabled(EChaosVDJointsDataVisualizationFlags::EnableClientDraw));
-	
-	if (!bCanDrawRecordedSolverInstance)
-	{
-		return;
-	}
-
 	// If nothing is selected, fallback to draw all joints
 	const bool bDrawOnlySelected = VisualizationContext.IsVisualizationFlagEnabled(EChaosVDJointsDataVisualizationFlags::OnlyDrawSelected) && JointConstraintDataComponent->GetCurrentSelectionHandle().IsSelected();
 	if (bDrawOnlySelected)
@@ -200,7 +197,7 @@ void FChaosVDJointConstraintsDataComponentVisualizer::DebugDrawAllAxis(const FCh
 {
 	for(int32 AxisIndex = 0; AxisIndex < 3 ; AxisIndex++)
 	{
-		FChaosVDDebugDrawUtils::DrawArrowVector(PDI, InPosition, InPosition + VisualizationContext.DebugDrawSettings->GeneralScale * VisualizationContext.DebugDrawSettings->ConstraintAxisLength * VisualizationContext.SpaceTransform.TransformVector(InRotationMatrix.GetAxis(AxisIndex)),nullptr, Chaos::VisualDebugger::Utils::GenerateSelectionAwareDebugColor(AxisColors[AxisIndex], InJointConstraintData.bIsSelectedInEditor), VisualizationContext.DebugDrawSettings->DepthPriority, LineThickness * 0.2f);
+		FChaosVDDebugDrawUtils::DrawArrowVector(PDI, InPosition, InPosition + VisualizationContext.DebugDrawSettings->GeneralScale * VisualizationContext.DebugDrawSettings->ConstraintAxisLength * VisualizationContext.SpaceTransform.TransformVector(InRotationMatrix.GetAxis(AxisIndex)),FText::GetEmpty(), Chaos::VisualDebugger::Utils::GenerateSelectionAwareDebugColor(AxisColors[AxisIndex], InJointConstraintData.bIsSelectedInEditor), VisualizationContext.DebugDrawSettings->DepthPriority, LineThickness * 0.2f);
 	}
 }
 
@@ -326,10 +323,10 @@ void FChaosVDJointConstraintsDataComponentVisualizer::DrawJointConstraint(const
 			Sb = FMath::Lerp(Pb, Xb, FMath::Clamp<double>(CoMSize / Lena, 0., 1.));
 		}
 
-		FChaosVDDebugDrawUtils::DrawLine(PDI, Pa, Sa, GenerateSelectionAwareDebugColor(White, InJointConstraintData.bIsSelectedInEditor), nullptr, VisualizationContext.DebugDrawSettings->DepthPriority, ConnectorThickness);
-		FChaosVDDebugDrawUtils::DrawLine(PDI, Pb, Sb, GenerateSelectionAwareDebugColor(White, InJointConstraintData.bIsSelectedInEditor), nullptr, VisualizationContext.DebugDrawSettings->DepthPriority, ConnectorThickness);
-		FChaosVDDebugDrawUtils::DrawLine(PDI, Sa, Xa, GenerateSelectionAwareDebugColor(Red, InJointConstraintData.bIsSelectedInEditor), nullptr, VisualizationContext.DebugDrawSettings->DepthPriority, ConnectorThickness);
-		FChaosVDDebugDrawUtils::DrawLine(PDI, Sb, Xb, GenerateSelectionAwareDebugColor(Cyan, InJointConstraintData.bIsSelectedInEditor), nullptr, VisualizationContext.DebugDrawSettings->DepthPriority, ConnectorThickness);
+		FChaosVDDebugDrawUtils::DrawLine(PDI, Pa, Sa, GenerateSelectionAwareDebugColor(White, InJointConstraintData.bIsSelectedInEditor), FText::GetEmpty(), VisualizationContext.DebugDrawSettings->DepthPriority, ConnectorThickness);
+		FChaosVDDebugDrawUtils::DrawLine(PDI, Pb, Sb, GenerateSelectionAwareDebugColor(White, InJointConstraintData.bIsSelectedInEditor), FText::GetEmpty(), VisualizationContext.DebugDrawSettings->DepthPriority, ConnectorThickness);
+		FChaosVDDebugDrawUtils::DrawLine(PDI, Sa, Xa, GenerateSelectionAwareDebugColor(Red, InJointConstraintData.bIsSelectedInEditor), FText::GetEmpty(), VisualizationContext.DebugDrawSettings->DepthPriority, ConnectorThickness);
+		FChaosVDDebugDrawUtils::DrawLine(PDI, Sb, Xb, GenerateSelectionAwareDebugColor(Cyan, InJointConstraintData.bIsSelectedInEditor), FText::GetEmpty(), VisualizationContext.DebugDrawSettings->DepthPriority, ConnectorThickness);
 	}
 
 	if (VisualizationContext.IsVisualizationFlagEnabled(EChaosVDJointsDataVisualizationFlags::CenterOfMassConnector))
@@ -350,16 +347,16 @@ void FChaosVDJointConstraintsDataComponentVisualizer::DrawJointConstraint(const
 			Sb = FMath::Lerp(Cb, Xb, FMath::Clamp<double>(CoMSize / Lena, 0., 1.));
 		}
 
-		FChaosVDDebugDrawUtils::DrawLine(PDI, Ca, Sa, GenerateSelectionAwareDebugColor(Black, InJointConstraintData.bIsSelectedInEditor), nullptr, VisualizationContext.DebugDrawSettings->DepthPriority, ConnectorThickness);
-		FChaosVDDebugDrawUtils::DrawLine(PDI, Cb, Sb, GenerateSelectionAwareDebugColor(Black, InJointConstraintData.bIsSelectedInEditor), nullptr, VisualizationContext.DebugDrawSettings->DepthPriority, ConnectorThickness);
-		FChaosVDDebugDrawUtils::DrawLine(PDI, Sa, Xa, GenerateSelectionAwareDebugColor(Red, InJointConstraintData.bIsSelectedInEditor), nullptr, VisualizationContext.DebugDrawSettings->DepthPriority, ConnectorThickness);
-		FChaosVDDebugDrawUtils::DrawLine(PDI, Sb, Xb, GenerateSelectionAwareDebugColor(Cyan, InJointConstraintData.bIsSelectedInEditor), nullptr, VisualizationContext.DebugDrawSettings->DepthPriority, ConnectorThickness);
+		FChaosVDDebugDrawUtils::DrawLine(PDI, Ca, Sa, GenerateSelectionAwareDebugColor(Black, InJointConstraintData.bIsSelectedInEditor), FText::GetEmpty(), VisualizationContext.DebugDrawSettings->DepthPriority, ConnectorThickness);
+		FChaosVDDebugDrawUtils::DrawLine(PDI, Cb, Sb, GenerateSelectionAwareDebugColor(Black, InJointConstraintData.bIsSelectedInEditor), FText::GetEmpty(), VisualizationContext.DebugDrawSettings->DepthPriority, ConnectorThickness);
+		FChaosVDDebugDrawUtils::DrawLine(PDI, Sa, Xa, GenerateSelectionAwareDebugColor(Red, InJointConstraintData.bIsSelectedInEditor), FText::GetEmpty(), VisualizationContext.DebugDrawSettings->DepthPriority, ConnectorThickness);
+		FChaosVDDebugDrawUtils::DrawLine(PDI, Sb, Xb, GenerateSelectionAwareDebugColor(Cyan, InJointConstraintData.bIsSelectedInEditor), FText::GetEmpty(), VisualizationContext.DebugDrawSettings->DepthPriority, ConnectorThickness);
 	}
 
 	if (VisualizationContext.IsVisualizationFlagEnabled(EChaosVDJointsDataVisualizationFlags::Stretch))
 	{
 		const float StretchThickness = 3.0f * LineThickness;
-		FChaosVDDebugDrawUtils::DrawLine(PDI, Xa, Xb, GenerateSelectionAwareDebugColor(Magenta, InJointConstraintData.bIsSelectedInEditor), nullptr, VisualizationContext.DebugDrawSettings->DepthPriority, StretchThickness);
+		FChaosVDDebugDrawUtils::DrawLine(PDI, Xa, Xb, GenerateSelectionAwareDebugColor(Magenta, InJointConstraintData.bIsSelectedInEditor), FText::GetEmpty(), VisualizationContext.DebugDrawSettings->DepthPriority, StretchThickness);
 	}
 
 	if (VisualizationContext.IsVisualizationFlagEnabled(EChaosVDJointsDataVisualizationFlags::Axes))
@@ -374,7 +371,7 @@ void FChaosVDJointConstraintsDataComponentVisualizer::DrawJointConstraint(const
 	if (VisualizationContext.IsVisualizationFlagEnabled(EChaosVDJointsDataVisualizationFlags::PushOut))
 	{
 		FLinearColor PushOutImpulseColor = FColor(0, 250, 250);
-		FChaosVDDebugDrawUtils::DrawLine(PDI, Xa, Xa + VisualizationContext.DebugDrawSettings->LinearImpulseScale * VisualizationContext.SpaceTransform.TransformVectorNoScale(InJointConstraintData.JointState.LinearImpulse), GenerateSelectionAwareDebugColor(PushOutImpulseColor, InJointConstraintData.bIsSelectedInEditor), nullptr, VisualizationContext.DebugDrawSettings->DepthPriority, LineThickness);
+		FChaosVDDebugDrawUtils::DrawLine(PDI, Xa, Xa + VisualizationContext.DebugDrawSettings->LinearImpulseScale * VisualizationContext.SpaceTransform.TransformVectorNoScale(InJointConstraintData.JointState.LinearImpulse), GenerateSelectionAwareDebugColor(PushOutImpulseColor, InJointConstraintData.bIsSelectedInEditor), FText::GetEmpty(), VisualizationContext.DebugDrawSettings->DepthPriority, LineThickness);
 	}
 
 	//TODO: Should we draw the Angular Impulse
diff --git a/Engine/Plugins/Experimental/ChaosVD/Source/ChaosVD/Private/Visualizers/ChaosVDParticleDataComponentVisualizer.cpp b/Engine/Plugins/Experimental/ChaosVD/Source/ChaosVD/Private/Visualizers/ChaosVDParticleDataComponentVisualizer.cpp
--- a/Engine/Plugins/Experimental/ChaosVD/Source/ChaosVD/Private/Visualizers/ChaosVDParticleDataComponentVisualizer.cpp
+++ b/Engine/Plugins/Experimental/ChaosVD/Source/ChaosVD/Private/Visualizers/ChaosVDParticleDataComponentVisualizer.cpp
@@ -13,6 +13,8 @@
 
 IMPLEMENT_HIT_PROXY(HChaosVDParticleDataProxy, HComponentVisProxy)
 
+#define LOCTEXT_NAMESPACE "ChaosVisualDebugger"
+
 /** Sets a Hit proxy which will be cleared out as soon this struct goes out of scope*/
 struct FChaosVDScopedParticleHitProxy
 {
@@ -62,6 +64,11 @@ void FChaosVDParticleDataComponentVisualizer::DrawVisualization(const UActorComp
 		return;
 	}
 
+	if (!SolverDataActor->IsVisible())
+	{
+		return;
+	}
+
 	const TSharedPtr<FChaosVDScene> CVDScene = SolverDataActor->GetScene().Pin();
 	if (!CVDScene)
 	{
@@ -98,7 +105,8 @@ void FChaosVDParticleDataComponentVisualizer::DrawVisualization(const UActorComp
 			VisualizationContext.bIsSelectedData = SolverDataActor->IsParticleSelectedByID(InParticleDataViewer.ParticleIndex);
 			DrawVisualizationForParticleData(Component, PDI, View, VisualizationContext, InParticleDataViewer);
 
-			return true;
+			// If we reach the debug draw limit for this frame, there is no need to continue processing particles
+			return FChaosVDDebugDrawUtils::CanDebugDraw();
 		});
 	}
 }
@@ -143,7 +151,7 @@ void FChaosVDParticleDataComponentVisualizer::DrawParticleVector(FPrimitiveDrawI
 	}
 
 	const FString DebugText = InVisualizationContext.bShowDebugText ? Chaos::VisualDebugger::Utils::GenerateDebugTextForVector(InVector, UEnum::GetDisplayValueAsText(VectorID).ToString(), Chaos::VisualDebugger::ParticleDataUnitsStrings::GetUnitByID(VectorID)) : TEXT("");
-	FChaosVDDebugDrawUtils::DrawArrowVector(PDI, StartLocation, StartLocation +  InVisualizationContext.DebugDrawSettings->GetScaleFortDataID(VectorID) * InVector, DebugText, InVisualizationContext.DebugDrawSettings->ColorSettings.GetColorForDataID(VectorID, InVisualizationContext.bIsSelectedData),  InVisualizationContext.DebugDrawSettings->DepthPriority, LineThickness);
+	FChaosVDDebugDrawUtils::DrawArrowVector(PDI, StartLocation, StartLocation +  InVisualizationContext.DebugDrawSettings->GetScaleFortDataID(VectorID) * InVector, FText::AsCultureInvariant(DebugText), InVisualizationContext.DebugDrawSettings->ColorSettings.GetColorForDataID(VectorID, InVisualizationContext.bIsSelectedData),  InVisualizationContext.DebugDrawSettings->DepthPriority, LineThickness);
 }
 
 void FChaosVDParticleDataComponentVisualizer::DrawVisualizationForParticleData(const UActorComponent* Component, FPrimitiveDrawInterface* PDI, const FSceneView* View, const FChaosVDParticleDataVisualizationContext& InVisualizationContext, const FChaosVDParticleDataWrapper& InParticleDataViewer)
@@ -204,8 +212,7 @@ void FChaosVDParticleDataComponentVisualizer::DrawVisualizationForParticleData(c
 				Sphere.SetSphere(InVisualizationContext.DebugDrawSettings->CenterOfMassRadius);
 				const FPhysicsShapeAdapter SphereShapeAdapter(FQuat::Identity, Sphere);
 
-				const FString DebugText = UEnum::GetDisplayValueAsText(EChaosVDParticleDataVisualizationFlags::CenterOfMass).ToString();
-				FChaosVDDebugDrawUtils::DrawImplicitObject(PDI, GeometryGenerator, &SphereShapeAdapter.GetGeometry(), FTransform(OwnerCoMLocation), InVisualizationContext.DebugDrawSettings->ColorSettings.GetColorForDataID(EChaosVDParticleDataVisualizationFlags::CenterOfMass, InVisualizationContext.bIsSelectedData), DebugText, InVisualizationContext.DebugDrawSettings->DepthPriority, LineThickness);
+				FChaosVDDebugDrawUtils::DrawImplicitObject(PDI, GeometryGenerator, &SphereShapeAdapter.GetGeometry(), FTransform(OwnerCoMLocation), InVisualizationContext.DebugDrawSettings->ColorSettings.GetColorForDataID(EChaosVDParticleDataVisualizationFlags::CenterOfMass, InVisualizationContext.bIsSelectedData), UEnum::GetDisplayValueAsText(EChaosVDParticleDataVisualizationFlags::CenterOfMass), InVisualizationContext.DebugDrawSettings->DepthPriority, LineThickness);
 			}
 		}
 	}
@@ -226,10 +233,10 @@ void FChaosVDParticleDataComponentVisualizer::DrawVisualizationForParticleData(c
 							FColor DebugDrawColor = InVisualizationContext.DebugDrawSettings->ColorSettings.GetColorForDataID(EChaosVDParticleDataVisualizationFlags::ClusterConnectivityEdge, InVisualizationContext.bIsSelectedData);
 							FVector BoxExtents(2,2,2);
 							FTransform BoxTransform(OwnerRotation, OwnerLocation);
-							FChaosVDDebugDrawUtils::DrawBox(PDI, BoxExtents, DebugDrawColor, BoxTransform, nullptr, InVisualizationContext.DebugDrawSettings->DepthPriority, LineThickness);
+							FChaosVDDebugDrawUtils::DrawBox(PDI, BoxExtents, DebugDrawColor, BoxTransform, FText::GetEmpty(), InVisualizationContext.DebugDrawSettings->DepthPriority, LineThickness);
 
 							FVector SiblingParticleLocation = InVisualizationContext.SpaceTransform.TransformPosition(SiblingParticleData->ParticlePositionRotation.MX);
-							FChaosVDDebugDrawUtils::DrawLine(PDI, OwnerLocation, SiblingParticleLocation, DebugDrawColor, FString::Printf(TEXT("Strain [%f]"), ConnectivityEdge.Strain), InVisualizationContext.DebugDrawSettings->DepthPriority, LineThickness);
+							FChaosVDDebugDrawUtils::DrawLine(PDI, OwnerLocation, SiblingParticleLocation, DebugDrawColor, FText::FormatOrdered(LOCTEXT("StrainDebugDraw","Strain {0}"), ConnectivityEdge.Strain), InVisualizationContext.DebugDrawSettings->DepthPriority, LineThickness);
 						}
 					}	
 				}
@@ -237,3 +244,5 @@ void FChaosVDParticleDataComponentVisualizer::DrawVisualizationForParticleData(c
 		}
 	}
 }
+
+#undef LOCTEXT_NAMESPACE
diff --git a/Engine/Plugins/Experimental/ChaosVD/Source/ChaosVD/Private/Visualizers/ChaosVDSceneQueryDataComponentVisualizer.cpp b/Engine/Plugins/Experimental/ChaosVD/Source/ChaosVD/Private/Visualizers/ChaosVDSceneQueryDataComponentVisualizer.cpp
--- a/Engine/Plugins/Experimental/ChaosVD/Source/ChaosVD/Private/Visualizers/ChaosVDSceneQueryDataComponentVisualizer.cpp
+++ b/Engine/Plugins/Experimental/ChaosVD/Source/ChaosVD/Private/Visualizers/ChaosVDSceneQueryDataComponentVisualizer.cpp
@@ -13,6 +13,7 @@
 #include "EditorModeManager.h"
 #include "EditorViewportClient.h"
 #include "SceneView.h"
+#include "Actors/ChaosVDSolverInfoActor.h"
 #include "Visualizers/ChaosVDDebugDrawUtils.h"
 #include "Visualizers/IChaosVDParticleVisualizationDataProvider.h"
 #include "Widgets/SChaosVDMainTab.h"
@@ -21,6 +22,8 @@ class UChaosVDEditorSettings;
 
 IMPLEMENT_HIT_PROXY(HChaosVDSceneQueryProxy, HComponentVisProxy)
 
+#define LOCTEXT_NAMESPACE "ChaosVisualDebugger"
+
 void FChaosVDSceneQueryDataComponentVisualizer::DrawVisualization(const UActorComponent* Component, const FSceneView* View, FPrimitiveDrawInterface* PDI)
 {
 	const UChaosVDSceneQueryDataComponent* SceneQueryDataComponent = Cast<UChaosVDSceneQueryDataComponent>(Component);
@@ -66,9 +69,11 @@ void FChaosVDSceneQueryDataComponentVisualizer::DrawVisualization(const UActorCo
 	if (EnumHasAnyFlags(EChaosVDSceneQueryVisualizationFlags::EnableDraw, static_cast<EChaosVDSceneQueryVisualizationFlags>(VisualizationContext.VisualizationFlags)))
 	{
 		// If Draw only selected Query is enabled, but no query is selected, just draw all queries
-		if (EnumHasAnyFlags(EChaosVDSceneQueryVisualizationFlags::OnlyDrawSelectedQuery, static_cast<EChaosVDSceneQueryVisualizationFlags>(VisualizationContext.VisualizationFlags)))
+
+		const TSharedPtr<FChaosVDQueryDataWrapper> SelectedQuery = SceneQueryDataComponent->GetSelectedQueryHandle().GetQueryData().Pin();
+		const bool bHasSelectedQuery = SelectedQuery.IsValid();
+		if (bHasSelectedQuery && EnumHasAnyFlags(EChaosVDSceneQueryVisualizationFlags::OnlyDrawSelectedQuery, static_cast<EChaosVDSceneQueryVisualizationFlags>(VisualizationContext.VisualizationFlags)))
 		{
-			const TSharedPtr<FChaosVDQueryDataWrapper> SelectedQuery = SceneQueryDataComponent->GetSelectedQueryHandle().GetQueryData().Pin();
 			DrawSceneQuery(Component, View, PDI, CVDScene, CVDRecording, VisualizationContext, SelectedQuery);
 		}
 		else
@@ -110,7 +115,7 @@ void FChaosVDSceneQueryDataComponentVisualizer::DrawLineTraceQuery(const UActorC
 {
 	PDI->SetHitProxy(new HChaosVDSceneQueryProxy(Component, VisualizationContext.DataSelectionHandle));
 
-	const FString DebugText = FString::Format(TEXT("Type: {0} \n Tag {1} | Owner Tag "), {TEXT("Line Trace"), SceneQueryData.CollisionQueryParams.TraceTag.ToString(), SceneQueryData.CollisionQueryParams.OwnerTag.ToString()});
+	const FText DebugText = FText::FormatOrdered(LOCTEXT("LineTraceDebugDrawText", "Type: Line Trace \n Tag {1} \n Owner Tag {2}"), FText::AsCultureInvariant(SceneQueryData.CollisionQueryParams.TraceTag.ToString()), FText::AsCultureInvariant(SceneQueryData.CollisionQueryParams.OwnerTag.ToString()));
 
 	FVector EndLocationToDraw;
 	if (SceneQueryData.SQVisitData.IsValidIndex(SceneQueryData.CurrentVisitIndex))
@@ -141,7 +146,7 @@ void FChaosVDSceneQueryDataComponentVisualizer::DrawOverlapQuery(const UActorCom
 	const Chaos::FConstImplicitObjectPtr InputShapePtr = VisualizationContext.InputGeometry;
 	if (ensure(InputShapePtr))
 	{
-		const FString DebugText = FString::Format(TEXT("Type: {0} \n Tag {1} | Owner Tag "), {TEXT("Overlap"), SceneQueryData.CollisionQueryParams.TraceTag.ToString(), SceneQueryData.CollisionQueryParams.OwnerTag.ToString()});
+		const FText DebugText = FText::FormatOrdered(LOCTEXT("LineTraceDebugDrawText", "Type: Overlap \n Tag {1} \n Owner Tag {2}"), FText::AsCultureInvariant(SceneQueryData.CollisionQueryParams.TraceTag.ToString()), FText::AsCultureInvariant(SceneQueryData.CollisionQueryParams.OwnerTag.ToString()));
 		FChaosVDDebugDrawUtils::DrawImplicitObject(PDI, VisualizationContext.GeometryGenerator.Pin(), InputShapePtr, FTransform(SceneQueryData.GeometryOrientation, SceneQueryData.StartLocation), VisualizationContext.DebugDrawColor, DebugText, ESceneDepthPriorityGroup::SDPG_Foreground);
 	}
 
@@ -160,7 +165,7 @@ void FChaosVDSceneQueryDataComponentVisualizer::DrawSweepQuery(const UActorCompo
 	const Chaos::FConstImplicitObjectPtr InputShapePtr = VisualizationContext.InputGeometry;
 	if (ensure(InputShapePtr))
 	{
-		FChaosVDDebugDrawUtils::DrawImplicitObject(PDI, VisualizationContext.GeometryGenerator.Pin(), InputShapePtr, FTransform(SceneQueryData.GeometryOrientation, SceneQueryData.StartLocation), VisualizationContext.DebugDrawColor, nullptr, ESceneDepthPriorityGroup::SDPG_Foreground);
+		FChaosVDDebugDrawUtils::DrawImplicitObject(PDI, VisualizationContext.GeometryGenerator.Pin(), InputShapePtr, FTransform(SceneQueryData.GeometryOrientation, SceneQueryData.StartLocation), VisualizationContext.DebugDrawColor, FText::GetEmpty(), ESceneDepthPriorityGroup::SDPG_Foreground);
 
 		FVector EndLocationToDraw;
 		if (SceneQueryData.SQVisitData.IsValidIndex(SceneQueryData.CurrentVisitIndex))
@@ -174,12 +179,13 @@ void FChaosVDSceneQueryDataComponentVisualizer::DrawSweepQuery(const UActorCompo
 			EndLocationToDraw = SceneQueryData.EndLocation;
 		}
 
-		FChaosVDDebugDrawUtils::DrawImplicitObject(PDI, VisualizationContext.GeometryGenerator.Pin(), InputShapePtr, FTransform(SceneQueryData.GeometryOrientation, EndLocationToDraw), VisualizationContext.DebugDrawColor, nullptr, ESceneDepthPriorityGroup::SDPG_Foreground);
+		FChaosVDDebugDrawUtils::DrawImplicitObject(PDI, VisualizationContext.GeometryGenerator.Pin(), InputShapePtr, FTransform(SceneQueryData.GeometryOrientation, EndLocationToDraw), VisualizationContext.DebugDrawColor, FText::GetEmpty(), ESceneDepthPriorityGroup::SDPG_Foreground);
 	}
 
 	PDI->SetHitProxy(nullptr);
 
-	const FString DebugText = FString::Format(TEXT("Type: {0} \n Tag {1} | Owner Tag "), {TEXT("Sweep"), SceneQueryData.CollisionQueryParams.TraceTag.ToString(), SceneQueryData.CollisionQueryParams.OwnerTag.ToString()});
+	const FText DebugText = FText::FormatOrdered(LOCTEXT("LineTraceDebugDrawText", "Type: Sweep \n Tag {1} \n Owner Tag {2}"), FText::AsCultureInvariant(SceneQueryData.CollisionQueryParams.TraceTag.ToString()), FText::AsCultureInvariant(SceneQueryData.CollisionQueryParams.OwnerTag.ToString()));
+
 	FChaosVDDebugDrawUtils::DrawArrowVector(PDI, SceneQueryData.StartLocation, SceneQueryData.EndLocation, DebugText, VisualizationContext.DebugDrawColor, SDPG_Foreground);
 
 	if (EnumHasAnyFlags(EChaosVDSceneQueryVisualizationFlags::DrawHits, static_cast<EChaosVDSceneQueryVisualizationFlags>(VisualizationContext.VisualizationFlags)))
@@ -202,16 +208,18 @@ void FChaosVDSceneQueryDataComponentVisualizer::DrawHits(const UActorComponent*
 
 		PDI->SetHitProxy(new HChaosVDSceneQueryProxy(Component, VisualizationContext.DataSelectionHandle));
 
-		FString HitDebugText = FString::Format(TEXT("Distance {0} \n FaceIndex {1} \n "), {SQVisitData.HitData.Distance, SQVisitData.HitData.FaceIdx});
+		const FText HitPointDebugText = FText::FormatOrdered(LOCTEXT("SceneQueryHitDebugText", "Distance {0} \n Face Index {1} \n "), SQVisitData.HitData.Distance, SQVisitData.HitData.FaceIdx);
+		static FText HitFaceNormalDebugText = FText::AsCultureInvariant(TEXT("Hit Face Normal"));
+		static FText HitWorldNormalDebugText = FText::AsCultureInvariant(TEXT("Hit World Normal"));
 
 		const FMatrix Axes = FRotationMatrix::MakeFromX(SQVisitData.HitData.WorldNormal);
 		constexpr float Thickness = 5.0f;
 		constexpr float CircleRadius = 5.0f;
 		constexpr int32 CircleSegments = 12;
 		constexpr float NormalScale = 10.5f;
-		FChaosVDDebugDrawUtils::DrawCircle(PDI, SQVisitData.HitData.WorldPosition, CircleRadius, CircleSegments, InColor, Thickness, Axes.GetUnitAxis(EAxis::Y), Axes.GetUnitAxis(EAxis::Z), HitDebugText,  ESceneDepthPriorityGroup::SDPG_Foreground);
-		FChaosVDDebugDrawUtils::DrawArrowVector(PDI, SQVisitData.HitData.WorldPosition, SQVisitData.HitData.WorldPosition + SQVisitData.HitData.FaceNormal * NormalScale, nullptr, (FLinearColor(InColor) * 0.65f).ToFColorSRGB(), ESceneDepthPriorityGroup::SDPG_Foreground);
-		FChaosVDDebugDrawUtils::DrawArrowVector(PDI, SQVisitData.HitData.WorldPosition, SQVisitData.HitData.WorldPosition + SQVisitData.HitData.WorldNormal * NormalScale, TEXT("Hit World Normal"), InColor, ESceneDepthPriorityGroup::SDPG_Foreground);
+		FChaosVDDebugDrawUtils::DrawCircle(PDI, SQVisitData.HitData.WorldPosition, CircleRadius, CircleSegments, InColor, Thickness, Axes.GetUnitAxis(EAxis::Y), Axes.GetUnitAxis(EAxis::Z), HitPointDebugText,  ESceneDepthPriorityGroup::SDPG_Foreground);
+		FChaosVDDebugDrawUtils::DrawArrowVector(PDI, SQVisitData.HitData.WorldPosition, SQVisitData.HitData.WorldPosition + SQVisitData.HitData.FaceNormal * NormalScale, HitFaceNormalDebugText, (FLinearColor(InColor) * 0.65f).ToFColorSRGB(), ESceneDepthPriorityGroup::SDPG_Foreground);
+		FChaosVDDebugDrawUtils::DrawArrowVector(PDI, SQVisitData.HitData.WorldPosition, SQVisitData.HitData.WorldPosition + SQVisitData.HitData.WorldNormal * NormalScale, HitWorldNormalDebugText, InColor, ESceneDepthPriorityGroup::SDPG_Foreground);
 
 		if (SQVisitData.bIsSelectedInEditor)
 		{
@@ -225,7 +233,7 @@ void FChaosVDSceneQueryDataComponentVisualizer::DrawHits(const UActorComponent*
 			constexpr float HitSelectionBoxSize = CircleRadius * 1.2f;
 
 			FVector SelectionBoxExtents(HitSelectionBoxSize,HitSelectionBoxSize,HitSelectionBoxSize);
-			FChaosVDDebugDrawUtils::DrawBox(PDI, SelectionBoxExtents, FColor::Yellow, SelectionBoxTransform, nullptr, ESceneDepthPriorityGroup::SDPG_Foreground);
+			FChaosVDDebugDrawUtils::DrawBox(PDI, SelectionBoxExtents, FColor::Yellow, SelectionBoxTransform, FText::GetEmpty(), ESceneDepthPriorityGroup::SDPG_Foreground);
 		}
 
 		PDI->SetHitProxy(nullptr);
@@ -247,12 +255,14 @@ void FChaosVDSceneQueryDataComponentVisualizer::DrawSceneQuery(const UActorCompo
 	{
 		return;
 	}
-			
-	const bool bIsServerQuery = CVDScene->IsSolverForServer(Query->WorldSolverID);
-	const bool bServerDrawIsEnabled = EnumHasAnyFlags(EChaosVDSceneQueryVisualizationFlags::DrawServerQueries, static_cast<EChaosVDSceneQueryVisualizationFlags>(VisualizationContext.VisualizationFlags));
-	const bool bClientDrawIsEnabled = EnumHasAnyFlags(EChaosVDSceneQueryVisualizationFlags::DrawClientQueries, static_cast<EChaosVDSceneQueryVisualizationFlags>(VisualizationContext.VisualizationFlags));
 
-	if ((bIsServerQuery && !bServerDrawIsEnabled) || (!bIsServerQuery && !bClientDrawIsEnabled))
+	AChaosVDSolverInfoActor* SolverInfoActor = CVDScene->GetSolverInfoActor(Query->WorldSolverID);
+	if (!SolverInfoActor)
+	{
+		return;
+	}
+	
+	if (!SolverInfoActor->IsVisible())
 	{
 		return;
 	}
@@ -317,3 +327,5 @@ void FChaosVDSceneQueryDataComponentVisualizer::DrawSceneQuery(const UActorCompo
 		break;
 	}
 }
+
+#undef LOCTEXT_NAMESPACE
diff --git a/Engine/Plugins/Experimental/ChaosVD/Source/ChaosVD/Private/Visualizers/ChaosVDSolverCollisionDataComponentVisualizer.cpp b/Engine/Plugins/Experimental/ChaosVD/Source/ChaosVD/Private/Visualizers/ChaosVDSolverCollisionDataComponentVisualizer.cpp
--- a/Engine/Plugins/Experimental/ChaosVD/Source/ChaosVD/Private/Visualizers/ChaosVDSolverCollisionDataComponentVisualizer.cpp
+++ b/Engine/Plugins/Experimental/ChaosVD/Source/ChaosVD/Private/Visualizers/ChaosVDSolverCollisionDataComponentVisualizer.cpp
@@ -17,6 +17,8 @@
 
 IMPLEMENT_HIT_PROXY(HChaosVDContactPointProxy, HComponentVisProxy)
 
+#define LOCTEXT_NAMESPACE "ChaosVisualDebugger"
+
 void FChaosVDCollisionDataFinder::SetIsSelected(bool bNewSelected)
 {
 	if (TSharedPtr<FChaosVDParticlePairMidPhase> MidPhaseData = OwningMidPhase.Pin())
@@ -55,6 +57,11 @@ void FChaosVDSolverCollisionDataComponentVisualizer::DrawVisualization(const UAc
 		return;
 	}
 
+	if (!SolverInfoContainer->IsVisible())
+	{
+		return;
+	}
+
 	FChaosVDVisualizationContext VisualizationContext;
 	VisualizationContext.SolverID = SolverInfoContainer->GetSolverID();
 	VisualizationContext.CVDScene = SolverInfoContainer->GetScene();
@@ -65,7 +72,12 @@ void FChaosVDSolverCollisionDataComponentVisualizer::DrawVisualization(const UAc
 		VisualizationContext.VisualizationFlags = EditorSettings->GlobalCollisionDataVisualizationFlags;
 	}
 
-	if (EnumHasAnyFlags(EChaosVDCollisionVisualizationFlags::EnableDrawForAllParticles, static_cast<EChaosVDCollisionVisualizationFlags>(VisualizationContext.VisualizationFlags)))
+	if (!EnumHasAnyFlags(EChaosVDCollisionVisualizationFlags::EnableDraw, static_cast<EChaosVDCollisionVisualizationFlags>(VisualizationContext.VisualizationFlags)))
+	{
+		return;
+	}
+
+	if (!EnumHasAnyFlags(EChaosVDCollisionVisualizationFlags::DrawDataOnlyForSelectedParticle, static_cast<EChaosVDCollisionVisualizationFlags>(VisualizationContext.VisualizationFlags)))
 	{
 		for (const TSharedPtr<FChaosVDParticlePairMidPhase>& MidPhase : CollisionDataComponent->GetMidPhases())
 		{
@@ -278,7 +290,7 @@ void FChaosVDSolverCollisionDataComponentVisualizer::DrawMidPhaseData(const UAct
 				{
 					if (ManifoldPoint.bIsValid && !ManifoldPoint.NetPushOut.IsNearlyZero())
 					{
-						FChaosVDDebugDrawUtils::DrawArrowVector(PDI, WorldPointPlaneLocation, WorldPointPlaneLocation + VisualizationContext.SpaceTransform.TransformPosition(ManifoldPoint.NetPushOut), TEXT("Net Push out"), PushOutImpulseColor, DebugDrawSettings.DepthPriority);
+						FChaosVDDebugDrawUtils::DrawArrowVector(PDI, WorldPointPlaneLocation, WorldPointPlaneLocation + VisualizationContext.SpaceTransform.TransformPosition(ManifoldPoint.NetPushOut),  UEnum::GetDisplayValueAsText(EChaosVDCollisionVisualizationFlags::NetPushOut), PushOutImpulseColor, DebugDrawSettings.DepthPriority);
 					}
 				}
 
@@ -286,17 +298,22 @@ void FChaosVDSolverCollisionDataComponentVisualizer::DrawMidPhaseData(const UAct
 				{
 					if (ManifoldPoint.bIsValid && !ManifoldPoint.NetImpulse.IsNearlyZero())
 					{
-						FChaosVDDebugDrawUtils::DrawArrowVector(PDI, WorldPointPlaneLocation, WorldPointPlaneLocation + VisualizationContext.SpaceTransform.TransformPosition(ManifoldPoint.NetImpulse), TEXT("Net Impulse"), ImpulseColor, DebugDrawSettings.DepthPriority);
+						FChaosVDDebugDrawUtils::DrawArrowVector(PDI, WorldPointPlaneLocation, WorldPointPlaneLocation + VisualizationContext.SpaceTransform.TransformPosition(ManifoldPoint.NetImpulse), UEnum::GetDisplayValueAsText(EChaosVDCollisionVisualizationFlags::NetImpulse), ImpulseColor, DebugDrawSettings.DepthPriority);
 					}
 				}
 
 				if (EnumHasAnyFlags(VisualizationFlags, EChaosVDCollisionVisualizationFlags::ContactPoints))
 				{
-					FChaosVDDebugDrawUtils::DrawCircle(PDI, WorldPlaneLocation, DebugDrawSettings.ContactCircleRadius, CircleSegments, DiscColor, LinesThickness, Axes.GetUnitAxis(EAxis::Y), Axes.GetUnitAxis(EAxis::Z), TEXT("Contact : Manifold Plane"), DebugDrawSettings.DepthPriority);
-					FChaosVDDebugDrawUtils::DrawCircle(PDI, WorldPointLocation, 0.5f * DebugDrawSettings.ContactCircleRadius, CircleSegments, DiscColor, LinesThickness, Axes.GetUnitAxis(EAxis::Y), Axes.GetUnitAxis(EAxis::Z), TEXT("Contact : Manifold Point"), DebugDrawSettings.DepthPriority);
+					static FText ContactPointsDebugBaseFText = UEnum::GetDisplayValueAsText(EChaosVDCollisionVisualizationFlags::ContactPoints);
+					static FText ManifoldPlaneDebugText = FText::FormatOrdered(LOCTEXT("ManifoldPlaneDebugText", "{0} | Manifold Plane"), ContactPointsDebugBaseFText);
+					static FText ManifoldPointDebugText = FText::FormatOrdered(LOCTEXT("ManifoldPointDebugText", "{0} | Manifold Point"), ContactPointsDebugBaseFText);
+					static FText ManifoldInitialPhiDebugText = FText::FormatOrdered(LOCTEXT("ManifoldInitialPhiDebugText", "{0} | Manifold Initial Phi"), ContactPointsDebugBaseFText);
+	
+					FChaosVDDebugDrawUtils::DrawCircle(PDI, WorldPlaneLocation, DebugDrawSettings.ContactCircleRadius, CircleSegments, DiscColor, LinesThickness, Axes.GetUnitAxis(EAxis::Y), Axes.GetUnitAxis(EAxis::Z), ManifoldPlaneDebugText, DebugDrawSettings.DepthPriority);
+					FChaosVDDebugDrawUtils::DrawCircle(PDI, WorldPointLocation, 0.5f * DebugDrawSettings.ContactCircleRadius, CircleSegments, DiscColor, LinesThickness, Axes.GetUnitAxis(EAxis::Y), Axes.GetUnitAxis(EAxis::Z), ManifoldPointDebugText, DebugDrawSettings.DepthPriority);
 					if (ManifoldPoint.InitialPhi != 0)
 					{
-						FChaosVDDebugDrawUtils::DrawCircle(PDI, WorldPlaneLocation + ManifoldPoint.InitialPhi * WorldPlaneNormal, 0.25f * DebugDrawSettings.ContactCircleRadius, CircleSegments, InitialPhiColor, LinesThickness, Axes.GetUnitAxis(EAxis::Y), Axes.GetUnitAxis(EAxis::Z), TEXT("Contact : Manifold Initial Phi"), DebugDrawSettings.DepthPriority);
+						FChaosVDDebugDrawUtils::DrawCircle(PDI, WorldPlaneLocation + ManifoldPoint.InitialPhi * WorldPlaneNormal, 0.25f * DebugDrawSettings.ContactCircleRadius, CircleSegments, InitialPhiColor, LinesThickness, Axes.GetUnitAxis(EAxis::Y), Axes.GetUnitAxis(EAxis::Z), ManifoldInitialPhiDebugText, DebugDrawSettings.DepthPriority);
 					}
 				}
 
@@ -304,14 +321,14 @@ void FChaosVDSolverCollisionDataComponentVisualizer::DrawMidPhaseData(const UAct
 				{
 					FColor NormalColor = ((ManifoldPoint.ContactPoint.ContactType != EChaosVDContactPointType::EdgeEdge) ? PlaneNormalColor : EdgeNormalColor);
 					const int32 Scale = DebugDrawSettings.ContactNormalScale * ContactLenScale;
-					FChaosVDDebugDrawUtils::DrawArrowVector(PDI, WorldPlaneLocation, WorldPlaneLocation + WorldPlaneNormal * Scale, TEXT("Contact Normal"), NormalColor, DebugDrawSettings.DepthPriority);
+					FChaosVDDebugDrawUtils::DrawArrowVector(PDI, WorldPlaneLocation, WorldPlaneLocation + WorldPlaneNormal * Scale, UEnum::GetDisplayValueAsText(EChaosVDCollisionVisualizationFlags::ContactNormal), NormalColor, DebugDrawSettings.DepthPriority);
 				}
 
 				if (EnumHasAnyFlags(VisualizationFlags, EChaosVDCollisionVisualizationFlags::AccumulatedImpulse))
 				{
 					if (!Constraint.AccumulatedImpulse.IsNearlyZero())
 					{
-						FChaosVDDebugDrawUtils::DrawArrowVector(PDI, WorldActorTransform0.GetLocation(), WorldActorTransform0.GetLocation() + Constraint.AccumulatedImpulse, TEXT("AccumulatedImpulse"), FColor::White, DebugDrawSettings.DepthPriority);
+						FChaosVDDebugDrawUtils::DrawArrowVector(PDI, WorldActorTransform0.GetLocation(), WorldActorTransform0.GetLocation() + Constraint.AccumulatedImpulse, UEnum::GetDisplayValueAsText(EChaosVDCollisionVisualizationFlags::AccumulatedImpulse), FColor::White, DebugDrawSettings.DepthPriority);
 					}
 				}
 
@@ -327,23 +344,23 @@ void FChaosVDSolverCollisionDataComponentVisualizer::DrawMidPhaseData(const UAct
 					if (Constraint.bWasManifoldRestored)
 					{
 						Color = FColor::Blue;
-						FChaosVDDebugDrawUtils::DrawBox(PDI, BoxExtents, Color, Transform, nullptr, DebugDrawSettings.DepthPriority);
+						FChaosVDDebugDrawUtils::DrawBox(PDI, BoxExtents, Color, Transform, FText::GetEmpty(), DebugDrawSettings.DepthPriority);
 					}
 					else if (ManifoldPoint.bWasRestored)
 					{
 						Color = FColor::Purple;
-						FChaosVDDebugDrawUtils::DrawBox(PDI, BoxExtents, Color, Transform, nullptr, DebugDrawSettings.DepthPriority);
+						FChaosVDDebugDrawUtils::DrawBox(PDI, BoxExtents, Color, Transform, FText::GetEmpty(), DebugDrawSettings.DepthPriority);
 					}
 					else if (ManifoldPoint.bWasReplaced)
 					{
 						Color = FColor::Orange;
-						FChaosVDDebugDrawUtils::DrawBox(PDI, BoxExtents, Color, Transform, nullptr, DebugDrawSettings.DepthPriority);
+						FChaosVDDebugDrawUtils::DrawBox(PDI, BoxExtents, Color, Transform, FText::GetEmpty(), DebugDrawSettings.DepthPriority);
 					}
 
 					if (MidPhase->bIsSleeping)
 					{
 						// This box should surround the debug draw box we made before
-						FChaosVDDebugDrawUtils::DrawBox(PDI, BoxExtents * 1.1f, FColor::Black, Transform, nullptr, DebugDrawSettings.DepthPriority);
+						FChaosVDDebugDrawUtils::DrawBox(PDI, BoxExtents * 1.1f, FColor::Black, Transform, FText::GetEmpty(), DebugDrawSettings.DepthPriority);
 					}
 				}
 
@@ -361,9 +378,11 @@ void FChaosVDSolverCollisionDataComponentVisualizer::DrawMidPhaseData(const UAct
 					const float ContactSelectionBoxSize = DebugDrawSettings.ContactCircleRadius * 1.5f;
 
 					FVector SelectionBoxExtents(ContactSelectionBoxSize,ContactSelectionBoxSize,ContactSelectionBoxSize);
-					FChaosVDDebugDrawUtils::DrawBox(PDI, SelectionBoxExtents, FColor::Yellow, SelectionBoxTransform, nullptr, DebugDrawSettings.DepthPriority);
+					FChaosVDDebugDrawUtils::DrawBox(PDI, SelectionBoxExtents, FColor::Yellow, SelectionBoxTransform, FText::GetEmpty(), DebugDrawSettings.DepthPriority);
 				}
 			}
 		}
 	}
 }
+
+#undef LOCTEXT_NAMESPACE
diff --git a/Engine/Plugins/Experimental/ChaosVD/Source/ChaosVD/Private/Visualizers/ChaosVDSolverCollisionDataComponentVisualizer.h b/Engine/Plugins/Experimental/ChaosVD/Source/ChaosVD/Private/Visualizers/ChaosVDSolverCollisionDataComponentVisualizer.h
--- a/Engine/Plugins/Experimental/ChaosVD/Source/ChaosVD/Private/Visualizers/ChaosVDSolverCollisionDataComponentVisualizer.h
+++ b/Engine/Plugins/Experimental/ChaosVD/Source/ChaosVD/Private/Visualizers/ChaosVDSolverCollisionDataComponentVisualizer.h
@@ -37,15 +37,16 @@ struct HChaosVDContactPointProxy : public HComponentVisProxy
 UENUM(meta = (Bitflags, UseEnumValuesAsMaskValuesInEditor = "true"))
 enum class EChaosVDCollisionVisualizationFlags: uint32
 {
-	None					= 0 UMETA(Hidden),
-	ContactPoints			= 1 << 0,
-	ContactInfo				= 1 << 1,
-	NetPushOut				= 1 << 2,
-	NetImpulse				= 1 << 3,
-	ContactNormal			= 1 << 4,
-	AccumulatedImpulse		= 1 << 5,
-	DrawInactiveContacts	= 1 << 6,
-	EnableDrawForAllParticles	= 1 << 7,
+	None							= 0 UMETA(Hidden),
+	ContactPoints					= 1 << 0,
+	ContactInfo						= 1 << 1,
+	NetPushOut						= 1 << 2,
+	NetImpulse						= 1 << 3,
+	ContactNormal					= 1 << 4,
+	AccumulatedImpulse				= 1 << 5,
+	DrawInactiveContacts			= 1 << 6,
+	DrawDataOnlyForSelectedParticle	= 1 << 7,
+	EnableDraw						= 1 << 8,
 };
 ENUM_CLASS_FLAGS(EChaosVDCollisionVisualizationFlags);
 
diff --git a/Engine/Plugins/Experimental/ChaosVD/Source/ChaosVD/Public/Actors/ChaosVDSolverInfoActor.h b/Engine/Plugins/Experimental/ChaosVD/Source/ChaosVD/Public/Actors/ChaosVDSolverInfoActor.h
--- a/Engine/Plugins/Experimental/ChaosVD/Source/ChaosVD/Public/Actors/ChaosVDSolverInfoActor.h
+++ b/Engine/Plugins/Experimental/ChaosVD/Source/ChaosVD/Public/Actors/ChaosVDSolverInfoActor.h
@@ -64,6 +64,8 @@ public:
 	void HandleColorsSettingsUpdated();
 	void RemoveSolverFolders(UWorld* World);
 
+	bool IsVisible() const;
+
 #if WITH_EDITOR
 	void SetIsTemporarilyHiddenInEditor(bool bIsHidden) override;
 #endif
diff --git a/Engine/Plugins/Experimental/ChaosVD/Source/ChaosVD/Public/ChaosVDEditorSettings.h b/Engine/Plugins/Experimental/ChaosVD/Source/ChaosVD/Public/ChaosVDEditorSettings.h
--- a/Engine/Plugins/Experimental/ChaosVD/Source/ChaosVD/Public/ChaosVDEditorSettings.h
+++ b/Engine/Plugins/Experimental/ChaosVD/Source/ChaosVD/Public/ChaosVDEditorSettings.h
@@ -320,19 +320,19 @@ public:
 
 	/** Set of flags to enable/disable visualization of specific particle data as debug draw */
 	UPROPERTY(EditAnywhere, Category = "Viewport Visualization Flags", meta = (Bitmask, BitmaskEnum = "/Script/ChaosVD.EChaosVDParticleDataVisualizationFlags"))
-	uint32 GlobalParticleDataVisualizationFlags = static_cast<uint32>(EChaosVDParticleDataVisualizationFlags::Velocity | EChaosVDParticleDataVisualizationFlags::AngularVelocity | EChaosVDParticleDataVisualizationFlags::DrawDataOnlyForSelectedParticle);
+	uint32 GlobalParticleDataVisualizationFlags = static_cast<uint32>(EChaosVDParticleDataVisualizationFlags::Velocity | EChaosVDParticleDataVisualizationFlags::AngularVelocity);
 
 	/** Set of flags to enable/disable visualization of specific collision data as debug draw */
 	UPROPERTY(EditAnywhere, Category = "Viewport Visualization Flags", meta = (Bitmask, BitmaskEnum = "/Script/ChaosVD.EChaosVDCollisionVisualizationFlags"))
-	uint32 GlobalCollisionDataVisualizationFlags = 0;
+	uint32 GlobalCollisionDataVisualizationFlags = static_cast<uint32>(EChaosVDCollisionVisualizationFlags::ContactInfo | EChaosVDCollisionVisualizationFlags::ContactPoints | EChaosVDCollisionVisualizationFlags::ContactNormal);
 
 	/** Set of flags to enable/disable visualization of specific scene queries data as debug draw */
 	UPROPERTY(EditAnywhere, Category = "Viewport Visualization Flags", meta = (Bitmask, BitmaskEnum = "/Script/ChaosVD.EChaosVDSceneQueryVisualizationFlags"))
-	uint32 GlobalSceneQueriesVisualizationFlags = static_cast<uint32>(EChaosVDSceneQueryVisualizationFlags::DrawClientQueries | EChaosVDSceneQueryVisualizationFlags::DrawServerQueries |  EChaosVDSceneQueryVisualizationFlags::DrawHits | EChaosVDSceneQueryVisualizationFlags::DrawLineTraceQueries);
+	uint32 GlobalSceneQueriesVisualizationFlags = static_cast<uint32>(EChaosVDSceneQueryVisualizationFlags::DrawHits | EChaosVDSceneQueryVisualizationFlags::DrawLineTraceQueries);
 
 	/** Set of flags to enable/disable visualization of specific scene queries data as debug draw */
 	UPROPERTY(EditAnywhere, Category = "Viewport Visualization Flags", meta = (Bitmask, BitmaskEnum = "/Script/ChaosVD.EChaosVDJointsDataVisualizationFlags"))
-	uint32 GlobalJointsDataVisualizationFlags = static_cast<uint32>(EChaosVDJointsDataVisualizationFlags::EnableClientDraw | EChaosVDJointsDataVisualizationFlags::EnableServerDraw | EChaosVDJointsDataVisualizationFlags::ActorConnector | EChaosVDJointsDataVisualizationFlags::DrawKinematic);
+	uint32 GlobalJointsDataVisualizationFlags = static_cast<uint32>(EChaosVDJointsDataVisualizationFlags::ActorConnector | EChaosVDJointsDataVisualizationFlags::DrawKinematic);
 
 	/** If true, text information (if available) will be drawn alongside any other debug draw shape */
 	UPROPERTY(EditAnywhere, Category = "Viewport Visualization")
diff --git a/Engine/Plugins/Experimental/ChaosVD/Source/ChaosVD/Public/Visualizers/ChaosVDDebugDrawUtils.h b/Engine/Plugins/Experimental/ChaosVD/Source/ChaosVD/Public/Visualizers/ChaosVDDebugDrawUtils.h
--- a/Engine/Plugins/Experimental/ChaosVD/Source/ChaosVD/Public/Visualizers/ChaosVDDebugDrawUtils.h
+++ b/Engine/Plugins/Experimental/ChaosVD/Source/ChaosVD/Public/Visualizers/ChaosVDDebugDrawUtils.h
@@ -6,33 +6,56 @@
 #include "Containers/Queue.h"
 #include "SceneManagement.h"
 
+class FChaosVDPlaybackViewportClient;
 class FChaosVDGeometryBuilder;
 
+enum class EChaosVDDebugDrawTextLocationMode
+{
+	World,
+	Screen
+};
+
 /** Utility methods that allows Debug draw into the Chaos VD Editor */
 class FChaosVDDebugDrawUtils
 {
 public:
 	//TODO: Create a common args struct that can be passed as parameter to all debug draw methods
 
-	static void DrawArrowVector(FPrimitiveDrawInterface* PDI, const FVector& StartLocation, const FVector& EndLocation, FStringView DebugText, const FColor& Color, ESceneDepthPriorityGroup DepthPriority = SDPG_World, float Thickness = 0.0f);
-	static void DrawPoint(FPrimitiveDrawInterface* PDI, const FVector& Location, FStringView DebugText, const FColor& Color, float Size = 10.0f, ESceneDepthPriorityGroup DepthPriority = SDPG_World);
-	static void DrawText(FStringView StringToDraw, const FVector& Location, const FColor& Color);
-	static void DrawCircle(FPrimitiveDrawInterface* PDI, const FVector& Origin, float Radius, int32 Segments, const FColor& Color, float Thickness, const FVector& XAxis, const FVector& YAxis, FStringView DebugText, ESceneDepthPriorityGroup DepthPriority = SDPG_World);
-	static void DrawBox(FPrimitiveDrawInterface* PDI, const FVector& InExtents, const FColor& InColor, const FTransform& InTransform, FStringView DebugText, ESceneDepthPriorityGroup DepthPriority = SDPG_World, float Thickness = 2.0f);
-	static void DrawLine(FPrimitiveDrawInterface* PDI, const FVector& InStartPosition, const FVector& InEndPosition, const FColor& InColor, FStringView DebugText, ESceneDepthPriorityGroup DepthPriority = SDPG_World, float Thickness = 2.0f);
-	static void DrawImplicitObject(FPrimitiveDrawInterface* PDI, const TSharedPtr<FChaosVDGeometryBuilder>& GeometryGenerator, const Chaos::FConstImplicitObjectPtr& ImplicitObject, const FTransform& InWorldTransform, const FColor& InColor, FStringView DebugText, ESceneDepthPriorityGroup DepthPriority = SDPG_World, float Thickness = 2.0f);
+	static void DrawArrowVector(FPrimitiveDrawInterface* PDI, const FVector& StartLocation, const FVector& EndLocation, const FText& InDebugText, const FColor& Color, ESceneDepthPriorityGroup DepthPriority = SDPG_World, float Thickness = 0.0f);
+	static void DrawPoint(FPrimitiveDrawInterface* PDI, const FVector& Location, const FText& InDebugText, const FColor& Color, float Size = 10.0f, ESceneDepthPriorityGroup DepthPriority = SDPG_World);
+	static void DrawString(FStringView StringToDraw, const FVector& Location, const FColor& Color, EChaosVDDebugDrawTextLocationMode LocationMode = EChaosVDDebugDrawTextLocationMode::World);
+	static void DrawText(const FText& InText, const FVector& Location, const FColor& Color, EChaosVDDebugDrawTextLocationMode LocationMode = EChaosVDDebugDrawTextLocationMode::World);
+	static void DrawOnScreenWarning(const FText& InText, const FColor& Color);
+	static void DrawCircle(FPrimitiveDrawInterface* PDI, const FVector& Origin, float Radius, int32 Segments, const FColor& Color, float Thickness, const FVector& XAxis, const FVector& YAxis, const FText& InDebugText, ESceneDepthPriorityGroup DepthPriority = SDPG_World);
+	static void DrawBox(FPrimitiveDrawInterface* PDI, const FVector& InExtents, const FColor& InColor, const FTransform& InTransform, const FText& InDebugText, ESceneDepthPriorityGroup DepthPriority = SDPG_World, float Thickness = 2.0f);
+	static void DrawLine(FPrimitiveDrawInterface* PDI, const FVector& InStartPosition, const FVector& InEndPosition, const FColor& InColor, const FText& InDebugText, ESceneDepthPriorityGroup DepthPriority = SDPG_World, float Thickness = 2.0f);
+	static void DrawImplicitObject(FPrimitiveDrawInterface* PDI, const TSharedPtr<FChaosVDGeometryBuilder>& GeometryGenerator, const Chaos::FConstImplicitObjectPtr& ImplicitObject, const FTransform& InWorldTransform, const FColor& InColor, const FText& InDebugText, ESceneDepthPriorityGroup DepthPriority = SDPG_World, float Thickness = 2.0f);
 
 	static void DrawCanvas(FViewport& InViewport, FSceneView& View, FCanvas& Canvas);
+	
+	static bool CanDebugDraw();
+	
+	static int32 MaxLinesToDrawPerFrame;
 
 private:
 	struct FChaosVDQueuedTextToDraw
 	{
-		FString Text;
-		FVector WorldPosition;
+		FText Text;
+		EChaosVDDebugDrawTextLocationMode LocationMode;
+		FVector Location;
 		FLinearColor Color;
 	};
 
-	static TQueue<FChaosVDQueuedTextToDraw> TexToDrawQueue;
+	static void DebugDrawFrameEnd();
+	static void IncreaseDebugDrawLineCounter();
+
+	static TQueue<FChaosVDQueuedTextToDraw> TextToDrawQueue;
+
+	static int32 CurrentLinesDrawn;
+	static bool bIsShowingDebugDrawLimitWarning;
+	static int32 CurrentWarningsBeingDrawn;
+
+	friend FChaosVDPlaybackViewportClient;
 };
 
 namespace Chaos::VisualDebugger::Utils
diff --git a/Engine/Plugins/Experimental/ChaosVD/Source/ChaosVD/Public/Visualizers/ChaosVDJointConstraintsDataComponentVisualizer.h b/Engine/Plugins/Experimental/ChaosVD/Source/ChaosVD/Public/Visualizers/ChaosVDJointConstraintsDataComponentVisualizer.h
--- a/Engine/Plugins/Experimental/ChaosVD/Source/ChaosVD/Public/Visualizers/ChaosVDJointConstraintsDataComponentVisualizer.h
+++ b/Engine/Plugins/Experimental/ChaosVD/Source/ChaosVD/Public/Visualizers/ChaosVDJointConstraintsDataComponentVisualizer.h
@@ -26,14 +26,10 @@ enum class EChaosVDJointsDataVisualizationFlags : uint32
 	DrawKinematic			= 1 << 6,
 	/** Draw the joint even if it is disabled */
 	DrawDisabled			= 1 << 7,
-	/** Enables Debug draw for Joint Constraint data from Game Client solvers */
-	EnableClientDraw		= 1 << 8,
-	/** Enables Debug draw for Joint Constraint data from Game Server solvers */
-	EnableServerDraw		= 1 << 9,
 	/** Only debugs draw data for a selected joint constraint */
-	OnlyDrawSelected		= 1 << 10,
-	/** Enables Debug draw for Joint Constraint data from any solver */
-	EnableDraw				= EnableClientDraw | EnableServerDraw,
+	OnlyDrawSelected		= 1 << 8,
+	/** Enables Debug draw for Joint Constraint data from any solver that is visible */
+	EnableDraw				= 1 << 9,
 };
 ENUM_CLASS_FLAGS(EChaosVDJointsDataVisualizationFlags);
 
diff --git a/Engine/Plugins/Experimental/ChaosVD/Source/ChaosVD/Public/Visualizers/ChaosVDSceneQueryDataComponentVisualizer.h b/Engine/Plugins/Experimental/ChaosVD/Source/ChaosVD/Public/Visualizers/ChaosVDSceneQueryDataComponentVisualizer.h
--- a/Engine/Plugins/Experimental/ChaosVD/Source/ChaosVD/Public/Visualizers/ChaosVDSceneQueryDataComponentVisualizer.h
+++ b/Engine/Plugins/Experimental/ChaosVD/Source/ChaosVD/Public/Visualizers/ChaosVDSceneQueryDataComponentVisualizer.h
@@ -17,8 +17,6 @@ struct FChaosVDQueryDataWrapper;
 struct FChaosVDSceneQueryVisualizationDataContext : public FChaosVDVisualizationContext
 {
 	FChaosVDSceneQuerySelectionHandle DataSelectionHandle = FChaosVDSceneQuerySelectionHandle(nullptr, INDEX_NONE);
-	bool bIsServerVisualizationEnabled = false;
-	bool bIsClientVisualizationEnabled = false;
 
 	/** Generates a random color based on the selection state and query ID, which will be used to debug draw the scene query */
 	void GenerateColor(int32 QueryID, bool bIsSelected)
@@ -65,11 +63,9 @@ enum class EChaosVDSceneQueryVisualizationFlags: uint32
 	DrawSweepQueries		= 1 << 2,
 	DrawOverlapQueries		= 1 << 3,
 	DrawHits				= 1 << 4,
-	DrawServerQueries		= 1 << 5,
-	DrawClientQueries		= 1 << 6,
-	OnlyDrawSelectedQuery	= 1 << 7,
-	HideEmptyQueries		= 1 << 8,
-	HideSubQueries			= 1 << 9,
+	OnlyDrawSelectedQuery	= 1 << 5,
+	HideEmptyQueries		= 1 << 6,
+	HideSubQueries			= 1 << 7,
 };
 ENUM_CLASS_FLAGS(EChaosVDSceneQueryVisualizationFlags);
 
