[ChaosVD] Fix for an issue where the editor hangs if we attempt to start a live recording when the trace server is not available

- Fixed an issue with the code that limited the number of acoonections attempts
- Reduced the number of connection attempts to 4
- Added a notification message to inform we are attempting to connect to a live session and on which attempt number we are
- If the recording failed to start, we don't attempt to connect to the live session because there will be none

diff --git a/Engine/Plugins/Experimental/ChaosVD/Source/ChaosVD/Private/ChaosVDPlaybackController.cpp b/Engine/Plugins/Experimental/ChaosVD/Source/ChaosVD/Private/ChaosVDPlaybackController.cpp
--- a/Engine/Plugins/Experimental/ChaosVD/Source/ChaosVD/Private/ChaosVDPlaybackController.cpp
+++ b/Engine/Plugins/Experimental/ChaosVD/Source/ChaosVD/Private/ChaosVDPlaybackController.cpp
@@ -51,7 +51,6 @@ bool FChaosVDPlaybackController::LoadChaosVDRecordingFromTraceSession(const FCha
 {
 	if (!ensure(!InSessionDescriptor.SessionName.IsEmpty()))
 	{
-		FMessageDialog::Open(EAppMsgType::Ok, LOCTEXT("LoadRecordingSessionFailedMessage", "Failed to create a session from the selected CVD recording. Please see the logs for more details... "));
 		return false;
 	}
 
diff --git a/Engine/Plugins/Experimental/ChaosVD/Source/ChaosVD/Private/Widgets/SChaosVDRecordingControls.cpp b/Engine/Plugins/Experimental/ChaosVD/Source/ChaosVD/Private/Widgets/SChaosVDRecordingControls.cpp
--- a/Engine/Plugins/Experimental/ChaosVD/Source/ChaosVD/Private/Widgets/SChaosVDRecordingControls.cpp
+++ b/Engine/Plugins/Experimental/ChaosVD/Source/ChaosVD/Private/Widgets/SChaosVDRecordingControls.cpp
@@ -2,6 +2,7 @@
 
 #include "Widgets/SChaosVDRecordingControls.h"
 
+#include "AsyncCompilationHelpers.h"
 #include "ChaosVDEngine.h"
 #include "ChaosVDModule.h"
 #include "ChaosVDStyle.h"
@@ -14,10 +15,12 @@
 #include "ToolMenu.h"
 #include "ToolMenus.h"
 #include "Framework/MultiBox/MultiBoxBuilder.h"
+#include "Misc/ScopedSlowTask.h"
 #include "Widgets/Images/SImage.h"
 #include "Widgets/Input/SButton.h"
 #include "Widgets/SChaosVDMainTab.h"
 #include "Widgets/Input/SComboButton.h"
+#include "Widgets/Notifications/SNotificationList.h"
 #include "Widgets/Text/STextBlock.h"
 
 #define LOCTEXT_NAMESPACE "ChaosVisualDebugger"
@@ -225,6 +228,11 @@ SChaosVDRecordingControls::~SChaosVDRecordingControls()
 		FChaosVDRuntimeModule::Get().RemoveRecordingStartedCallback(RecordingStartedHandle);
 		FChaosVDRuntimeModule::Get().RemoveRecordingStopCallback(RecordingStoppedHandle);
 	}
+
+	if (ConnectionAttemptNotification.IsValid())
+	{
+		ConnectionAttemptNotification->ExpireAndFadeout();
+	}
 }
 
 const FSlateBrush* SChaosVDRecordingControls::GetRecordOrStopButton(EChaosVDRecordingMode RecordingMode) const
@@ -293,7 +301,12 @@ void SChaosVDRecordingControls::HandleRecordingStart()
 
 void SChaosVDRecordingControls::AttemptToConnectToLiveSession()
 {
-	bAutoConnectionAttemptInProgress = true;
+	if(!bAutoConnectionAttemptInProgress)
+	{
+		bAutoConnectionAttemptInProgress = true;
+		PushConnectionAttemptNotification();
+	}
+
 	// We need to wait at least one tick before attempting to connect
 	FTSTicker::GetCoreTicker().AddTicker(FTickerDelegate::CreateLambda([WeakThis = AsWeak()](float DeltaTime)
 	{
@@ -301,6 +314,10 @@ void SChaosVDRecordingControls::AttemptToConnectToLiveSession()
 		{
 			if (const TSharedPtr<SChaosVDMainTab> MainTabSharedPtr = RecordingControlsPtr->MainTabWeakPtr.Pin())
 			{
+				RecordingControlsPtr->CurrentConnectionAttempts++;
+
+				RecordingControlsPtr->UpdateConnectionAttemptNotification();
+
 				static FString SessionAddress(TEXT("127.0.0.1"));
 			
 				int32 SessionID = 0;
@@ -325,13 +342,13 @@ void SChaosVDRecordingControls::AttemptToConnectToLiveSession()
 					}
 					else
 					{
-						RecordingControlsPtr->bAutoConnectionAttemptInProgress = false;
+						RecordingControlsPtr->HandleConnectionAttemptResult(EChaosVDLiveConnectionAttemptResult::Failed);
 						UE_LOG(LogChaosVDEditor, Error, TEXT("[%s] Failed to connect to live session | [%d] attempts exhausted..."), ANSI_TO_TCHAR(__FUNCTION__), RecordingControlsPtr->MaxAutoplayConnectionAttempts);	
 					}
 				}
 				else
 				{
-					RecordingControlsPtr->bAutoConnectionAttemptInProgress = false;
+					RecordingControlsPtr->HandleConnectionAttemptResult(EChaosVDLiveConnectionAttemptResult::Success);
 				}
 			}
 		}
@@ -351,7 +368,12 @@ FReply SChaosVDRecordingControls::ToggleRecordingState(EChaosVDRecordingMode Rec
 
 			FChaosVDRuntimeModule::Get().StartRecording(RecordingArgs);
 
-			AttemptToConnectToLiveSession();
+			// Only attempt to connect if we managed to start a recording.
+			// If we failed, the runtime module takes care of showing a pop-up error in the editor already
+			if (FChaosVDRuntimeModule::Get().IsRecording())
+			{
+				AttemptToConnectToLiveSession();
+			}	
 		}
 		else
 		{
@@ -482,4 +504,51 @@ FText SChaosVDRecordingControls::GetRecordingTimeText() const
 	return FText::Format(LOCTEXT("RecordingTimer","{0} s"), SecondsText);
 }
 
+void SChaosVDRecordingControls::PushConnectionAttemptNotification()
+{
+	FNotificationInfo Info(LOCTEXT("ConnectingToLiceSessionMessge", "Connecting to Live Session ..."));
+	Info.bFireAndForget = false;
+	Info.FadeOutDuration = 3.0f;
+	Info.ExpireDuration = 0.0f;
+
+	ConnectionAttemptNotification = FSlateNotificationManager::Get().AddNotification(Info);
+	
+	if (ConnectionAttemptNotification.IsValid())
+	{
+		ConnectionAttemptNotification->SetCompletionState(SNotificationItem::CS_Pending);
+	}
+}
+
+void SChaosVDRecordingControls::UpdateConnectionAttemptNotification()
+{
+	if (ConnectionAttemptNotification.IsValid())
+	{
+		ConnectionAttemptNotification->SetSubText(FText::FormatOrdered(LOCTEXT("SessionConnectionAttemptSubText", "Attempt {0} / {1}"), CurrentConnectionAttempts, MaxAutoplayConnectionAttempts));
+	}
+}
+
+void SChaosVDRecordingControls::HandleConnectionAttemptResult(EChaosVDLiveConnectionAttemptResult Result)
+{
+	CurrentConnectionAttempts = 0;
+	bAutoConnectionAttemptInProgress = false;
+
+	if (ConnectionAttemptNotification.IsValid())
+	{
+		if (Result == EChaosVDLiveConnectionAttemptResult::Success)
+		{
+			ConnectionAttemptNotification->SetText(LOCTEXT("SessionConnectionSuccess", "Connected!"));
+			ConnectionAttemptNotification->SetSubText(FText::GetEmpty());
+			ConnectionAttemptNotification->SetCompletionState(SNotificationItem::ECompletionState::CS_Success);
+		}
+		else
+		{
+			ConnectionAttemptNotification->SetText(LOCTEXT("SessionConnectionFailedText", "Failed to connect"));
+			ConnectionAttemptNotification->SetSubText(LOCTEXT("SessionConnectionFailedSubText", "See the logs for more details..."));
+			ConnectionAttemptNotification->SetCompletionState(SNotificationItem::ECompletionState::CS_Fail);
+		}
+
+		ConnectionAttemptNotification->ExpireAndFadeout();
+	}
+}
+
 #undef LOCTEXT_NAMESPACE 
diff --git a/Engine/Plugins/Experimental/ChaosVD/Source/ChaosVD/Public/Widgets/SChaosVDRecordingControls.h b/Engine/Plugins/Experimental/ChaosVD/Source/ChaosVD/Public/Widgets/SChaosVDRecordingControls.h
--- a/Engine/Plugins/Experimental/ChaosVD/Source/ChaosVD/Public/Widgets/SChaosVDRecordingControls.h
+++ b/Engine/Plugins/Experimental/ChaosVD/Source/ChaosVD/Public/Widgets/SChaosVDRecordingControls.h
@@ -13,9 +13,10 @@ namespace Chaos::VisualDebugger
 	struct FChaosVDOptionalDataChannel;
 }
 
+class FAsyncCompilationNotification;
+class FReply;
 class SButton;
 class SChaosVDMainTab;
-class FReply;
 struct FSlateBrush;
 
 UCLASS()
@@ -35,6 +36,13 @@ enum class EChaosVDRecordingMode
 	Live
 };
 
+UENUM()
+enum class EChaosVDLiveConnectionAttemptResult
+{
+	Success,
+	Failed
+};
+
 typedef Chaos::VisualDebugger::FChaosVDOptionalDataChannel FCVDDataChannel;
 
 class SChaosVDRecordingControls : public SCompoundWidget
@@ -79,6 +87,11 @@ protected:
 	bool IsRecording() const;
 
 	FText GetRecordingTimeText() const;
+
+	void PushConnectionAttemptNotification();
+	void UpdateConnectionAttemptNotification();
+
+	void HandleConnectionAttemptResult(EChaosVDLiveConnectionAttemptResult Result);
 	
 	FName StatusBarID;
 	
@@ -90,14 +103,16 @@ protected:
 
 	TWeakPtr<SChaosVDMainTab> MainTabWeakPtr;
 
-	int32 MaxAutoplayConnectionAttempts = 20;
-	float IntervalBetweenAutoplayConnectionAttemptsSeconds = 0.1f;
-	bool bAutoConnectionAttemptInProgress = false;
+	int32 MaxAutoplayConnectionAttempts = 4;
+	float IntervalBetweenAutoplayConnectionAttemptsSeconds = 1.0f;
 	int32 CurrentConnectionAttempts = 0;
+	bool bAutoConnectionAttemptInProgress = false;
 
 	bool bRecordingButtonHovered = false;
 
 	FCurveSequence RecordingAnimation;
 
+	TSharedPtr<SNotificationItem> ConnectionAttemptNotification;
+
 	static inline const FName RecordingControlsToolbarName = FName("ChaosVD.MainToolBar.RecordingControls");
 };
