[ChaosVD] Moving most headers to the private folder in the ChaosVD plugin in preparation to release the plugin as beta in 5.4.

The public API for the plugin that will allow users to extend will be ready for 5.5

diff --git a/Engine/Plugins/Experimental/ChaosVD/Source/ChaosVD/Public/ChaosVDModule.h b/Engine/Plugins/Experimental/ChaosVD/Source/ChaosVD/Public/ChaosVDModule.h
--- a/Engine/Plugins/Experimental/ChaosVD/Source/ChaosVD/Public/ChaosVDModule.h
+++ b/Engine/Plugins/Experimental/ChaosVD/Source/ChaosVD/Public/ChaosVDModule.h
@@ -12,6 +12,7 @@ class SDockTab;
 class FSpawnTabArgs;
 struct FGuid;
 
+/** Main module class for the Chaos Visual Debugger editor */
 DECLARE_LOG_CATEGORY_EXTERN(LogChaosVDEditor, Log, Log);
 class FChaosVDModule : public IModuleInterface
 {
@@ -23,8 +24,10 @@ public:
 	virtual void StartupModule() override;
 	virtual void ShutdownModule() override;
 
+	/** Returns the Trace manager instance used by CVD */
 	TSharedPtr<FChaosVDTraceManager>& GetTraceManager() { return ChaosVDTraceManager; };
 
+	/** Spawns a new instance of CVD */
 	void SpawnCVDTab();
 
 private:
