[ChaosVD] Fix for localization key conficts in FChaosVDSceneQueryDataComponentVisualizer

diff --git a/Engine/Plugins/ChaosVD/Source/ChaosVD/Private/Visualizers/ChaosVDSceneQueryDataComponentVisualizer.cpp b/Engine/Plugins/ChaosVD/Source/ChaosVD/Private/Visualizers/ChaosVDSceneQueryDataComponentVisualizer.cpp
--- a/Engine/Plugins/ChaosVD/Source/ChaosVD/Private/Visualizers/ChaosVDSceneQueryDataComponentVisualizer.cpp
+++ b/Engine/Plugins/ChaosVD/Source/ChaosVD/Private/Visualizers/ChaosVDSceneQueryDataComponentVisualizer.cpp
@@ -146,7 +146,7 @@ void FChaosVDSceneQueryDataComponentVisualizer::DrawOverlapQuery(const UActorCom
 	const Chaos::FConstImplicitObjectPtr InputShapePtr = VisualizationContext.InputGeometry;
 	if (ensure(InputShapePtr))
 	{
-		const FText DebugText = FText::FormatOrdered(LOCTEXT("LineTraceDebugDrawText", "Type: Overlap \n Tag {1} \n Owner Tag {2}"), FText::AsCultureInvariant(SceneQueryData.CollisionQueryParams.TraceTag.ToString()), FText::AsCultureInvariant(SceneQueryData.CollisionQueryParams.OwnerTag.ToString()));
+		const FText DebugText = FText::FormatOrdered(LOCTEXT("OverlapDebugDrawText", "Type: Overlap \n Tag {1} \n Owner Tag {2}"), FText::AsCultureInvariant(SceneQueryData.CollisionQueryParams.TraceTag.ToString()), FText::AsCultureInvariant(SceneQueryData.CollisionQueryParams.OwnerTag.ToString()));
 		FChaosVDDebugDrawUtils::DrawImplicitObject(PDI, VisualizationContext.GeometryGenerator.Pin(), InputShapePtr, FTransform(SceneQueryData.GeometryOrientation, SceneQueryData.StartLocation), VisualizationContext.DebugDrawColor, DebugText, ESceneDepthPriorityGroup::SDPG_Foreground);
 	}
 
@@ -184,7 +184,7 @@ void FChaosVDSceneQueryDataComponentVisualizer::DrawSweepQuery(const UActorCompo
 
 	PDI->SetHitProxy(nullptr);
 
-	const FText DebugText = FText::FormatOrdered(LOCTEXT("LineTraceDebugDrawText", "Type: Sweep \n Tag {1} \n Owner Tag {2}"), FText::AsCultureInvariant(SceneQueryData.CollisionQueryParams.TraceTag.ToString()), FText::AsCultureInvariant(SceneQueryData.CollisionQueryParams.OwnerTag.ToString()));
+	const FText DebugText = FText::FormatOrdered(LOCTEXT("SweepDebugDrawText", "Type: Sweep \n Tag {1} \n Owner Tag {2}"), FText::AsCultureInvariant(SceneQueryData.CollisionQueryParams.TraceTag.ToString()), FText::AsCultureInvariant(SceneQueryData.CollisionQueryParams.OwnerTag.ToString()));
 
 	FChaosVDDebugDrawUtils::DrawArrowVector(PDI, SceneQueryData.StartLocation, SceneQueryData.EndLocation, DebugText, VisualizationContext.DebugDrawColor, SDPG_Foreground);
 
