[ChaosVD] Fixed an issue where you can't change between actions in the Transform Toolbar

Also fixed another issue where if you select the "Select mode", the transform toolbar would dissapear completly

diff --git a/Engine/Plugins/ChaosVD/Source/ChaosVD/Private/ChaosVDPlaybackViewportClient.cpp b/Engine/Plugins/ChaosVD/Source/ChaosVD/Private/ChaosVDPlaybackViewportClient.cpp
--- a/Engine/Plugins/ChaosVD/Source/ChaosVD/Private/ChaosVDPlaybackViewportClient.cpp
+++ b/Engine/Plugins/ChaosVD/Source/ChaosVD/Private/ChaosVDPlaybackViewportClient.cpp
@@ -26,7 +26,7 @@
 #include "Visualizers/ChaosVDDebugDrawUtils.h"
 #include "Widgets/SChaosVDMainTab.h"
 
-FChaosVDPlaybackViewportClient::FChaosVDPlaybackViewportClient(const TSharedPtr<FEditorModeTools>& InModeTools) : FEditorViewportClient(InModeTools.Get()), CVDWorld(nullptr)
+FChaosVDPlaybackViewportClient::FChaosVDPlaybackViewportClient(const TSharedPtr<FEditorModeTools>& InModeTools, const TSharedPtr<SEditorViewport>& InEditorViewportWidget) : FEditorViewportClient(InModeTools.Get(), nullptr, InEditorViewportWidget), CVDWorld(nullptr)
 {
 	Widget->SetUsesEditorModeTools(InModeTools.Get());
 
diff --git a/Engine/Plugins/ChaosVD/Source/ChaosVD/Private/ChaosVDPlaybackViewportClient.h b/Engine/Plugins/ChaosVD/Source/ChaosVD/Private/ChaosVDPlaybackViewportClient.h
--- a/Engine/Plugins/ChaosVD/Source/ChaosVD/Private/ChaosVDPlaybackViewportClient.h
+++ b/Engine/Plugins/ChaosVD/Source/ChaosVD/Private/ChaosVDPlaybackViewportClient.h
@@ -15,7 +15,7 @@ class FChaosVDPlaybackViewportClient : public FEditorViewportClient
 {
 public:
 
-	FChaosVDPlaybackViewportClient(const TSharedPtr<FEditorModeTools>& InModeTools);
+	FChaosVDPlaybackViewportClient(const TSharedPtr<FEditorModeTools>& InModeTools, const TSharedPtr<SEditorViewport>& InEditorViewportWidget);
 	virtual ~FChaosVDPlaybackViewportClient() override;
 
 	virtual void ProcessClick(FSceneView& View, HHitProxy* HitProxy, FKey Key, EInputEvent Event, uint32 HitX, uint32 HitY) override;
diff --git a/Engine/Plugins/ChaosVD/Source/ChaosVD/Private/Widgets/SChaosVDPlaybackViewport.cpp b/Engine/Plugins/ChaosVD/Source/ChaosVD/Private/Widgets/SChaosVDPlaybackViewport.cpp
--- a/Engine/Plugins/ChaosVD/Source/ChaosVD/Private/Widgets/SChaosVDPlaybackViewport.cpp
+++ b/Engine/Plugins/ChaosVD/Source/ChaosVD/Private/Widgets/SChaosVDPlaybackViewport.cpp
@@ -105,9 +105,16 @@ TSharedPtr<FExtender> SChaosVDPlaybackViewport::GetExtenders() const
 	return Extender;
 }
 
+EVisibility SChaosVDPlaybackViewport::GetTransformToolbarVisibility() const
+{
+	// We want to always show the transform tool bar. We disable each action that is not supported for a selected actor individually.
+	// Without doing this, if you select an unsupported mode, the entire toolbar disappears
+	return EVisibility::Visible;
+}
+
 TSharedRef<FEditorViewportClient> SChaosVDPlaybackViewport::MakeEditorViewportClient()
 {
-	TSharedPtr<FChaosVDPlaybackViewportClient> NewViewport = MakeShared<FChaosVDPlaybackViewportClient>(EditorModeTools);
+	TSharedPtr<FChaosVDPlaybackViewportClient> NewViewport = MakeShared<FChaosVDPlaybackViewportClient>(EditorModeTools, GetViewportWidget());
 
 	NewViewport->SetAllowCinematicControl(false);
 	
diff --git a/Engine/Plugins/ChaosVD/Source/ChaosVD/Private/Widgets/SChaosVDPlaybackViewport.h b/Engine/Plugins/ChaosVD/Source/ChaosVD/Private/Widgets/SChaosVDPlaybackViewport.h
--- a/Engine/Plugins/ChaosVD/Source/ChaosVD/Private/Widgets/SChaosVDPlaybackViewport.h
+++ b/Engine/Plugins/ChaosVD/Source/ChaosVD/Private/Widgets/SChaosVDPlaybackViewport.h
@@ -40,6 +40,8 @@ public:
 	virtual void OnFloatingButtonClicked() override {};
 	// END ICommonEditorViewportToolbarInfoProvider interface
 
+	virtual EVisibility GetTransformToolbarVisibility() const override;
+
 protected:
 
 	virtual TSharedRef<FEditorViewportClient> MakeEditorViewportClient() override;
