[ChaosVD] Removed three object tracking methods that had issues. They will be properly re-introduced in UE 5.5

Added an UI command to toggle object tracking using the F8 key in addition to the menu option in the show menu

#tests Tested that CVD can follow selected objects as expected, either by using the option in the menu or by using the key shortcut, either by using the option in the menu or by using the key shortcut

diff --git a/Engine/Plugins/ChaosVD/Source/ChaosVD/Private/ChaosVDCommands.cpp b/Engine/Plugins/ChaosVD/Source/ChaosVD/Private/ChaosVDCommands.cpp
new file mode 100644
--- /dev/null
+++ b/Engine/Plugins/ChaosVD/Source/ChaosVD/Private/ChaosVDCommands.cpp
@@ -0,0 +1,18 @@
+// Copyright Epic Games, Inc. All Rights Reserved.
+
+#include "ChaosVDCommands.h"
+
+#include "ChaosVDStyle.h"
+
+#define LOCTEXT_NAMESPACE "ChaosVisualDebugger"
+
+FChaosVDCommands::FChaosVDCommands() : TCommands<FChaosVDCommands>(TEXT("ActorPalette"), NSLOCTEXT("ChaosVisualDebugger", "ChaosVisualDebuggerEditor", "Chaos Visual Debugger Editor"), NAME_None, FChaosVDStyle::GetStyleSetName())
+{
+}
+
+void FChaosVDCommands::RegisterCommands()
+{
+	UI_COMMAND(TrackUntrackSelectedObject, "Start or Stop tracking the selected object", "Start or Stop tracking the selected object", EUserInterfaceActionType::Button, FInputChord(EKeys::F8));
+}
+
+#undef LOCTEXT_NAMESPACE
diff --git a/Engine/Plugins/ChaosVD/Source/ChaosVD/Private/ChaosVDCommands.h b/Engine/Plugins/ChaosVD/Source/ChaosVD/Private/ChaosVDCommands.h
new file mode 100644
--- /dev/null
+++ b/Engine/Plugins/ChaosVD/Source/ChaosVD/Private/ChaosVDCommands.h
@@ -0,0 +1,19 @@
+// Copyright Epic Games, Inc. All Rights Reserved.
+
+#pragma once
+
+#include "Framework/Commands/Commands.h"
+#include "Templates/SharedPointer.h"
+
+class FChaosVDCommands : public TCommands<FChaosVDCommands>
+{
+public:
+
+	FChaosVDCommands();
+
+	// TCommands<> interface
+	virtual void RegisterCommands() override;
+
+public:
+	TSharedPtr<FUICommandInfo> TrackUntrackSelectedObject;
+};
\ No newline at end of file
diff --git a/Engine/Plugins/ChaosVD/Source/ChaosVD/Private/ChaosVDEditorSettings.h b/Engine/Plugins/ChaosVD/Source/ChaosVD/Private/ChaosVDEditorSettings.h
--- a/Engine/Plugins/ChaosVD/Source/ChaosVD/Private/ChaosVDEditorSettings.h
+++ b/Engine/Plugins/ChaosVD/Source/ChaosVD/Private/ChaosVDEditorSettings.h
@@ -19,17 +19,6 @@ class UMaterial;
 
 DECLARE_MULTICAST_DELEGATE_OneParam(FChaosVDSettingChaged, UChaosVDEditorSettings* CVDEditorSettingsObject)
 
-UENUM()
-enum class EChaosVDActorTrackingMode
-{
-	/** Follow the selected target keeping the specified distance from it. */
-	ByDistanceOffset,
-	/** Follow the selected target keeping its bounding box inside of CVD's camera view. */
-	ByBoundingBox,
-	/** Snap CVD's camera to the transform of the current selected object */
-	MatchTransform 
-};
-
 UENUM()
 enum class EChaosVDActorTrackingTarget
 {
@@ -37,10 +26,6 @@ enum class EChaosVDActorTrackingTarget
 	Disabled,
 	/** Follow the current selected object */
 	SelectedObject,
-	/** Follow the selected recorded transform ID */
-	RecordedTransform,
-	/** Follow the selected recorded location ID */
-	RecordedLocation
 };
 
 UENUM(meta = (Bitflags, UseEnumValuesAsMaskValuesInEditor = "true"))
@@ -363,16 +348,8 @@ public:
 	UPROPERTY(EditAnywhere, Category = "Viewport Tracking")
 	EChaosVDActorTrackingTarget TrackingTarget;
 
-	/** Sets how tracking should be performed by CVD's Camera */
-	UPROPERTY(EditAnywhere, Category = "Viewport Tracking", meta=(EditCondition = "TrackingTarget != EChaosVDActorTrackingTarget::Disabled", EditConditionHides))
-	EChaosVDActorTrackingMode TrackingOptions;
-
-	/** Distance from the selected tracking target CVD's camera should attempt to be at */
-	UPROPERTY(EditAnywhere, Category = "Viewport Tracking", meta=(EditCondition = "TrackingOptions == EChaosVDActorTrackingMode::ByDistanceOffset && TrackingTarget != EChaosVDActorTrackingTarget::Disabled", EditConditionHides))
-	float TrackingDistanceOffset = 1500.0f;
-
 	/** By how much we should expand the bounding box used to track a target by bounding box. Used to see more of the screen while tracking in Bounding Box mode */
-	UPROPERTY(EditAnywhere, Category = "Viewport Tracking", meta=(EditCondition = "TrackingOptions == EChaosVDActorTrackingMode::ByBoundingBox && TrackingTarget != EChaosVDActorTrackingTarget::Disabled", EditConditionHides))
+	UPROPERTY(EditAnywhere, Category = "Viewport Tracking", meta=(EditCondition = "TrackingTarget != EChaosVDActorTrackingTarget::Disabled"))
 	float ExpandViewTrackingBy = 60.0f;
 
 	/** Set of flags to enable/disable visibility of specific types of geometry/particles */
@@ -403,9 +380,6 @@ public:
 	FChaosVDSettingChaged& OnFarClippingOverrideChanged() { return FarClippingOverrideChangedDelegate; }
 
 	FChaosVDSettingChaged& OnPlaybackSettingsChanged() { return PlaybackSettingsChangedDelegate ; }
-
-	TSharedPtr<FName> SelectedTrackedTransformName;
-	TSharedPtr<FName> SelectedTrackedLocationName;
 	
 	virtual void PostEditUndo() override;
 
diff --git a/Engine/Plugins/ChaosVD/Source/ChaosVD/Private/ChaosVDModule.cpp b/Engine/Plugins/ChaosVD/Source/ChaosVD/Private/ChaosVDModule.cpp
--- a/Engine/Plugins/ChaosVD/Source/ChaosVD/Private/ChaosVDModule.cpp
+++ b/Engine/Plugins/ChaosVD/Source/ChaosVD/Private/ChaosVDModule.cpp
@@ -1,6 +1,8 @@
 // Copyright Epic Games, Inc. All Rights Reserved.
 
 #include "ChaosVDModule.h"
+
+#include "ChaosVDCommands.h"
 #include "ChaosVDStyle.h"
 #include "ChaosVDParticleActorCustomization.h"
 #include "ChaosVDTabsIDs.h"
@@ -37,6 +39,8 @@ FChaosVDModule& FChaosVDModule::Get()
 void FChaosVDModule::StartupModule()
 {	
 	FChaosVDStyle::Initialize();
+	
+	FChaosVDCommands::Register();
 
 	RegisterClassesCustomDetails();
 
@@ -54,6 +58,8 @@ void FChaosVDModule::StartupModule()
 void FChaosVDModule::ShutdownModule()
 {
 	FChaosVDStyle::Shutdown();
+	
+	FChaosVDCommands::Unregister();
 
 	FGlobalTabmanager::Get()->UnregisterNomadTabSpawner(FChaosVDTabID::ChaosVisualDebuggerTab);
 
diff --git a/Engine/Plugins/ChaosVD/Source/ChaosVD/Private/ChaosVDPlaybackViewportClient.cpp b/Engine/Plugins/ChaosVD/Source/ChaosVD/Private/ChaosVDPlaybackViewportClient.cpp
--- a/Engine/Plugins/ChaosVD/Source/ChaosVD/Private/ChaosVDPlaybackViewportClient.cpp
+++ b/Engine/Plugins/ChaosVD/Source/ChaosVD/Private/ChaosVDPlaybackViewportClient.cpp
@@ -3,6 +3,7 @@
 #include "ChaosVDPlaybackViewportClient.h"
 
 #include "ChaosVDEditorSettings.h"
+#include "ChaosVDEngine.h"
 #include "ChaosVDModule.h"
 #include "ChaosVDParticleActor.h"
 #include "ChaosVDPlaybackController.h"
@@ -191,146 +192,23 @@ void FChaosVDPlaybackViewportClient::HandleViewportSettingsChanged(UChaosVDEdito
 	}
 }
 
-void FChaosVDPlaybackViewportClient::TrackActor(AActor* ActorToTrack, EChaosVDActorTrackingMode TrackingMode)
-{
-	if (!ActorToTrack)
-	{
-		return;
-	}
-	
-	if (const UChaosVDEditorSettings* CVDEditorSettings = GetDefault<UChaosVDEditorSettings>())
-	{
-		switch(TrackingMode)
-		{
-			case EChaosVDActorTrackingMode::ByBoundingBox:
-				{
-					FBox ActorBounds = ActorToTrack->GetComponentsBoundingBox(false);
-					FocusViewportOnBox(ActorBounds.ExpandBy(CVDEditorSettings->ExpandViewTrackingBy), true);
-					break;
-				}
-			case EChaosVDActorTrackingMode::ByDistanceOffset:
-			case EChaosVDActorTrackingMode::MatchTransform:
-				{
-					TrackTransform(ActorToTrack->GetActorTransform(), TrackingMode);
-					break;
-				}
-			default:
-				{
-					ensureMsgf(false, TEXT("Actor tracking requested with invalid options. The actor will not be tracked"));
-					break;
-				}
-		}	
-	}
-}
-
-void FChaosVDPlaybackViewportClient::TrackTransform(const FTransform& TransformToTrack, EChaosVDActorTrackingMode TrackingMode)
-{
-	if (const UChaosVDEditorSettings* CVDEditorSettings = GetDefault<UChaosVDEditorSettings>())
-	{
-		switch(TrackingMode)
-		{
-		case EChaosVDActorTrackingMode::ByBoundingBox:
-			{
-				ensureMsgf(false, TEXT("Tracking by Bounding box only supported with Actors"));
-				break;
-			}
-		case EChaosVDActorTrackingMode::ByDistanceOffset:
-			{
-				FViewportCameraTransform& ViewTransform = GetViewTransform();
-				FVector ActorLocation = TransformToTrack.GetLocation();
-					
-				constexpr bool bEnable = false;
-				ToggleOrbitCamera(bEnable);
-
-				FVector ActorToCamDir = ViewTransform.GetLocation() - ActorLocation;
-				ActorToCamDir.Normalize();
-
-				FVector TargetLocation = ActorLocation + ActorToCamDir * CVDEditorSettings->TrackingDistanceOffset;
-
-				ViewTransform.SetRotation((-ActorToCamDir).Rotation());
-				ViewTransform.SetLookAt(ActorLocation);
-				ViewTransform.TransitionToLocation(TargetLocation, EditorViewportWidget, true);
-					
-				// Tell the viewport to redraw itself.
-				Invalidate();
-				break;
-			}
-		case EChaosVDActorTrackingMode::MatchTransform:
-			{
-				constexpr bool bEnable = false;
-				ToggleOrbitCamera(bEnable);
-					
-				FViewportCameraTransform& ViewTransform = GetViewTransform();
-				ViewTransform.SetRotation(TransformToTrack.GetRotation().Rotator());
-
-				ViewTransform.SetLookAt(TransformToTrack.GetLocation());
-				ViewTransform.TransitionToLocation(TransformToTrack.GetLocation(), EditorViewportWidget, true);
-
-				Invalidate();
-				break;
-			}		
-		default:
-			{
-				ensureMsgf(false, TEXT("Actor tracking requested with invalid options. The actor will not be tracked"));
-				break;
-			}
-		}	
-	}
-}
-
-void FChaosVDPlaybackViewportClient::PerformSelectedTrackingForFrame(FChaosVDGameFrameData* FrameData)
+void FChaosVDPlaybackViewportClient::TrackSelectedObject()
 {
 	if (const TSharedPtr<FChaosVDScene> CVDSceneSharedPtr = CVDScene.Pin())
 	{
 		if (const UChaosVDEditorSettings* CVDEditorSettings = GetDefault<UChaosVDEditorSettings>())
 		{
-			switch (CVDEditorSettings->TrackingTarget)
+			if (ModeTools.IsValid() && CVDEditorSettings->TrackingTarget == EChaosVDActorTrackingTarget::SelectedObject)
 			{
-			case EChaosVDActorTrackingTarget::SelectedObject:
-				{
-					if (ModeTools.IsValid())
-					{
-						USelection* CurrentSelection = ModeTools->GetSelectedActors();
-		
-						//TODO: Update this if we add multi selection support
-						if (AActor* SelectedActor = CurrentSelection ? CurrentSelection->GetTop<AActor>() : nullptr)
-						{
-							TrackActor(SelectedActor, CVDEditorSettings->TrackingOptions);
-						}
-					}
-					break;
-				}
-			case EChaosVDActorTrackingTarget::RecordedTransform:
-				{
-					// TODO: Find a better place to store the current selected Transform name, it should not be the editor settings object
-					if (const TSharedPtr<FName>& TransformName = CVDEditorSettings->SelectedTrackedTransformName)
-					{
-						if (const FChaosVDTrackedTransform* TrackedTransform = FrameData->RecordedNonSolverTransformsByID.Find(*TransformName))
-						{
-							TrackTransform(TrackedTransform->Transform, CVDEditorSettings->TrackingOptions);
-						}
-					}
-								
-					break;
-				}
-			case EChaosVDActorTrackingTarget::RecordedLocation:
+				USelection* CurrentSelection = ModeTools->GetSelectedActors();
+
+				//TODO: Update this if we add multi selection support
+				if (const AActor* SelectedActor = CurrentSelection ? CurrentSelection->GetTop<AActor>() : nullptr)
 				{
-					// TODO: Find a better place to store the current selected Location name, it should not be the editor settings object
-					if (const TSharedPtr<FName>& LocationName = CVDEditorSettings->SelectedTrackedLocationName)
-					{
-						if (const FChaosVDTrackedLocation* TrackedTransform = FrameData->RecordedNonSolverLocationsByID.Find(*LocationName))
-						{
-							FTransform LocationTransform;
-							LocationTransform.SetLocation(TrackedTransform->Location);
-							TrackTransform(LocationTransform, CVDEditorSettings->TrackingOptions);
-						}
-					}
-								
-					break;
+					const FBox ActorBounds = SelectedActor->GetComponentsBoundingBox(false);
+					FocusViewportOnBox(ActorBounds.ExpandBy(CVDEditorSettings->ExpandViewTrackingBy), true);		
 				}
-			default:
-				break;
-			}			
+			}
 		}
 	}
 }
@@ -346,6 +224,15 @@ bool FChaosVDPlaybackViewportClient::InputKey(const FInputKeyEventArgs& EventArg
 	return FEditorViewportClient::InputKey(EventArgs);
 }
 
+void FChaosVDPlaybackViewportClient::ToggleObjectTrackingIfSelected()
+{
+	// Currently we only have two options, so toggle between them
+	if (UChaosVDEditorSettings* CVDEditorSettings = GetMutableDefault<UChaosVDEditorSettings>())
+	{
+		CVDEditorSettings->TrackingTarget = CVDEditorSettings->TrackingTarget == EChaosVDActorTrackingTarget::Disabled ? EChaosVDActorTrackingTarget::SelectedObject : EChaosVDActorTrackingTarget::Disabled;
+	}
+}
+
 void FChaosVDPlaybackViewportClient::Draw(const FSceneView* View, FPrimitiveDrawInterface* PDI)
 {
 	if (View)
@@ -356,6 +243,8 @@ void FChaosVDPlaybackViewportClient::Draw(const FSceneView* View, FPrimitiveDraw
 		const_cast<FSceneView*>(View)->bAllowTranslucentPrimitivesInHitProxy = true;
 	}
 
+	TrackSelectedObject();
+
 	const TSharedPtr<SChaosVDMainTab> MainTabToolkitHost = ModeTools.IsValid() ? StaticCastSharedPtr<SChaosVDMainTab>(ModeTools->GetToolkitHost()) : nullptr;
 	if (!MainTabToolkitHost.IsValid())
 	{
diff --git a/Engine/Plugins/ChaosVD/Source/ChaosVD/Private/ChaosVDPlaybackViewportClient.h b/Engine/Plugins/ChaosVD/Source/ChaosVD/Private/ChaosVDPlaybackViewportClient.h
--- a/Engine/Plugins/ChaosVD/Source/ChaosVD/Private/ChaosVDPlaybackViewportClient.h
+++ b/Engine/Plugins/ChaosVD/Source/ChaosVD/Private/ChaosVDPlaybackViewportClient.h
@@ -26,19 +26,19 @@ public:
 
 	virtual void Draw(const FSceneView* View, FPrimitiveDrawInterface* PDI) override;
 	virtual void DrawCanvas(FViewport& InViewport, FSceneView& View, FCanvas& Canvas) override;
-	
-	void TrackActor(AActor* ActorToTrack, EChaosVDActorTrackingMode TrackingMode);
-	void TrackTransform(const FTransform& TransformToTrack, EChaosVDActorTrackingMode TrackingMode);
-	
-	void PerformSelectedTrackingForFrame(FChaosVDGameFrameData* FrameData);
 
 	virtual bool InputKey(const FInputKeyEventArgs& EventArgs) override;
 
+	void ToggleObjectTrackingIfSelected();
+
 private:
+
+	void TrackSelectedObject();
+
 	void HandleObjectFocused(UObject* FocusedObject);
 	void HandleActorMoving(AActor* MovedActor) const;
 	void HandleViewportSettingsChanged(UChaosVDEditorSettings* SettingsObject);
-	
+
 	FDelegateHandle ObjectFocusedDelegateHandle;
 	UWorld* CVDWorld;
 	TWeakPtr<FChaosVDScene> CVDScene;
diff --git a/Engine/Plugins/ChaosVD/Source/ChaosVD/Private/Widgets/SChaosVDPlaybackViewport.cpp b/Engine/Plugins/ChaosVD/Source/ChaosVD/Private/Widgets/SChaosVDPlaybackViewport.cpp
--- a/Engine/Plugins/ChaosVD/Source/ChaosVD/Private/Widgets/SChaosVDPlaybackViewport.cpp
+++ b/Engine/Plugins/ChaosVD/Source/ChaosVD/Private/Widgets/SChaosVDPlaybackViewport.cpp
@@ -2,6 +2,7 @@
 
 #include "Widgets/SChaosVDPlaybackViewport.h"
 
+#include "ChaosVDCommands.h"
 #include "ChaosVDEditorMode.h"
 #include "ChaosVDEditorModeTools.h"
 #include "ChaosVDEditorSettings.h"
@@ -105,6 +106,21 @@ TSharedPtr<FExtender> SChaosVDPlaybackViewport::GetExtenders() const
 	return Extender;
 }
 
+void SChaosVDPlaybackViewport::BindCommands()
+{
+	SEditorViewport::BindCommands();
+
+	const FChaosVDCommands& Commands = FChaosVDCommands::Get();
+
+	if (ensure(Client))
+	{
+		const TSharedRef<FChaosVDPlaybackViewportClient> ViewportClientRef = StaticCastSharedRef<FChaosVDPlaybackViewportClient>(Client.ToSharedRef());
+		CommandList->MapAction(
+			Commands.TrackUntrackSelectedObject,
+			FExecuteAction::CreateSP(ViewportClientRef, &FChaosVDPlaybackViewportClient::ToggleObjectTrackingIfSelected));
+	}
+}
+
 EVisibility SChaosVDPlaybackViewport::GetTransformToolbarVisibility() const
 {
 	// We want to always show the transform tool bar. We disable each action that is not supported for a selected actor individually.
@@ -217,14 +233,6 @@ void SChaosVDPlaybackViewport::HandleControllerTrackFrameUpdated(TWeakPtr<FChaos
 			}
 
 			GameFramesTimelineWidget->SetTargetFrameTime(ControllerSharedPtr->GetFrameTimeForTrack(EChaosVDTrackType::Game, FChaosVDPlaybackController::GameTrackID, *GameTrackInfo));
-
-			if (const TSharedPtr<FChaosVDRecording> RecordingData = ControllerSharedPtr->GetCurrentRecording().Pin())
-			{
-				if (FChaosVDGameFrameData* FrameData = RecordingData->GetGameFrameData_AssumesLocked(GameTrackFrame))
-				{
-					PlaybackViewportClient->PerformSelectedTrackingForFrame(FrameData);
-				}
-			}
 		}
 	}
 }
@@ -236,7 +244,7 @@ void SChaosVDPlaybackViewport::HandlePostSelectionChange(const UTypedElementSele
 
 void SChaosVDPlaybackViewport::OnPlaybackSceneUpdated()
 {
-	PlaybackViewportClient->bNeedsRedraw = true;	
+	PlaybackViewportClient->bNeedsRedraw = true;
 }
 
 void SChaosVDPlaybackViewport::OnSolverVisibilityUpdated(int32 SolverID, bool bNewVisibility)
diff --git a/Engine/Plugins/ChaosVD/Source/ChaosVD/Private/Widgets/SChaosVDPlaybackViewport.h b/Engine/Plugins/ChaosVD/Source/ChaosVD/Private/Widgets/SChaosVDPlaybackViewport.h
--- a/Engine/Plugins/ChaosVD/Source/ChaosVD/Private/Widgets/SChaosVDPlaybackViewport.h
+++ b/Engine/Plugins/ChaosVD/Source/ChaosVD/Private/Widgets/SChaosVDPlaybackViewport.h
@@ -40,6 +40,8 @@ public:
 	virtual void OnFloatingButtonClicked() override {};
 	// END ICommonEditorViewportToolbarInfoProvider interface
 
+	virtual void BindCommands() override;
+
 	virtual EVisibility GetTransformToolbarVisibility() const override;
 
 protected:
diff --git a/Engine/Plugins/ChaosVD/Source/ChaosVD/Private/Widgets/SChaosVDVisualizationControls.cpp b/Engine/Plugins/ChaosVD/Source/ChaosVD/Private/Widgets/SChaosVDVisualizationControls.cpp
--- a/Engine/Plugins/ChaosVD/Source/ChaosVD/Private/Widgets/SChaosVDVisualizationControls.cpp
+++ b/Engine/Plugins/ChaosVD/Source/ChaosVD/Private/Widgets/SChaosVDVisualizationControls.cpp
@@ -7,7 +7,6 @@
 #include "ChaosVDPlaybackController.h"
 #include "ChaosVDPlaybackViewportClient.h"
 #include "ChaosVDRecording.h"
-#include "ChaosVDScene.h"
 #include "Elements/Framework/TypedElementSelectionSet.h"
 #include "PropertyEditorModule.h"
 #include "Modules/ModuleManager.h"
@@ -15,7 +14,8 @@
 #include "Widgets/Layout/SExpandableArea.h"
 #include "Widgets/SChaosVDNameListPicker.h"
 #include "Widgets/SChaosVDPlaybackViewport.h"
-#include "Widgets/Text/STextBlock.h"
+
+#define LOCTEXT_NAMESPACE "ChaosVisualDebugger"
 
 void SChaosVDVisualizationControls::Construct(const FArguments& InArgs, const TWeakPtr<FChaosVDPlaybackController>& InPlaybackController,const TWeakPtr<SChaosVDPlaybackViewport>& InPlaybackViewport)
 {
@@ -28,7 +28,7 @@ void SChaosVDVisualizationControls::Construct(const FArguments& InArgs, const TW
 	DetailsViewArgs.bAllowSearch = false;
 	DetailsViewArgs.NameAreaSettings = FDetailsViewArgs::HideNameArea;
 
-	TSharedRef<IDetailsView> DetailsPanel = PropertyEditorModule.CreateDetailView(DetailsViewArgs);
+	const TSharedRef<IDetailsView> DetailsPanel = PropertyEditorModule.CreateDetailView(DetailsViewArgs);
 
 	DetailsPanel->SetObject(GetMutableDefault<UChaosVDEditorSettings>());
 
@@ -41,148 +41,13 @@ void SChaosVDVisualizationControls::Construct(const FArguments& InArgs, const TW
 		[
 			DetailsPanel
 		]
-		+SVerticalBox::Slot()
-		.AutoHeight()
-		.Padding(2,5,2,5)
-		[
-			SNew(SExpandableArea)
-				.InitiallyCollapsed(false)
-				.BorderBackgroundColor(FLinearColor::White)
-				.Padding(FMargin(8.f))
-				.Visibility_Raw(this, &SChaosVDVisualizationControls::GetTrackedDataOptionsVisibility)
-				.HeaderContent()
-				[
-					SNew(SHorizontalBox)
-					+SHorizontalBox::Slot()
-					.HAlign(HAlign_Left)
-					.VAlign(VAlign_Center)
-					.AutoWidth()
-					.Padding(0.f, 0.f, 0.f, 0.f)
-					[
-						SNew(STextBlock)
-						.Text(FText::FromString("Trackable Data"))
-						.Font(FCoreStyle::Get().GetFontStyle("ExpandableArea.TitleFont"))
-					]
-				]
-				.BodyContent()
-				[
-					SNew(SVerticalBox)
-					+SVerticalBox::Slot()
-					[
-						SNew(SHorizontalBox)
-						.Visibility_Raw(this, &SChaosVDVisualizationControls::GetTransformsPickerVisibility)
-						+SHorizontalBox::Slot()
-						.VAlign(VAlign_Center)
-						.HAlign(HAlign_Left)
-						[
-							SNew(STextBlock)
-							.Text(FText::FromString("Transform to Track"))
-							.Font(FCoreStyle::Get().GetFontStyle("ExpandableArea.TitleFont"))
-						]
-						+SHorizontalBox::Slot()
-						.VAlign(VAlign_Center)
-						.HAlign(HAlign_Right)
-						[
-							SAssignNew(TransformNamePickerWidget, SChaosVDNameListPicker)
-							.OnNameSleceted_Raw(this, &SChaosVDVisualizationControls::HandleTransformNameSelected)
-						]	
-					]
-					+SVerticalBox::Slot()
-					.Padding(2,5,2,5)
-					[
-						SNew(SHorizontalBox)
-						.Visibility_Raw(this, &SChaosVDVisualizationControls::GetLocationsPickerVisibility)
-						+SHorizontalBox::Slot()
-						.VAlign(VAlign_Center)
-						.HAlign(HAlign_Left)
-						[
-							SNew(STextBlock)
-							.Text(FText::FromString("Location to Track"))
-							.Font(FCoreStyle::Get().GetFontStyle("ExpandableArea.TitleFont"))
-						]
-						+SHorizontalBox::Slot()
-						.VAlign(VAlign_Center)
-						.HAlign(HAlign_Right)
-						[
-							SAssignNew(LocationsNamePickerWidget, SChaosVDNameListPicker)
-							.OnNameSleceted_Raw(this, &SChaosVDVisualizationControls::HandleLocationNameSelected)
-						]
-					]
-				]
-		]
 	];
 }
 
-void SChaosVDVisualizationControls::HandleControllerTrackFrameUpdated(TWeakPtr<FChaosVDPlaybackController> InController, const FChaosVDTrackInfo* UpdatedTrackInfo, FGuid InstigatorGuid)
-{
-	if (UpdatedTrackInfo == nullptr)
-	{
-		return;
-	}
-
-	if (UpdatedTrackInfo->TrackType != EChaosVDTrackType::Game)
-	{
-		return;
-	}
-
-	const TSharedPtr<SChaosVDPlaybackViewport> PlaybackViewportSharedPtr = PlaybackViewport.Pin();
-	const TSharedPtr<FChaosVDPlaybackController> ControllerSharedPtr = InController.Pin();
-		
-	if (!PlaybackViewportSharedPtr.IsValid() || !ControllerSharedPtr.IsValid())
-	{
-		return;
-	}
-
-	if (const TSharedPtr<FChaosVDRecording> RecordingSharedPtr  = ControllerSharedPtr->GetCurrentRecording().Pin())
-	{
-		if (FChaosVDGameFrameData* FrameData = RecordingSharedPtr->GetGameFrameData_AssumesLocked(UpdatedTrackInfo->CurrentFrame))
-		{
-			UpdateNameListFromKeys(FrameData->RecordedNonSolverTransformsByID, *TransformNamePickerWidget.Get());
-			UpdateNameListFromKeys(FrameData->RecordedNonSolverLocationsByID, *LocationsNamePickerWidget.Get());
-		}
-	}
-}
-
-void SChaosVDVisualizationControls::HandleLocationNameSelected(TSharedPtr<FName> SelectedName)
-{
-	if (UChaosVDEditorSettings* Settings = GetMutableDefault<UChaosVDEditorSettings>())
-	{
-		Settings->SelectedTrackedLocationName = SelectedName;
-	}
-}
 
-void SChaosVDVisualizationControls::HandleTransformNameSelected(TSharedPtr<FName> SelectedName)
+void SChaosVDVisualizationControls::RegisterCommandsUI()
 {
-	if (UChaosVDEditorSettings* Settings = GetMutableDefault<UChaosVDEditorSettings>())
-	{
-		Settings->SelectedTrackedTransformName = SelectedName;
-	}
-}
-
-EVisibility SChaosVDVisualizationControls::GetLocationsPickerVisibility() const
-{
-	if (const UChaosVDEditorSettings* Settings = GetDefault<UChaosVDEditorSettings>())
-	{
-		return Settings->TrackingTarget == EChaosVDActorTrackingTarget::RecordedLocation ? EVisibility::Visible : EVisibility::Collapsed;
-	}
-
-	return EVisibility::Collapsed;
-}
-
-EVisibility SChaosVDVisualizationControls::GetTransformsPickerVisibility() const
-{
-	if (const UChaosVDEditorSettings* Settings = GetDefault<UChaosVDEditorSettings>())
-	{
-		return Settings->TrackingTarget == EChaosVDActorTrackingTarget::RecordedTransform ? EVisibility::Visible : EVisibility::Collapsed;
-	}
-
-	return EVisibility::Collapsed;
+	
 }
 
-
-EVisibility SChaosVDVisualizationControls::GetTrackedDataOptionsVisibility() const
-{
-	const bool bIsTrackedDataVisible = GetLocationsPickerVisibility() == EVisibility::Visible ||  GetTransformsPickerVisibility() == EVisibility::Visible;
-
-	return bIsTrackedDataVisible ? EVisibility::Visible : EVisibility::Collapsed;
-}
+#undef LOCTEXT_NAMESPACE
diff --git a/Engine/Plugins/ChaosVD/Source/ChaosVD/Private/Widgets/SChaosVDVisualizationControls.h b/Engine/Plugins/ChaosVD/Source/ChaosVD/Private/Widgets/SChaosVDVisualizationControls.h
--- a/Engine/Plugins/ChaosVD/Source/ChaosVD/Private/Widgets/SChaosVDVisualizationControls.h
+++ b/Engine/Plugins/ChaosVD/Source/ChaosVD/Private/Widgets/SChaosVDVisualizationControls.h
@@ -5,6 +5,7 @@
 #include "Widgets/SChaosVDNameListPicker.h"
 #include "Templates/SharedPointer.h"
 
+class FUICommandInfo;
 class SChaosVDPlaybackViewport;
 class FChaosVDPlaybackViewportClient;
 struct FChaosVDTrackInfo;
@@ -21,37 +22,10 @@ public:
 
 	void Construct(const FArguments& InArgs, const TWeakPtr<FChaosVDPlaybackController>& InPlaybackController, const TWeakPtr<SChaosVDPlaybackViewport>& InPlaybackViewport);
 
-	virtual void HandleControllerTrackFrameUpdated(TWeakPtr<FChaosVDPlaybackController> InController, const FChaosVDTrackInfo* UpdatedTrackInfo, FGuid InstigatorGuid) override;
-
 protected:
 
-	void HandleLocationNameSelected(TSharedPtr<FName> SelectedName);
-	void HandleTransformNameSelected(TSharedPtr<FName> SelectedName);
-
-	template<typename ValueType>
-	void UpdateNameListFromKeys(const TMap<FName, ValueType>& MapToEvaluate, SChaosVDNameListPicker& NameList);
-	
-	EVisibility GetLocationsPickerVisibility() const;
-	EVisibility GetTransformsPickerVisibility() const;
-	EVisibility GetTrackedDataOptionsVisibility() const;
-
-	TSharedPtr<SChaosVDNameListPicker> TransformNamePickerWidget;
-	
-	TSharedPtr<SChaosVDNameListPicker> LocationsNamePickerWidget;
+	void RegisterCommandsUI();
 
 	TWeakPtr<SChaosVDPlaybackViewport> PlaybackViewport;
 
-	TSharedPtr<FName> SelectedTrackedTransformName;
-	TSharedPtr<FName> SelectedTrackedLocationName;
-};
-
-template <typename ValueType>
-void SChaosVDVisualizationControls::UpdateNameListFromKeys(const TMap<FName, ValueType>& MapToEvaluate, SChaosVDNameListPicker& NameList)
-{
-	TArray<TSharedPtr<FName>> NewNameList;
-	NewNameList.Reserve(MapToEvaluate.Num());
-
-	Algo::Transform(MapToEvaluate, NewNameList, [](const TPair<FName, ValueType>& TransformData){ return MakeShared<FName>(TransformData.Key);});
-	
-	NameList.UpdateNameList(MoveTemp(NewNameList));
-}
+};
\ No newline at end of file
diff --git a/Engine/Source/Runtime/Experimental/Chaos/Public/ChaosVisualDebugger/ChaosVDTraceMacros.h b/Engine/Source/Runtime/Experimental/Chaos/Public/ChaosVisualDebugger/ChaosVDTraceMacros.h
--- a/Engine/Source/Runtime/Experimental/Chaos/Public/ChaosVisualDebugger/ChaosVDTraceMacros.h
+++ b/Engine/Source/Runtime/Experimental/Chaos/Public/ChaosVisualDebugger/ChaosVDTraceMacros.h
@@ -125,16 +125,6 @@
 		FChaosVisualDebuggerTrace::TraceMidPhasesFromCollisionConstraints(CollisionConstraints);
 	#endif
 
-	#ifndef CVD_TRACE_NON_SOLVER_LOCATION
-		#define CVD_TRACE_NON_SOLVER_LOCATION(Location, DebugName) \
-		FChaosVisualDebuggerTrace::TraceNonSolverLocation(Location, DebugName);
-	#endif
-
-	#ifndef CVD_TRACE_NON_SOLVER_TRANSFORM
-		#define CVD_TRACE_NON_SOLVER_TRANSFORM(Transform, DebugName) \
-		FChaosVisualDebuggerTrace::TraceNonSolverTransform(Transform, DebugName);
-	#endif
-
 	#ifndef CVD_TRACE_INVALIDATE_CACHED_GEOMETRY
 		#define CVD_TRACE_INVALIDATE_CACHED_GEOMETRY(ImplicitObjectPtr) \
 		FChaosVisualDebuggerTrace::InvalidateGeometryFromCache(ImplicitObjectPtr);
@@ -236,14 +226,6 @@
 		#define CVD_TRACE_STEP_MID_PHASES_FROM_COLLISION_CONSTRAINTS(DataChannel, CollisionConstraints)
 	#endif
 
-	#ifndef CVD_TRACE_NON_SOLVER_LOCATION
-	      #define CVD_TRACE_NON_SOLVER_LOCATION(Location, DebugName)
-	#endif
-
-	#ifndef CVD_TRACE_NON_SOLVER_TRANSFORM
-		#define CVD_TRACE_NON_SOLVER_TRANSFORM(Transform, DebugName)
-	#endif
-
 	#ifndef CVD_TRACE_INVALIDATE_CACHED_GEOMETRY
 		#define CVD_TRACE_INVALIDATE_CACHED_GEOMETRY(ImplicitObjectPtr)
 	#endif
diff --git a/Engine/Source/Runtime/Experimental/Chaos/Public/ChaosVisualDebugger/ChaosVisualDebuggerTrace.h b/Engine/Source/Runtime/Experimental/Chaos/Public/ChaosVisualDebugger/ChaosVisualDebuggerTrace.h
--- a/Engine/Source/Runtime/Experimental/Chaos/Public/ChaosVisualDebugger/ChaosVisualDebuggerTrace.h
+++ b/Engine/Source/Runtime/Experimental/Chaos/Public/ChaosVisualDebugger/ChaosVisualDebuggerTrace.h
@@ -245,22 +245,6 @@ public:
 	 */
 	static CHAOS_API void InvalidateGeometryFromCache(const Chaos::FImplicitObject* CachedGeometryToInvalidate);
 
-	/**
-	 * Traces the provided location using with the provided ID -
-	 * These are not tied to any solver step so they will be recorded as part of the current game frame data
-	 * @param InLocation Location to Trace
-	 * @param DebugNameID Name to use as ID
-	 */
-	static CHAOS_API void TraceNonSolverLocation(const FVector& InLocation, FStringView DebugNameID);
-
-	/**
-	 * Traces the provided transform using with the provided ID -
-	 * These are not tied to any solver step so they will be recorded as part of the current game frame data
-	 * @param InTransform Transform to Trace
-	 * @param DebugNameID Name to use as ID
-	 */
-	static CHAOS_API void TraceNonSolverTransform(const FTransform& InTransform, FStringView DebugNameID);
-
 	static CHAOS_API void TraceSceneQueryStart(const Chaos::FImplicitObject* InputGeometry, const FQuat& GeometryOrientation, const FVector& Start, const FVector& End, ECollisionChannel TraceChannel, FChaosVDCollisionQueryParams&& Params, FChaosVDCollisionResponseParams&& ResponseParams, FChaosVDCollisionObjectQueryParams&& ObjectParams, EChaosVDSceneQueryType QueryType, EChaosVDSceneQueryMode QueryMode, int32 SolverID, bool bIsRetry);
 	static CHAOS_API void TraceSceneQueryVisit(FChaosVDQueryVisitStep&& InQueryVisitData);
 
@@ -301,6 +285,22 @@ private:
 	 * @param ContextData Context to be used to tied this Trace event to a specific solver frame and step
 	 */
 	static CHAOS_API void TraceParticle(Chaos::FGeometryParticleHandle* ParticleHandle, const FChaosVDContext& ContextData);
+
+	/**
+	 * Traces the provided location using with the provided ID -
+	 * These are not tied to any solver step so they will be recorded as part of the current game frame data
+	 * @param InLocation Location to Trace
+	 * @param DebugNameID Name to use as ID
+	 */
+	static void TraceNonSolverLocation(const FVector& InLocation, FStringView DebugNameID);
+
+	/**
+	 * Traces the provided transform using with the provided ID -
+	 * These are not tied to any solver step so they will be recorded as part of the current game frame data
+	 * @param InTransform Transform to Trace
+	 * @param DebugNameID Name to use as ID
+	 */
+	static void TraceNonSolverTransform(const FTransform& InTransform, FStringView DebugNameID);
 	
 	/** Resets the state of the CVD Tracer */
 	static void Reset();
