[ChaosVD] Fix for ensure that happens when the Stop button is pressed without loading a recording first

If a recording is not loaded, the playback buttons need to be disabled. This CL does that.

#tests I tested that when CVD is opened, the playback controls are disabled until you load a CVD file

diff --git a/Engine/Plugins/ChaosVD/Source/ChaosVD/Private/Widgets/SChaosVDPlaybackViewport.cpp b/Engine/Plugins/ChaosVD/Source/ChaosVD/Private/Widgets/SChaosVDPlaybackViewport.cpp
--- a/Engine/Plugins/ChaosVD/Source/ChaosVD/Private/Widgets/SChaosVDPlaybackViewport.cpp
+++ b/Engine/Plugins/ChaosVD/Source/ChaosVD/Private/Widgets/SChaosVDPlaybackViewport.cpp
@@ -85,6 +85,7 @@ void SChaosVDPlaybackViewport::Construct(const FArguments& InArgs, TWeakPtr<FCha
 			+SVerticalBox::Slot()
 			[
 				SAssignNew(GameFramesTimelineWidget, SChaosVDTimelineWidget)
+				.IsEnabled_Raw(this, &SChaosVDPlaybackViewport::CanPlayback)
 				.ButtonVisibilityFlags(static_cast<uint16>(EChaosVDTimelineElementIDFlags::AllPlayback))
 				.OnFrameChanged_Raw(this, &SChaosVDPlaybackViewport::OnFrameSelectionUpdated)
 				.OnButtonClicked(this, &SChaosVDPlaybackViewport::HandlePlaybackButtonClicked)
@@ -294,4 +295,11 @@ void SChaosVDPlaybackViewport::HandlePlaybackButtonClicked(EChaosVDPlaybackButto
 	Chaos::VisualDebugger::HandleUserPlaybackInputControl(ButtonID, *this, PlaybackController);
 }
 
+bool SChaosVDPlaybackViewport::CanPlayback() const
+{
+	const TSharedPtr<FChaosVDPlaybackController> PlaybackControllerPtr = PlaybackController.Pin();
+
+	return PlaybackControllerPtr && PlaybackControllerPtr->IsRecordingLoaded();
+}
+
 #undef LOCTEXT_NAMESPACE
diff --git a/Engine/Plugins/ChaosVD/Source/ChaosVD/Private/Widgets/SChaosVDPlaybackViewport.h b/Engine/Plugins/ChaosVD/Source/ChaosVD/Private/Widgets/SChaosVDPlaybackViewport.h
--- a/Engine/Plugins/ChaosVD/Source/ChaosVD/Private/Widgets/SChaosVDPlaybackViewport.h
+++ b/Engine/Plugins/ChaosVD/Source/ChaosVD/Private/Widgets/SChaosVDPlaybackViewport.h
@@ -62,6 +62,8 @@ protected:
 
 	void HandlePlaybackButtonClicked(EChaosVDPlaybackButtonsID ButtonID);
 
+	bool CanPlayback() const;
+
 	TSharedPtr<SChaosVDTimelineWidget> GameFramesTimelineWidget;
 
 	TSharedPtr<FChaosVDPlaybackViewportClient> PlaybackViewportClient;
