[ChaosVD] Fix crash that happens when a selected particle changes its collision data to Simulated to Query Only

CVD shows query only geometry as translucent. As we use Instanced Static Mesh Components, changing the visualization from Opaque to Translucent it means that CVD has to migrate the Instance Handle used to access the geometry to another component with the correct material.

This is done by CVD's mesh instance handle system under the hood. The issue is that we were removing the instance from the old component before trying to clear the selection. The fix is just revert the order of that operation (clear the selection first and then remove).

In addition to the fix, I added an ensure and early out in case we get into this situation by other reasons. That should avoid the crash and leave a trace in the crash reporter.

#tests I tested that by loading a CVD file that presented this issue, the editor no longer crashes when the collision data is updated and the component changes

diff --git a/Engine/Plugins/ChaosVD/Source/ChaosVD/Private/Components/ChaosVDInstancedStaticMeshComponent.cpp b/Engine/Plugins/ChaosVD/Source/ChaosVD/Private/Components/ChaosVDInstancedStaticMeshComponent.cpp
--- a/Engine/Plugins/ChaosVD/Source/ChaosVD/Private/Components/ChaosVDInstancedStaticMeshComponent.cpp
+++ b/Engine/Plugins/ChaosVD/Source/ChaosVD/Private/Components/ChaosVDInstancedStaticMeshComponent.cpp
@@ -52,6 +52,12 @@ void UChaosVDInstancedStaticMeshComponent::SetIsSelected(const TSharedPtr<FChaos
 
 	const int32 InstanceIndex = InInstanceHandle->GetMeshInstanceIndex();
 
+	if (!ensure(IsValidInstance(InstanceIndex)))
+	{
+		UE_LOG(LogChaosVDEditor, Error, TEXT("[%s] Attempted to update a mesh instance using a handle with an invalid instance index | Handle Instance Index [%d] | Current Instance Conut [%d]"), ANSI_TO_TCHAR(__FUNCTION__), InstanceIndex, GetInstanceCount());
+		return;
+	}
+
 	NotifySMInstanceSelectionChanged({this, InstanceIndex}, bIsSelected);
 }
 
