[ChaosVD] Fix for crash that can happen when a recording is being loaded or unloaded

- Updated the playback controller to make sure it updates the available track data as soon a recording is unloaded. Also we check if there is any new data in the recording we need to process/cache before we execute any deferred gt updates.
- For any cases this fix missed, I updated the code so instead of crashing, we ensure and early out (and added a log error so the problem is visible)

#tests I tested this change by loading and unloading different CVD recordings, also switching between recording to file and auto load, and by doing a live recording. All without closing the CVD instance and I don't get the crash anymore

diff --git a/Engine/Plugins/ChaosVD/Source/ChaosVD/Private/ChaosVDPlaybackController.cpp b/Engine/Plugins/ChaosVD/Source/ChaosVD/Private/ChaosVDPlaybackController.cpp
--- a/Engine/Plugins/ChaosVD/Source/ChaosVD/Private/ChaosVDPlaybackController.cpp
+++ b/Engine/Plugins/ChaosVD/Source/ChaosVD/Private/ChaosVDPlaybackController.cpp
@@ -101,7 +101,10 @@ void FChaosVDPlaybackController::UnloadCurrentRecording(EChaosVDUnloadRecordingF
 		LoadedRecording.Reset();
 	}
 
-	TrackInfoPerType.Reset();
+	// This will make sure the cached data used by the UI is up to date.
+	// It already handles internally an unloaded recording, in which case the cached data will be properly reset
+
+	HandleCurrentRecordingUpdated();
 
 	if (const TSharedPtr<FChaosVDScene> SceneToControlSharedPtr = SceneToControl.Pin())
 	{
@@ -185,17 +188,18 @@ void FChaosVDPlaybackController::GoToRecordedSolverStep_AssumesLocked(const int3
 			TSharedPtr<FChaosVDTrackInfo> CurrentTrackInfo;
 			if (TrackInfoByIDMap* TrackInfoByID = TrackInfoPerType.Find(EChaosVDTrackType::Solver))
 			{
-				if (TSharedPtr<FChaosVDTrackInfo>* TrackInfo = TrackInfoByID->Find(InTrackID))
+				if (const TSharedPtr<FChaosVDTrackInfo>* TrackInfo = TrackInfoByID->Find(InTrackID))
 				{
 					CurrentTrackInfo = *TrackInfo;
-					if (!ensure(CurrentTrackInfo.IsValid()))
-					{
-						UE_LOG(LogChaosVDEditor, Warning, TEXT("[%s] Track info for track ID [%d]. We can't continue..."), ANSI_TO_TCHAR(__FUNCTION__), InTrackID);
-						return;
-					}
 				}
 			}
 
+			if (!ensure(CurrentTrackInfo.IsValid()))
+			{
+				UE_LOG(LogChaosVDEditor, Error, TEXT("[%s] Track info for track ID [%d]. We can't continue..."), ANSI_TO_TCHAR(__FUNCTION__), InTrackID);
+				return;
+			}
+
 			if (FChaosVDSolverFrameData* SolverFrameData = LoadedRecording->GetSolverFrameData_AssumesLocked(InTrackID, FrameNumber))
 			{
 				const int32 FrameDiff = FrameNumber - CurrentTrackInfo->CurrentFrame;
@@ -267,7 +271,13 @@ void FChaosVDPlaybackController::GoToRecordedGameFrame_AssumesLocked(const int32
 		{
 			if (TrackInfoByIDMap* TrackInfoByID = TrackInfoPerType.Find(EChaosVDTrackType::Game))
 			{
-				if (const TSharedPtr<FChaosVDTrackInfo>& TrackInfoSharedPtr = TrackInfoByID->FindChecked(GameTrackID))
+				const TSharedPtr<FChaosVDTrackInfo>* TrackInfoSharedPtrPtr = TrackInfoByID->Find(GameTrackID);
+				if (!ensure(TrackInfoSharedPtrPtr))
+				{
+					return;
+				}
+				
+				if (const TSharedPtr<FChaosVDTrackInfo> TrackInfoSharedPtr = *TrackInfoSharedPtrPtr)
 				{
 					if (const FChaosVDGameFrameData* FoundGameFrameData = LoadedRecording->GetGameFrameData_AssumesLocked(FrameNumber))
 					{
@@ -589,9 +599,14 @@ void FChaosVDPlaybackController::GetAvailableTrackInfosAtTrackFrame(EChaosVDTrac
 	LoadedRecording->GetAvailableSolverIDsAtGameFrameNumber(CorrectedFrameNumber, AvailableSolversID);
 	
 	TrackInfoByIDMap& TrackInfoMap = TrackInfoPerType.FindOrAdd(TrackTypeToFind);
-	for (int32 SolverID : AvailableSolversID)
+	for (const int32 SolverID : AvailableSolversID)
 	{
-		OutTrackInfo.Add(TrackInfoMap.FindChecked(SolverID));
+		// The recording might have the solver data available added because it was added the trace analysis thread, 
+		// but the playback controller didn't process it in the game thread yet
+		if (const TSharedPtr<FChaosVDTrackInfo>* SolverTrackInfo = TrackInfoMap.Find(SolverID))
+		{
+			OutTrackInfo.Add(*SolverTrackInfo);
+		}	
 	}
 }
 
@@ -615,6 +630,20 @@ bool FChaosVDPlaybackController::Tick(float DeltaTime)
 		return true;
 	}
 
+
+	const bool bIsRecordingLoaded = LoadedRecording.IsValid();
+
+	if (bIsRecordingLoaded)
+	{
+		uint64 CurrentLastUpdatedTime = LoadedRecording->GetLastUpdatedTimeAsCycle();
+		if (CurrentLastUpdatedTime != RecordingLastSeenTimeUpdatedAsCycle)
+		{
+			RecordingLastSeenTimeUpdatedAsCycle = CurrentLastUpdatedTime;
+
+			HandleCurrentRecordingUpdated();
+		}
+	}
+
 	if (bHasPendingGTUpdateBroadcast)
 	{
 		ControllerUpdatedDelegate.Broadcast(ThisWeakPtr);
@@ -630,16 +659,8 @@ bool FChaosVDPlaybackController::Tick(float DeltaTime)
 		}
 	}
 
-	if (LoadedRecording.IsValid())
+	if (bIsRecordingLoaded)
 	{
-		uint64 CurrentLastUpdatedTime = LoadedRecording->GetLastUpdatedTimeAsCycle();
-		if (CurrentLastUpdatedTime != RecordingLastSeenTimeUpdatedAsCycle)
-		{
-			RecordingLastSeenTimeUpdatedAsCycle = CurrentLastUpdatedTime;
-
-			HandleCurrentRecordingUpdated();
-		}
-
 		// Load at least the first frame
 		if (!bPlayedFirstFrame)
 		{
