[ChaosVD] Fix an issue where CVD shows an incorrect vector while visualizing recorded Line Traces hits

This CL also fixes an issue where during load, CVD constantly tries to playback the same frame over and over.

This is an issue for two reasons

- Everything that can only be selected (like scene queries) when the frame doesn't change, cannot be selected until the CVD file finish loading, or during live debugging
- In complex CVD files, this increases the loading time noticeably as CVD will be attempting to update the whole scene with each frame that is loaded to the recording (even if you are not playing).

The fix for the incorrect vector is safe. For the selection/loading issue, it is safe but it is more riskier as it is a core part of CVD. For that reason I gated the change behind a cvar p.Chaos.VD.Tool.BroadcastGameFrameUpdateOnlyIfChanged . The default it is to use the new behavior that fixes the issue, but it can be used to revert back to the old behavior.

#tests I tested that after this changes, I can load the affected CVD file and I don't see the wrong vectors anymore. I also verified I can select scene queries while the CVD file is still loading. I repeated the process several times to ensure there is no side effects with the change. Lastly I tested that I can use the cvar to revert to the previous behavior

diff --git a/Engine/Plugins/ChaosVD/Source/ChaosVD/Private/Trace/DataProcessors/ChaosVDSceneQueryVisitDataProcessor.cpp b/Engine/Plugins/ChaosVD/Source/ChaosVD/Private/Trace/DataProcessors/ChaosVDSceneQueryVisitDataProcessor.cpp
--- a/Engine/Plugins/ChaosVD/Source/ChaosVD/Private/Trace/DataProcessors/ChaosVDSceneQueryVisitDataProcessor.cpp
+++ b/Engine/Plugins/ChaosVD/Source/ChaosVD/Private/Trace/DataProcessors/ChaosVDSceneQueryVisitDataProcessor.cpp
@@ -32,6 +32,14 @@ bool FChaosVDSceneQueryVisitDataProcessor::ProcessRawData(const TArray<uint8>& I
 				TSharedPtr<FChaosVDQueryDataWrapper> QueryDataPtr = *QueryDataPtrPtr;
 				if (QueryDataPtrPtr->IsValid())
 				{
+					//TODO: There is an existing issue where CVD is recording garbage data for the Hit Face Normalif the query is a line trace.
+					// That value is not used in line traces, so for now just clear it out until we can implement a better solution
+					// where we include any post-processing done on the hit data during the HitConversion step.
+					if (QueryDataPtr->Type == EChaosVDSceneQueryType::RayCast)
+					{
+						VisitStepData.HitData.FaceNormal = FVector::ZeroVector;
+					}
+
 					if (VisitStepData.HitData.HasValidData())
 					{
 						// Quick and dirty way of show the hits in the details panel. If copying this data around becomes a bottle neck we can write a customization layout for it
diff --git a/Engine/Plugins/ChaosVD/Source/ChaosVD/Private/Visualizers/ChaosVDSceneQueryDataComponentVisualizer.cpp b/Engine/Plugins/ChaosVD/Source/ChaosVD/Private/Visualizers/ChaosVDSceneQueryDataComponentVisualizer.cpp
--- a/Engine/Plugins/ChaosVD/Source/ChaosVD/Private/Visualizers/ChaosVDSceneQueryDataComponentVisualizer.cpp
+++ b/Engine/Plugins/ChaosVD/Source/ChaosVD/Private/Visualizers/ChaosVDSceneQueryDataComponentVisualizer.cpp
@@ -219,7 +219,12 @@ void FChaosVDSceneQueryDataComponentVisualizer::DrawHits(const UActorComponent*
 		constexpr float NormalScale = 10.5f;
 		FChaosVDDebugDrawUtils::DrawCircle(PDI, SQVisitData.HitData.WorldPosition, CircleRadius, CircleSegments, InColor, Thickness, Axes.GetUnitAxis(EAxis::Y), Axes.GetUnitAxis(EAxis::Z), HitPointDebugText,  ESceneDepthPriorityGroup::SDPG_Foreground);
 		FChaosVDDebugDrawUtils::DrawArrowVector(PDI, SQVisitData.HitData.WorldPosition, SQVisitData.HitData.WorldPosition + SQVisitData.HitData.FaceNormal * NormalScale, HitFaceNormalDebugText, (FLinearColor(InColor) * 0.65f).ToFColorSRGB(), ESceneDepthPriorityGroup::SDPG_Foreground);
-		FChaosVDDebugDrawUtils::DrawArrowVector(PDI, SQVisitData.HitData.WorldPosition, SQVisitData.HitData.WorldPosition + SQVisitData.HitData.WorldNormal * NormalScale, HitWorldNormalDebugText, InColor, ESceneDepthPriorityGroup::SDPG_Foreground);
+
+		// Hit Face Normal is not used in line traces
+		if (SceneQueryData.Type != EChaosVDSceneQueryType::RayCast)
+		{
+			FChaosVDDebugDrawUtils::DrawArrowVector(PDI, SQVisitData.HitData.WorldPosition, SQVisitData.HitData.WorldPosition + SQVisitData.HitData.WorldNormal * NormalScale, HitWorldNormalDebugText, InColor, ESceneDepthPriorityGroup::SDPG_Foreground);
+		}
 
 		if (SQVisitData.bIsSelectedInEditor)
 		{
diff --git a/Engine/Plugins/ChaosVD/Source/ChaosVD/Private/Widgets/SChaosVDPlaybackViewport.cpp b/Engine/Plugins/ChaosVD/Source/ChaosVD/Private/Widgets/SChaosVDPlaybackViewport.cpp
--- a/Engine/Plugins/ChaosVD/Source/ChaosVD/Private/Widgets/SChaosVDPlaybackViewport.cpp
+++ b/Engine/Plugins/ChaosVD/Source/ChaosVD/Private/Widgets/SChaosVDPlaybackViewport.cpp
@@ -20,6 +20,15 @@
 
 #define LOCTEXT_NAMESPACE "ChaosVisualDebugger"
 
+namespace Chaos::VisualDebugger::Cvars
+{
+	static bool bBroadcastGameFrameUpdateEvenIfNotChanged = false;
+	static FAutoConsoleVariableRef CVarChaosVDBroadcastGameFrameUpdateEvenIfNotChanged(
+		TEXT("p.Chaos.VD.Tool.BroadcastGameFrameUpdateEvenIfNotChanged"),
+		bBroadcastGameFrameUpdateEvenIfNotChanged,
+		TEXT("If true, each time we get a controller data updated event, a game frame update will be triggered even if the frame didn't change..."));
+}
+
 SChaosVDPlaybackViewport::~SChaosVDPlaybackViewport()
 {
 	PlaybackViewportClient->Viewport = nullptr;
@@ -182,7 +191,9 @@ void SChaosVDPlaybackViewport::HandlePlaybackControllerDataUpdated(TWeakPtr<FCha
 		{
 			// Max is inclusive and we use this to request as the index on the recorded frames/steps arrays so we need to -1 to the available frames/steps
 			GameFramesTimelineWidget->UpdateMinMaxValue(0, TrackInfo->MaxFrames != INDEX_NONE ? TrackInfo->MaxFrames -1  : 0);
-			GameFramesTimelineWidget->SetCurrentTimelineFrame(TrackInfo->CurrentFrame);
+
+			const bool bNeedsToBroadcastChange = Chaos::VisualDebugger::Cvars::bBroadcastGameFrameUpdateEvenIfNotChanged || (GameFramesTimelineWidget->GetCurrentFrame() != TrackInfo->CurrentFrame);
+			GameFramesTimelineWidget->SetCurrentTimelineFrame(TrackInfo->CurrentFrame, bNeedsToBroadcastChange ? EChaosVDSetTimelineFrameFlags::BroadcastChange : EChaosVDSetTimelineFrameFlags::Silent);
 
 			constexpr uint16 PlaybackElementDisabledDuringLiveSession = static_cast<uint16>(EChaosVDTimelineElementIDFlags::Stop | EChaosVDTimelineElementIDFlags::Next | EChaosVDTimelineElementIDFlags::Prev);
 			if (ControllerSharedPtr->IsPlayingLiveSession())
diff --git a/Engine/Plugins/ChaosVD/Source/ChaosVD/Private/Widgets/SChaosVDTimelineWidget.h b/Engine/Plugins/ChaosVD/Source/ChaosVD/Private/Widgets/SChaosVDTimelineWidget.h
--- a/Engine/Plugins/ChaosVD/Source/ChaosVD/Private/Widgets/SChaosVDTimelineWidget.h
+++ b/Engine/Plugins/ChaosVD/Source/ChaosVD/Private/Widgets/SChaosVDTimelineWidget.h
@@ -90,6 +90,8 @@ public:
 
 	void SetAutoStopEnabled(bool bNewEnabled) { bAutoStopEnabled = bNewEnabled; }
 
+	int32 GetCurrentFrame() const { return CurrentFrame; }
+
 	void Play();
 	FReply  Stop();
 
diff --git a/Engine/Source/Runtime/Experimental/ChaosVisualDebugger/Public/DataWrappers/ChaosVDQueryDataWrappers.h b/Engine/Source/Runtime/Experimental/ChaosVisualDebugger/Public/DataWrappers/ChaosVDQueryDataWrappers.h
--- a/Engine/Source/Runtime/Experimental/ChaosVisualDebugger/Public/DataWrappers/ChaosVDQueryDataWrappers.h
+++ b/Engine/Source/Runtime/Experimental/ChaosVisualDebugger/Public/DataWrappers/ChaosVDQueryDataWrappers.h
@@ -295,6 +295,7 @@ struct FChaosVDQueryHitData : public FChaosVDWrapperDataBase
 	UPROPERTY(VisibleAnywhere, Category=QueryData)
 	FVector WorldNormal = FVector::ZeroVector;
 
+	/** This is not used for Line Traces */
 	UPROPERTY(VisibleAnywhere, Category=QueryData)
 	FVector FaceNormal = FVector::ZeroVector;
 };
