[ChaosVD] Make sure the editor does not throttle the CPU if we are connected to a live session

This is the only sceneraio where the editor might not be the application that is in focus. A specific case is when we are live debugging a running game client

diff --git a/Engine/Plugins/ChaosVD/Source/ChaosVD/Private/Widgets/SChaosVDMainTab.cpp b/Engine/Plugins/ChaosVD/Source/ChaosVD/Private/Widgets/SChaosVDMainTab.cpp
--- a/Engine/Plugins/ChaosVD/Source/ChaosVD/Private/Widgets/SChaosVDMainTab.cpp
+++ b/Engine/Plugins/ChaosVD/Source/ChaosVD/Private/Widgets/SChaosVDMainTab.cpp
@@ -155,6 +155,13 @@ void SChaosVDMainTab::Construct(const FArguments& InArgs, TSharedPtr<FChaosVDEng
 	// Make sure these tabs are always focused at the start
 	TabManager->TryInvokeTab(FChaosVDTabID::SolversTrack);
 	TabManager->TryInvokeTab(FChaosVDTabID::DetailsPanel);
+
+	SetUpDisableCPUThrottlingDelegate();
+}
+
+SChaosVDMainTab::~SChaosVDMainTab()
+{
+	CleanUpDisableCPUThrottlingDelegate();
 }
 
 void SChaosVDMainTab::BringToFront()
@@ -452,6 +459,26 @@ bool SChaosVDMainTab::ConnectToLiveSession(int32 SessionID, const FString Sessio
 	return bSuccess;
 }
 
+void SChaosVDMainTab::SetUpDisableCPUThrottlingDelegate()
+{
+	if (GEditor)
+	{
+		GEditor->ShouldDisableCPUThrottlingDelegates.Add(UEditorEngine::FShouldDisableCPUThrottling::CreateSP(this, &SChaosVDMainTab::ShouldDisableCPUThrottling));
+		DisableCPUThrottleHandle = GEditor->ShouldDisableCPUThrottlingDelegates.Last().GetHandle();
+	}
+}
+
+void SChaosVDMainTab::CleanUpDisableCPUThrottlingDelegate() const
+{
+	if (GEditor)
+	{
+		GEditor->ShouldDisableCPUThrottlingDelegates.RemoveAll([this](const UEditorEngine::FShouldDisableCPUThrottling& Delegate)
+		{
+			return Delegate.GetHandle() == DisableCPUThrottleHandle;
+		});
+	}
+}
+
 void SChaosVDMainTab::RegisterMainTabMenu()
 {
 	const UToolMenus* ToolMenus = UToolMenus::Get();
@@ -571,6 +598,12 @@ void SChaosVDMainTab::BrowseLiveSessionsFromTraceStore() const
 	}
 }
 
+bool SChaosVDMainTab::ShouldDisableCPUThrottling() const
+{
+	// If we are playing a live session, it is likely the editor will be in the background, so we need to disable CPU Throttling
+	return ChaosVDEngine && ChaosVDEngine->GetCurrentSessionDescriptor().bIsLiveSession;
+}
+
 FReply SChaosVDMainTab::HandleSessionConnectionClicked()
 {
 	if (TSharedPtr<FChaosVDPlaybackController> PlaybackControllerPtr =  GetChaosVDEngineInstance()->GetPlaybackController())
diff --git a/Engine/Plugins/ChaosVD/Source/ChaosVD/Private/Widgets/SChaosVDMainTab.h b/Engine/Plugins/ChaosVD/Source/ChaosVD/Private/Widgets/SChaosVDMainTab.h
--- a/Engine/Plugins/ChaosVD/Source/ChaosVD/Private/Widgets/SChaosVDMainTab.h
+++ b/Engine/Plugins/ChaosVD/Source/ChaosVD/Private/Widgets/SChaosVDMainTab.h
@@ -44,6 +44,8 @@ public:
 
 	void Construct(const FArguments& InArgs, TSharedPtr<FChaosVDEngine> InChaosVDEngine);
 
+	virtual ~SChaosVDMainTab() override;
+
 	TSharedRef<FChaosVDEngine> GetChaosVDEngineInstance() const { return ChaosVDEngine.ToSharedRef(); }
 
 	template<typename TabType>
@@ -71,6 +73,9 @@ public:
 	bool ConnectToLiveSession(int32 SessionID, FString SessionAddress) const;
 
 private:
+	
+	void SetUpDisableCPUThrottlingDelegate();
+	void CleanUpDisableCPUThrottlingDelegate() const;
 
 	void RegisterMainTabMenu();
 
@@ -96,6 +101,8 @@ private:
 
 	void BrowseLiveSessionsFromTraceStore() const;
 
+	bool ShouldDisableCPUThrottling() const;
+
 	TSharedPtr<FChaosVDEngine> ChaosVDEngine;
 
 	FName StatusBarID;
@@ -113,6 +120,8 @@ private:
 
 	FOnActiveViewportChanged ViewportChangedDelegate;
 
+	FDelegateHandle DisableCPUThrottleHandle;
+
 	FReply HandleSessionConnectionClicked();
 	FText GetConnectButtonText() const;
 	FText GetConnectButtonTooltipText() const;
