[ChaosVD] Fix for an issue where static particles creation, and disabling particles was only being recorded only as part of keyframes (causing a 10 seconds delay for that data to be shown in CVD)

Now we trace the particle data each time a particle is disabled or created, independenty of what solver stage we are in. On load, CVD will add it to the latest available stage when these actions happened (this is to keep the change to a minimum in case we want to propose it for a possible 5.4.1 hotfix release. I will investigate what other better alternatives of doing this we have in a follow up CL).

diff --git a/Engine/Plugins/ChaosVD/Source/ChaosVD/Private/ChaosVDPlaybackController.cpp b/Engine/Plugins/ChaosVD/Source/ChaosVD/Private/ChaosVDPlaybackController.cpp
--- a/Engine/Plugins/ChaosVD/Source/ChaosVD/Private/ChaosVDPlaybackController.cpp
+++ b/Engine/Plugins/ChaosVD/Source/ChaosVD/Private/ChaosVDPlaybackController.cpp
@@ -16,6 +16,15 @@
 
 #define LOCTEXT_NAMESPACE "ChaosVisualDebugger"
 
+namespace Chaos::VisualDebugger::Cvars
+{
+	static bool bPlayAllPreviousFrameSteps = true;
+	static FAutoConsoleVariableRef CVarChaosVDbPlayAllPreviousFrameSteps(
+		TEXT("p.Chaos.VD.Tool.PlayAllPreviousFrameSteps"),
+		bPlayAllPreviousFrameSteps,
+		TEXT("If true, each time we get playback a solver frame in a specific stage, we will play all the previous steps from that frame in sequence to ensure we have the correct visualization for what happened in that frame."));
+}
+
 FChaosVDPlaybackController::FChaosVDPlaybackController(const TWeakPtr<FChaosVDScene>& InSceneToControl)
 {
 	SceneToControl = InSceneToControl;
@@ -179,6 +188,19 @@ void FChaosVDPlaybackController::HandleFrameRateOverrideSettingsChanged(UChaosVD
 	CurrentFrameRateOverride = CVDSettings->bPlaybackAtRecordedFrameRate ? InvalidFrameRateOverride : CVDSettings->TargetFrameRateOverride;
 }
 
+void FChaosVDPlaybackController::PlaySolverStepData(int32 TrackID, const TSharedRef<FChaosVDScene>& InSceneToControlSharedPtr, const FChaosVDSolverFrameData& InSolverFrameData, int32 StepIndex)
+{
+	if (InSolverFrameData.SolverSteps.IsValidIndex(StepIndex))
+	{
+		InSceneToControlSharedPtr->UpdateFromRecordedStepData(TrackID, InSolverFrameData.SolverSteps[StepIndex], InSolverFrameData);
+	}
+	else
+	{
+		// This is common if we stop PIE, change worlds, and PIE again without stopping the recording
+		UE_LOG(LogChaosVDEditor, Verbose, TEXT("[%s] Tried to scrub to an invalid step | Step Number [%d] ..."), ANSI_TO_TCHAR(__FUNCTION__), StepIndex);
+	}
+}
+
 void FChaosVDPlaybackController::GoToRecordedSolverStep_AssumesLocked(const int32 InTrackID, const int32 FrameNumber, const int32 Step, FGuid InstigatorID, int32 Attempts)
 {
 	if (const TSharedPtr<FChaosVDScene> SceneToControlSharedPtr = SceneToControl.Pin())
@@ -218,35 +240,19 @@ void FChaosVDPlaybackController::GoToRecordedSolverStep_AssumesLocked(const int3
 					PlayFromClosestKeyFrame_AssumesLocked(InTrackID, FrameNumber, *SceneToControlSharedPtr.Get());
 				}
 
-				if (CurrentTrackInfo->LockedOnStep != INDEX_NONE)
+				const int32 StepDiff = Step - CurrentTrackInfo->CurrentStep;
+				const bool bNeedsPlayPreviousSteps = CurrentTrackInfo->CurrentFrame != FrameNumber || StepDiff < 0 || FMath::Abs(StepDiff) > FrameDriftTolerance;
+
+				if (Chaos::VisualDebugger::Cvars::bPlayAllPreviousFrameSteps && bNeedsPlayPreviousSteps)
 				{
-					// If this track is locked to a specific step, we need to play back the previous steps on the current frame, because not all steps capture the same data.
-					// For example, Particles positions are fully captured in the first sub-step and the last one
-					for (int32 StepIndex = 0; StepIndex <= CurrentTrackInfo->LockedOnStep; StepIndex++)
+					for (int32 StepIndex = 0; StepIndex <= Step; StepIndex++)
 					{
-						if (SolverFrameData && SolverFrameData->SolverSteps.IsValidIndex(CurrentTrackInfo->LockedOnStep))
-						{
-							SceneToControlSharedPtr->UpdateFromRecordedStepData(InTrackID, SolverFrameData->SolverSteps[StepIndex], *SolverFrameData);
-						}
-						else
-						{
-							// This is common if we stop PIE, change worlds, and PIE again without stopping the recording
-							UE_LOG(LogChaosVDEditor, Verbose, TEXT("[%s] Tried to scrub to an invalid step | Step Number [%d] ..."), ANSI_TO_TCHAR(__FUNCTION__), CurrentTrackInfo->LockedOnStep);
-							return;
-						}
+						PlaySolverStepData(InTrackID, SceneToControlSharedPtr.ToSharedRef(), *SolverFrameData, StepIndex);
 					}
 				}
 				else
 				{
-					if (SolverFrameData && SolverFrameData->SolverSteps.IsValidIndex(Step))
-					{
-						SceneToControlSharedPtr->UpdateFromRecordedStepData(InTrackID, SolverFrameData->SolverSteps[Step], *SolverFrameData);
-					}
-					else
-					{
-						// This is common if we stop PIE, change worlds, and PIE again without stopping the recording
-						UE_LOG(LogChaosVDEditor, Verbose, TEXT("[%s] Tried to scrub to an invalid step | Step Number [%d] ..."), ANSI_TO_TCHAR(__FUNCTION__), CurrentTrackInfo->LockedOnStep);
-					}
+					PlaySolverStepData(InTrackID, SceneToControlSharedPtr.ToSharedRef(), *SolverFrameData, Step);
 				}
 
 				CurrentTrackInfo->CurrentFrame = FrameNumber;
diff --git a/Engine/Plugins/ChaosVD/Source/ChaosVD/Private/ChaosVDPlaybackController.h b/Engine/Plugins/ChaosVD/Source/ChaosVD/Private/ChaosVDPlaybackController.h
--- a/Engine/Plugins/ChaosVD/Source/ChaosVD/Private/ChaosVDPlaybackController.h
+++ b/Engine/Plugins/ChaosVD/Source/ChaosVD/Private/ChaosVDPlaybackController.h
@@ -265,6 +265,8 @@ protected:
 
 	void HandleFrameRateOverrideSettingsChanged(UChaosVDEditorSettings* CVDSettings);
 
+	void PlaySolverStepData(int32 TrackID, const TSharedRef<FChaosVDScene>& InSceneToControlSharedPtr, const FChaosVDSolverFrameData& InSolverFrameData, int32 StepIndex);
+
 	/** Map containing all track info, by track type*/
 	TMap<EChaosVDTrackType, TrackInfoByIDMap> TrackInfoPerType;
 
diff --git a/Engine/Plugins/ChaosVD/Source/ChaosVD/Private/Widgets/SChaosVDSolverPlaybackControls.cpp b/Engine/Plugins/ChaosVD/Source/ChaosVD/Private/Widgets/SChaosVDSolverPlaybackControls.cpp
--- a/Engine/Plugins/ChaosVD/Source/ChaosVD/Private/Widgets/SChaosVDSolverPlaybackControls.cpp
+++ b/Engine/Plugins/ChaosVD/Source/ChaosVD/Private/Widgets/SChaosVDSolverPlaybackControls.cpp
@@ -123,6 +123,14 @@ SChaosVDSolverPlaybackControls::~SChaosVDSolverPlaybackControls()
 {
 	if (const TSharedPtr<FChaosVDPlaybackController> PlaybackControllerPtr = PlaybackController.Pin())
 	{
+		// TODO: Doing this here to try to keep this change small enough for a 5.4.1 hotfix release
+		// This should be handled by the playback controller, but currently we are not doing any tracking of what solver tracks are being unloaded/loaded there
+		// This can be done as part of UE-192940 , which is a task to remove a similar dependency in the UI for the playback logic
+		if (FChaosVDTrackInfo* SolverTrackInfo = PlaybackControllerPtr->GetMutableTrackInfo(EChaosVDTrackType::Solver, SolverID))
+		{
+			SolverTrackInfo->CurrentFrame = 0;
+		}
+
 		PlaybackControllerPtr->ReleaseExclusivePlaybackControls(*this);
 	}
 }
diff --git a/Engine/Source/Runtime/Experimental/Chaos/Private/Chaos/Framework/PhysicsSolverBase.cpp b/Engine/Source/Runtime/Experimental/Chaos/Private/Chaos/Framework/PhysicsSolverBase.cpp
--- a/Engine/Source/Runtime/Experimental/Chaos/Private/Chaos/Framework/PhysicsSolverBase.cpp
+++ b/Engine/Source/Runtime/Experimental/Chaos/Private/Chaos/Framework/PhysicsSolverBase.cpp
@@ -102,6 +102,8 @@ namespace Chaos
 		SCOPE_CYCLE_COUNTER(STAT_ChaosTick);
 		CSV_SCOPED_TIMING_STAT_EXCLUSIVE(Physics);
 		CSV_SCOPED_TIMING_STAT(PhysicsVerbose, StepSolver);
+		
+		CVD_SCOPE_CONTEXT(Solver.GetChaosVDContextData());
 
 #if PHYSICS_THREAD_CONTEXT
 		FPhysicsThreadContextScope Scope(/*IsPhysicsThreadContext=*/true);
diff --git a/Engine/Source/Runtime/Experimental/Chaos/Public/Chaos/PBDRigidsEvolution.h b/Engine/Source/Runtime/Experimental/Chaos/Public/Chaos/PBDRigidsEvolution.h
--- a/Engine/Source/Runtime/Experimental/Chaos/Public/Chaos/PBDRigidsEvolution.h
+++ b/Engine/Source/Runtime/Experimental/Chaos/Public/Chaos/PBDRigidsEvolution.h
@@ -422,6 +422,8 @@ public:
 			}
 		}
 
+		CVD_TRACE_PARTICLE(Particle);
+
 		// Flag as dirty to update the spatial query acceleration structures
 		DirtyParticle(*Particle);
 	}
@@ -462,6 +464,8 @@ public:
 			// we don't leave a residual value if we get enabled again
 			Rigid->ClearIsMovingKinematic();
 		}
+
+		CVD_TRACE_PARTICLE(Particle);
 	}
 
 	/**
