[ChaosVD] Make sure the first frame is always automatically played when a CVD recording is loaded

Also, fixed an issue where the seconds counter will count faster each time a recording is stopped/started

diff --git a/Engine/Plugins/ChaosVD/Source/ChaosVD/Private/Widgets/SChaosVDPlaybackViewport.cpp b/Engine/Plugins/ChaosVD/Source/ChaosVD/Private/Widgets/SChaosVDPlaybackViewport.cpp
--- a/Engine/Plugins/ChaosVD/Source/ChaosVD/Private/Widgets/SChaosVDPlaybackViewport.cpp
+++ b/Engine/Plugins/ChaosVD/Source/ChaosVD/Private/Widgets/SChaosVDPlaybackViewport.cpp
@@ -192,7 +192,7 @@ void SChaosVDPlaybackViewport::HandlePlaybackControllerDataUpdated(TWeakPtr<FCha
 			// Max is inclusive and we use this to request as the index on the recorded frames/steps arrays so we need to -1 to the available frames/steps
 			GameFramesTimelineWidget->UpdateMinMaxValue(0, TrackInfo->MaxFrames != INDEX_NONE ? TrackInfo->MaxFrames -1  : 0);
 
-			const bool bNeedsToBroadcastChange = Chaos::VisualDebugger::Cvars::bBroadcastGameFrameUpdateEvenIfNotChanged || (GameFramesTimelineWidget->GetCurrentFrame() != TrackInfo->CurrentFrame);
+			const bool bNeedsToBroadcastChange = Chaos::VisualDebugger::Cvars::bBroadcastGameFrameUpdateEvenIfNotChanged || (GameFramesTimelineWidget->GetCurrentFrame() == 0 || GameFramesTimelineWidget->GetCurrentFrame() != TrackInfo->CurrentFrame);
 			GameFramesTimelineWidget->SetCurrentTimelineFrame(TrackInfo->CurrentFrame, bNeedsToBroadcastChange ? EChaosVDSetTimelineFrameFlags::BroadcastChange : EChaosVDSetTimelineFrameFlags::Silent);
 
 			constexpr uint16 PlaybackElementDisabledDuringLiveSession = static_cast<uint16>(EChaosVDTimelineElementIDFlags::Stop | EChaosVDTimelineElementIDFlags::Next | EChaosVDTimelineElementIDFlags::Prev);
diff --git a/Engine/Source/Runtime/Experimental/ChaosVisualDebugger/Private/ChaosVDRuntimeModule.cpp b/Engine/Source/Runtime/Experimental/ChaosVisualDebugger/Private/ChaosVDRuntimeModule.cpp
--- a/Engine/Source/Runtime/Experimental/ChaosVisualDebugger/Private/ChaosVDRuntimeModule.cpp
+++ b/Engine/Source/Runtime/Experimental/ChaosVisualDebugger/Private/ChaosVDRuntimeModule.cpp
@@ -264,6 +264,14 @@ void FChaosVDRuntimeModule::StopRecording()
 	if (FullCaptureRequesterHandle.IsValid())
 	{
 		FTSTicker::GetCoreTicker().RemoveTicker(FullCaptureRequesterHandle);
+
+		FullCaptureRequesterHandle.Reset();
+	}
+	
+	if (RecordingTimerHandle.IsValid())
+	{
+		FTSTicker::GetCoreTicker().RemoveTicker(RecordingTimerHandle);
+		RecordingTimerHandle.Reset();
 	}
 
 	{
