[ChaosVD] Fix for crash when loading a a CVD file

The issue was that a delegate was calling code that tries to clear the delegate itself. I updated the code so that clear step happens in the next tick

diff --git a/Engine/Plugins/ChaosVD/Source/ChaosVD/Private/ChaosVDGeometryBuilder.cpp b/Engine/Plugins/ChaosVD/Source/ChaosVD/Private/ChaosVDGeometryBuilder.cpp
--- a/Engine/Plugins/ChaosVD/Source/ChaosVD/Private/ChaosVDGeometryBuilder.cpp
+++ b/Engine/Plugins/ChaosVD/Source/ChaosVD/Private/ChaosVDGeometryBuilder.cpp
@@ -187,12 +187,16 @@ void FChaosVDGeometryBuilder::Initialize(const TWeakPtr<FChaosVDScene>& ChaosVDS
 	MeshComponentsWaitingForGeometry = MakeUnique<FObjectsWaitingGeometryList<FMeshComponentWeakPtr>>(ProcessMeshComponent, NSLOCTEXT("ChaosVisualDebugger", "GeometryGenNotification","Mesh Components"), ShouldProcessObjectsForKey);
 
 	GameThreadTickDelegate = FTSTicker::GetCoreTicker().AddTicker(FTickerDelegate::CreateRaw(this, &FChaosVDGeometryBuilder::GameThreadTick));
+
+	constexpr int32 MeshPendingDisposalContainerDefaultSize = 500; 
+	MeshComponentsPendingDisposal.Reserve(MeshPendingDisposalContainerDefaultSize);
 }
 
 void FChaosVDGeometryBuilder::AddReferencedObjects(FReferenceCollector& Collector)
 {
 	Collector.AddStableReferenceMap(DynamicMeshCacheMap);
 	Collector.AddStableReferenceMap(StaticMeshCacheMap);
+	Collector.AddReferencedObjects(MeshComponentsPendingDisposal);
 }
 
 bool FChaosVDGeometryBuilder::DoesImplicitContainType(const Chaos::FImplicitObject* InImplicitObject, const Chaos::EImplicitObjectType ImplicitTypeToCheck)
@@ -395,7 +399,7 @@ void FChaosVDGeometryBuilder::DestroyMeshComponent(UMeshComponent* MeshComponent
 		}
 	}
 
-	ComponentMeshPool.DisposeMeshComponent(MeshComponent);
+	MeshComponentsPendingDisposal.Add(MeshComponent);
 }
 
 TMap<uint32, UChaosVDInstancedStaticMeshComponent*>& FChaosVDGeometryBuilder::GetInstancedStaticMeshComponentCacheMap(EChaosVDMeshAttributesFlags MeshAttributeFlags)
@@ -636,6 +640,16 @@ bool FChaosVDGeometryBuilder::GameThreadTick(float DeltaTime)
 		MeshComponentsWaitingForGeometry->ProcessWaitingObjects(CurrentGeometryTasksProcessedNum);
 	}
 
+	for (TObjectPtr<UMeshComponent>& MeshComponentPtr : MeshComponentsPendingDisposal)
+	{
+		if (MeshComponentPtr)
+		{
+			ComponentMeshPool.DisposeMeshComponent(MeshComponentPtr);
+		}
+	}
+
+	MeshComponentsPendingDisposal.Reset();
+
 	return true;
 }
 
diff --git a/Engine/Plugins/ChaosVD/Source/ChaosVD/Private/ChaosVDGeometryBuilder.h b/Engine/Plugins/ChaosVD/Source/ChaosVD/Private/ChaosVDGeometryBuilder.h
--- a/Engine/Plugins/ChaosVD/Source/ChaosVD/Private/ChaosVDGeometryBuilder.h
+++ b/Engine/Plugins/ChaosVD/Source/ChaosVD/Private/ChaosVDGeometryBuilder.h
@@ -324,6 +324,9 @@ private:
 	/** Map containing already initialized Instanced static mesh components for mesh instances that required a negative scale transform and use a translucent material, for any given geometry key */
 	TMap<uint32, UChaosVDInstancedStaticMeshComponent*> TranslucentMirroredInstancedMeshComponentByGeometryKey;
 
+	/** Components that need to be processed and added to the pool in the next frame */
+	TArray<TObjectPtr<UMeshComponent>> MeshComponentsPendingDisposal;
+
 	/** Instance of uninitialized mesh components pool */
 	FChaosVDMeshComponentPool ComponentMeshPool;
 
