[ChaosVD] Fix crash that can happen when a new solver is being loaded while scrubbing

The issue was that we were never unbinding from the OnVisibilityChanged delegate when the solver playback controls were destroyed

diff --git a/Engine/Plugins/ChaosVD/Source/ChaosVD/Private/Widgets/SChaosVDSolverPlaybackControls.cpp b/Engine/Plugins/ChaosVD/Source/ChaosVD/Private/Widgets/SChaosVDSolverPlaybackControls.cpp
--- a/Engine/Plugins/ChaosVD/Source/ChaosVD/Private/Widgets/SChaosVDSolverPlaybackControls.cpp
+++ b/Engine/Plugins/ChaosVD/Source/ChaosVD/Private/Widgets/SChaosVDSolverPlaybackControls.cpp
@@ -132,6 +132,11 @@ SChaosVDSolverPlaybackControls::~SChaosVDSolverPlaybackControls()
 		}
 
 		PlaybackControllerPtr->ReleaseExclusivePlaybackControls(*this);
+		
+		if (TSharedPtr<FChaosVDScene> Scene = PlaybackControllerPtr->GetControllerScene().Pin())
+		{
+			Scene->OnSolverVisibilityUpdated().RemoveAll(this);
+		}
 	}
 }
 
