[ChaosVD] Implemented a details view to show Shape instance data and implicit object data, and fixed an issue where the incorrect shape instance index was used when we evaluate if a geometry is for a query only shape.

- Added a way to select specific meshes instead of the all meshes representing the geometry of a particle
- Now we calculate the shape instance index during the same traversal in the object hierarchy we already do to know what meshes generate
- Partially fixed an issue where the Scene Query Inspector , Select Visited Shape button would not select the correct shape
- Added support to serialize Sim and Query data values from the ShapeInstance data structure
- Removed the subobject editor from the CVD's detail panel as it is not used anymore.
- Added support to inspect structures not just ubjects in CVD details pannel.

Known issues:
- There are still issues with getting the correct shape instance index for Geometry of a Cluster union while it is being modified. It seems we get the update on the geometry change, but not the update on the shape instance array. I will log that as a separate issue and work on that next

diff --git a/Engine/Plugins/ChaosVD/Source/ChaosVD/Private/ChaosVDExtractedGeometryDataHandle.cpp b/Engine/Plugins/ChaosVD/Source/ChaosVD/Private/ChaosVDExtractedGeometryDataHandle.cpp
--- a/Engine/Plugins/ChaosVD/Source/ChaosVD/Private/ChaosVDExtractedGeometryDataHandle.cpp
+++ b/Engine/Plugins/ChaosVD/Source/ChaosVD/Private/ChaosVDExtractedGeometryDataHandle.cpp
@@ -4,7 +4,7 @@
 
 #include "Chaos/ImplicitObject.h"
 
-FName FChaosVDExtractedGeometryDataHandle::GetName() const
+FName FChaosVDExtractedGeometryDataHandle::GetTypeName() const
 {
 	using namespace Chaos;
 
diff --git a/Engine/Plugins/ChaosVD/Source/ChaosVD/Private/ChaosVDExtractedGeometryDataHandle.h b/Engine/Plugins/ChaosVD/Source/ChaosVD/Private/ChaosVDExtractedGeometryDataHandle.h
--- a/Engine/Plugins/ChaosVD/Source/ChaosVD/Private/ChaosVDExtractedGeometryDataHandle.h
+++ b/Engine/Plugins/ChaosVD/Source/ChaosVD/Private/ChaosVDExtractedGeometryDataHandle.h
@@ -27,8 +27,8 @@ struct FChaosVDExtractedGeometryDataHandle
 	uint32 GetGeometryKey() const { return GeometryKey; }
 	void SetGeometryKey(const uint32 Key) { GeometryKey = Key; }
 
-	/** Name to be used to create labels in editor for this data */
-	FName GetName() const;
+	/** Name of the type this underlying implicit object this instance represents */
+	FName GetTypeName() const;
 
 	/** Extracted transform of the generated implicit object in component space */
 	const FTransform& GetRelativeTransform() const { return Transform; }
@@ -37,8 +37,8 @@ struct FChaosVDExtractedGeometryDataHandle
 	/** Ptr to the implicit object this handle provides access to and from which the mesh was generated */
 	const Chaos::FImplicitObject* GetImplicitObject() const { return ImplicitObject; }
 
-	/** Returns the Index on which this Implicit object was in the Hierarchy of the root one */
-	int32 GetImplicitObjectIndex() const { return ObjectIndex; }
+	/** Returns the Index to correct the shape data instance in the shape instances array of the particle that owns this geometry */
+	int32 GetShapeInstanceIndex() const { return ShapeInstanceIndex; }
 
 	/** If the object was an union and this was a leaf, this getter returns the root union object, otherwise it will be the same as GetImplicitObject */
 	const Chaos::FImplicitObject* GetRootImplicitObject() const { return RootImplicitObject; }
@@ -47,7 +47,7 @@ struct FChaosVDExtractedGeometryDataHandle
 	void SetImplicitObject(const Chaos::FImplicitObject* InImplicitObject) { ImplicitObject = InImplicitObject; }
 
 	/** Sets the Index on which this Implicit object was in the Hierarchy of the root one */
-	void SetImplicitObjectIndex(int32 InObjectIndex) { ObjectIndex = InObjectIndex; }
+	void SetShapeInstanceIndex(int32 InObjectIndex) { ShapeInstanceIndex = InObjectIndex; }
 	
 	/** Sets the Implicit Object Ptr from which the data was extracted and the mesh was generated */
 	void SetRootImplicitObject(const Chaos::FImplicitObject* InImplicitObject) { RootImplicitObject = InImplicitObject; }
@@ -71,7 +71,7 @@ private:
 	uint32 DataComponentKey = 0;
 
 	/** Index on which this Implicit object was in the Hierarchy of the root one */
-	int32 ObjectIndex = INDEX_NONE;
+	int32 ShapeInstanceIndex = INDEX_NONE;
 
 	/** Extracted transform of the generated implicit object in component space */
 	FTransform Transform;
diff --git a/Engine/Plugins/ChaosVD/Source/ChaosVD/Private/ChaosVDGeometryBuilder.h b/Engine/Plugins/ChaosVD/Source/ChaosVD/Private/ChaosVDGeometryBuilder.h
--- a/Engine/Plugins/ChaosVD/Source/ChaosVD/Private/ChaosVDGeometryBuilder.h
+++ b/Engine/Plugins/ChaosVD/Source/ChaosVD/Private/ChaosVDGeometryBuilder.h
@@ -121,7 +121,7 @@ public:
 private:
 	
 	template<typename MeshType>
-	void CreateMeshesFromImplicit_Internal(const Chaos::FImplicitObject* InRootImplicitObject,const Chaos::FImplicitObject* InLeafImplicitObject, AActor* Owner, TArray<TSharedPtr<FChaosVDExtractedGeometryDataHandle>>& OutMeshDataHandles, const int32 DesiredLODCount = 0, const Chaos::FRigidTransform3& InTransform = Chaos::FRigidTransform3(), const int32 MeshIndex = 0);
+	void CreateMeshesFromImplicit_Internal(const Chaos::FImplicitObject* InRootImplicitObject,const Chaos::FImplicitObject* InLeafImplicitObject, AActor* Owner, TArray<TSharedPtr<FChaosVDExtractedGeometryDataHandle>>& OutMeshDataHandles, const int32 DesiredLODCount = 0, const Chaos::FRigidTransform3& InTransform = Chaos::FRigidTransform3(), const int32 ShapeInstanceIndex = 0);
 	
 public:
 	/**
@@ -246,7 +246,7 @@ private:
 	 * @return Returns a handle to the generated data that can be used to access the generated mesh when ready
 	 */
 	template <typename MeshType>
-	TSharedPtr<FChaosVDExtractedGeometryDataHandle> ExtractGeometryDataForImplicit(const Chaos::FImplicitObject* InImplicitObject, const Chaos::FRigidTransform3& InTransform, const int32 Index);
+	TSharedPtr<FChaosVDExtractedGeometryDataHandle> ExtractGeometryDataForImplicit(const Chaos::FImplicitObject* InImplicitObject, const Chaos::FRigidTransform3& InTransform);
 
 	/** Returns true if the implicit object if of one of the types we need to unpack before generating a mesh for it */
 	bool ImplicitObjectNeedsUnpacking(const Chaos::FImplicitObject* InImplicitObject) const;
@@ -362,7 +362,7 @@ private:
 };
 
 template <typename MeshType>
-TSharedPtr<FChaosVDExtractedGeometryDataHandle> FChaosVDGeometryBuilder::ExtractGeometryDataForImplicit(const Chaos::FImplicitObject* InImplicitObject, const Chaos::FRigidTransform3& InTransform, const int32 Index)
+TSharedPtr<FChaosVDExtractedGeometryDataHandle> FChaosVDGeometryBuilder::ExtractGeometryDataForImplicit(const Chaos::FImplicitObject* InImplicitObject, const Chaos::FRigidTransform3& InTransform)
 {
 	const uint32 ImplicitObjectHash = InImplicitObject->GetTypeHash();
 
@@ -401,23 +401,28 @@ TSharedPtr<FChaosVDExtractedGeometryDataHandle> FChaosVDGeometryBuilder::Extract
 }
 
 template <typename MeshType>
-void FChaosVDGeometryBuilder::CreateMeshesFromImplicit_Internal(const Chaos::FImplicitObject* InRootImplicitObject, const Chaos::FImplicitObject* InLeafImplicitObject, AActor* Owner, TArray<TSharedPtr<FChaosVDExtractedGeometryDataHandle>>& OutMeshDataHandles, const int32 DesiredLODCount, const Chaos::FRigidTransform3& InTransform, const int32 MeshIndex)
+void FChaosVDGeometryBuilder::CreateMeshesFromImplicit_Internal(const Chaos::FImplicitObject* InRootImplicitObject, const Chaos::FImplicitObject* InLeafImplicitObject, AActor* Owner, TArray<TSharedPtr<FChaosVDExtractedGeometryDataHandle>>& OutMeshDataHandles, const int32 DesiredLODCount, const Chaos::FRigidTransform3& InTransform, const int32 ShapeInstanceIndex)
 {
 	TRACE_CPUPROFILER_EVENT_SCOPE(FChaosVDGeometryBuilder::CreateMeshesFromImplicit_Internal);
 
 	using namespace Chaos;
 
-	const EImplicitObjectType InnerType = GetInnerType(InRootImplicitObject->GetType());
+	const EImplicitObjectType InnerType = GetInnerType(InLeafImplicitObject->GetType());
 	
 	if (InnerType == ImplicitObjectType::Union || InnerType == ImplicitObjectType::UnionClustered)
 	{
 		if (const FImplicitObjectUnion* Union = InLeafImplicitObject->template AsA<FImplicitObjectUnion>())
 		{
+			const bool bIsRootUnion = InRootImplicitObject == InLeafImplicitObject;
+
 			for (int32 ObjectIndex = 0; ObjectIndex < Union->GetObjects().Num(); ++ObjectIndex)
 			{
 				const FImplicitObjectPtr& UnionImplicit = Union->GetObjects()[ObjectIndex];
 
-				CreateMeshesFromImplicitObject<MeshType>(UnionImplicit.GetReference(), Owner, OutMeshDataHandles, DesiredLODCount, InTransform, ObjectIndex);	
+				// If this union it is not the root implicit object, then all its objects will share the same Instance index
+				int32 CurrentShapeInstanceIndex = bIsRootUnion ? ObjectIndex : ShapeInstanceIndex;
+
+				CreateMeshesFromImplicit_Internal<MeshType>(InRootImplicitObject, UnionImplicit.GetReference(), Owner, OutMeshDataHandles, DesiredLODCount, InTransform, CurrentShapeInstanceIndex);	
 			}
 		}
 
@@ -428,16 +433,17 @@ void FChaosVDGeometryBuilder::CreateMeshesFromImplicit_Internal(const Chaos::FIm
 	{
 		if (const TImplicitObjectTransformed<FReal, 3>* Transformed = InLeafImplicitObject->template GetObject<TImplicitObjectTransformed<FReal, 3>>())
 		{
-			CreateMeshesFromImplicitObject<MeshType>(Transformed->GetTransformedObject(), Owner, OutMeshDataHandles, DesiredLODCount, Transformed->GetTransform(), MeshIndex);
+			// For transformed objects, the Instance index is the same so we pass it in without changing it
+			CreateMeshesFromImplicit_Internal<MeshType>(InRootImplicitObject, Transformed->GetTransformedObject(), Owner, OutMeshDataHandles, DesiredLODCount, Transformed->GetTransform(), ShapeInstanceIndex);
 		}
 		
 		return;
 	}
 
-	if (const TSharedPtr<FChaosVDExtractedGeometryDataHandle> MeshDataHandle = ExtractGeometryDataForImplicit<MeshType>(InLeafImplicitObject, InTransform, MeshIndex))
-	{	
+	if (const TSharedPtr<FChaosVDExtractedGeometryDataHandle> MeshDataHandle = ExtractGeometryDataForImplicit<MeshType>(InLeafImplicitObject, InTransform))
+	{
 		MeshDataHandle->SetImplicitObject(InLeafImplicitObject);
-		MeshDataHandle->SetImplicitObjectIndex(MeshIndex);
+		MeshDataHandle->SetShapeInstanceIndex(ShapeInstanceIndex);
 		MeshDataHandle->SetRootImplicitObject(InRootImplicitObject);
 
 		OutMeshDataHandles.Add(MeshDataHandle);
@@ -519,7 +525,7 @@ ComponentType* FChaosVDGeometryBuilder::GetAvailableInstancedStaticMeshComponent
 	else
 	{
 		// If no exiting component meets our requirements, get a new one form the pool
-		ComponentType* Component = ComponentMeshPool.AcquireMeshComponent<ComponentType>(MeshComponentsContainerActor, InExtractedGeometryDataHandle->GetName());
+		ComponentType* Component = ComponentMeshPool.AcquireMeshComponent<ComponentType>(MeshComponentsContainerActor, InExtractedGeometryDataHandle->GetTypeName());
 		if (!InitializeMeshComponent<ComponentType>(MeshComponentsContainerActor, Component))
 		{
 			return nullptr;
@@ -551,7 +557,7 @@ ComponentType* FChaosVDGeometryBuilder::GetAvailableMeshComponent(const TSharedP
 	}
 	else
 	{
-		MeshComponent = ComponentMeshPool.AcquireMeshComponent<ComponentType>(MeshComponentsContainerActor, InExtractedGeometryDataHandle->GetName());
+		MeshComponent = ComponentMeshPool.AcquireMeshComponent<ComponentType>(MeshComponentsContainerActor, InExtractedGeometryDataHandle->GetTypeName());
 		if (!InitializeMeshComponent<ComponentType>(MeshComponentsContainerActor, MeshComponent))
 		{
 			return nullptr;
diff --git a/Engine/Plugins/ChaosVD/Source/ChaosVD/Private/ChaosVDGeometryDataComponent.cpp b/Engine/Plugins/ChaosVD/Source/ChaosVD/Private/ChaosVDGeometryDataComponent.cpp
--- a/Engine/Plugins/ChaosVD/Source/ChaosVD/Private/ChaosVDGeometryDataComponent.cpp
+++ b/Engine/Plugins/ChaosVD/Source/ChaosVD/Private/ChaosVDGeometryDataComponent.cpp
@@ -13,22 +13,22 @@
 
 FChaosVDMeshDataInstanceHandle::FChaosVDMeshDataInstanceHandle(int32 InInstanceIndex, UMeshComponent* InMeshComponent, int32 InParticleID, int32 InSolverID)
 {
-	MeshComponent = InMeshComponent;
-	MeshInstanceIndex = InInstanceIndex;
-	OwningParticleID = InParticleID;
-	OwningSolverID = InSolverID;
+	InstanceState.MeshComponent = InMeshComponent;
+	InstanceState.MeshInstanceIndex = InInstanceIndex;
+	InstanceState.OwningParticleID = InParticleID;
+	InstanceState.OwningSolverID = InSolverID;
 
-	if (Cast<UInstancedStaticMeshComponent>(MeshComponent))
+	if (Cast<UInstancedStaticMeshComponent>(InstanceState.MeshComponent))
 	{
-		MeshComponentType = EChaosVDMeshComponent::InstancedStatic;
+		InstanceState.MeshComponentType = EChaosVDMeshComponent::InstancedStatic;
 	}
-	else if (Cast<UStaticMeshComponent>(MeshComponent))
+	else if (Cast<UStaticMeshComponent>(InstanceState.MeshComponent))
 	{
-		MeshComponentType = EChaosVDMeshComponent::Static;
+		InstanceState.MeshComponentType = EChaosVDMeshComponent::Static;
 	}
 	else
 	{
-		MeshComponentType = EChaosVDMeshComponent::Dynamic;
+		InstanceState.MeshComponentType = EChaosVDMeshComponent::Dynamic;
 	}
 }
 
@@ -42,31 +42,49 @@ void FChaosVDMeshDataInstanceHandle::SetWorldTransform(const FTransform& InTrans
 
 	const FTransform ExtractedRelativeTransform = ExtractedGeometryHandle->GetRelativeTransform();
 
-	CurrentWorldTransform.SetLocation(InTransform.TransformPosition(ExtractedRelativeTransform.GetLocation()));
-	CurrentWorldTransform.SetRotation(InTransform.TransformRotation(ExtractedRelativeTransform.GetRotation()));
-	CurrentWorldTransform.SetScale3D(ExtractedGeometryHandle->GetRelativeTransform().GetScale3D());
+	InstanceState.CurrentWorldTransform.SetLocation(InTransform.TransformPosition(ExtractedRelativeTransform.GetLocation()));
+	InstanceState.CurrentWorldTransform.SetRotation(InTransform.TransformRotation(ExtractedRelativeTransform.GetRotation()));
+	InstanceState.CurrentWorldTransform.SetScale3D(ExtractedRelativeTransform.GetScale3D());
 	
 	if (IChaosVDGeometryComponent* CVDGeometryComponent = Cast<IChaosVDGeometryComponent>(GetMeshComponent()))
 	{
-		CVDGeometryComponent->UpdateInstanceWorldTransform(AsShared(), CurrentWorldTransform);
+		CVDGeometryComponent->UpdateInstanceWorldTransform(AsShared(), InstanceState.CurrentWorldTransform);
+	}
+}
+
+void FChaosVDMeshDataInstanceHandle::SetGeometryHandle(const TSharedPtr<FChaosVDExtractedGeometryDataHandle>& InHandle)
+{
+	ExtractedGeometryHandle = InHandle;
+
+	using namespace Chaos;
+	if (ExtractedGeometryHandle)
+	{
+		InstanceState.ImplicitObjectInfo.bIsRootObject = ExtractedGeometryHandle->GetRootImplicitObject() == ExtractedGeometryHandle->GetImplicitObject();
+		InstanceState.ImplicitObjectInfo.ShapeInstanceIndex = ExtractedGeometryHandle->GetShapeInstanceIndex();
+		InstanceState.ImplicitObjectInfo.ImplicitObjectType = ExtractedGeometryHandle->GetTypeName();
+		InstanceState.ImplicitObjectInfo.RelativeTransform = ExtractedGeometryHandle->GetRelativeTransform();
+	}
+	else
+	{
+		InstanceState.ImplicitObjectInfo = FChaosVDImplicitObjectBasicView();
 	}
 }
 
 void FChaosVDMeshDataInstanceHandle::SetInstanceColor(const FLinearColor& NewColor)
 {
-	if (CurrentGeometryColor != NewColor)
+	if (InstanceState.CurrentGeometryColor != NewColor)
 	{
 		if (IChaosVDGeometryComponent* CVDGeometryComponent = Cast<IChaosVDGeometryComponent>(GetMeshComponent()))
 		{
 			CVDGeometryComponent->UpdateInstanceColor(AsShared(), NewColor);
-			CurrentGeometryColor = NewColor;
+			InstanceState.CurrentGeometryColor = NewColor;
 		}
 	}
 }
 
 void FChaosVDMeshDataInstanceHandle::UpdateMeshComponentForCollisionData(const FChaosVDShapeCollisionData& InCollisionData)
 {
-	if (InCollisionData.bIsValid && CollisionData != InCollisionData)
+	if (InCollisionData.bIsValid && InstanceState.CollisionData != InCollisionData)
 	{
 		if (const TSharedPtr<FChaosVDGeometryBuilder> GeometryBuilderPtr = GeometryBuilderInstance.Pin())
 		{
@@ -90,7 +108,7 @@ void FChaosVDMeshDataInstanceHandle::UpdateMeshComponentForCollisionData(const F
 			{
 				if (RequiredMeshAttributes != CVDOldGeometryComponent->GetMeshComponentAttributeFlags())
 				{
-					if (bIsSelected)
+					if (InstanceState.bIsSelected)
 					{
 						CVDOldGeometryComponent->SetIsSelected(AsShared(), false);
 					}
@@ -108,10 +126,10 @@ void FChaosVDMeshDataInstanceHandle::UpdateMeshComponentForCollisionData(const F
 				if (IChaosVDGeometryComponent* CVDNewGeometryComponent = Cast<IChaosVDGeometryComponent>(GetMeshComponent()))
 				{
 					// Reset the color so it is updated in the next Update color calls (which always happens after updating the shape instance data)
-					CurrentGeometryColor = FLinearColor(ForceInitToZero);
+					InstanceState.CurrentGeometryColor = FLinearColor(ForceInitToZero);
 		
-					CVDNewGeometryComponent->UpdateInstanceVisibility(AsShared(), bIsVisible);
-					CVDNewGeometryComponent->SetIsSelected(AsShared(), bIsSelected);
+					CVDNewGeometryComponent->UpdateInstanceVisibility(AsShared(), InstanceState.bIsVisible);
+					CVDNewGeometryComponent->SetIsSelected(AsShared(), InstanceState.bIsSelected);
 				}
 			}
 		}
@@ -127,7 +145,7 @@ void FChaosVDMeshDataInstanceHandle::SetGeometryCollisionData(const FChaosVDShap
 		UpdateMeshComponentForCollisionData(InCollisionData);
 	}
 
-	CollisionData = InCollisionData;
+	InstanceState.CollisionData = InCollisionData;
 }
 
 void FChaosVDMeshDataInstanceHandle::SetIsSelected(bool bInIsSelected)
@@ -137,7 +155,7 @@ void FChaosVDMeshDataInstanceHandle::SetIsSelected(bool bInIsSelected)
 		CVDGeometryComponent->SetIsSelected(AsShared(), bInIsSelected);
 	}
 
-	bIsSelected = bInIsSelected;
+	InstanceState.bIsSelected = bInIsSelected;
 }
 
 void FChaosVDMeshDataInstanceHandle::SetVisibility(bool bInIsVisible)
@@ -147,7 +165,7 @@ void FChaosVDMeshDataInstanceHandle::SetVisibility(bool bInIsVisible)
 		CVDGeometryComponent->UpdateInstanceVisibility(AsShared(), bInIsVisible);
 	}
 
-	bIsVisible = bInIsVisible;
+	InstanceState.bIsVisible = bInIsVisible;
 }
 
 void FChaosVDMeshDataInstanceHandle::HandleInstanceIndexUpdated(TArrayView<const FInstancedStaticMeshDelegates::FInstanceIndexUpdateData> InIndexUpdates)
@@ -161,9 +179,9 @@ void FChaosVDMeshDataInstanceHandle::HandleInstanceIndexUpdated(TArrayView<const
 				break; // We don't need to process 'Added' updates as they can't affect existing IDs
 			case FInstancedStaticMeshDelegates::EInstanceIndexUpdateType::Relocated:
 				{
-					if (MeshInstanceIndex == IndexUpdateData.OldIndex)
+					if (InstanceState.MeshInstanceIndex == IndexUpdateData.OldIndex)
 					{
-						MeshInstanceIndex = IndexUpdateData.Index;
+						InstanceState.MeshInstanceIndex = IndexUpdateData.Index;
 					}
 					break;
 				}
@@ -172,9 +190,9 @@ void FChaosVDMeshDataInstanceHandle::HandleInstanceIndexUpdated(TArrayView<const
 			case FInstancedStaticMeshDelegates::EInstanceIndexUpdateType::Cleared:
 			case FInstancedStaticMeshDelegates::EInstanceIndexUpdateType::Destroyed:
 				{
-					if (MeshInstanceIndex == IndexUpdateData.Index)
+					if (InstanceState.MeshInstanceIndex == IndexUpdateData.Index)
 					{
-						MeshInstanceIndex = INDEX_NONE;
+						InstanceState.MeshInstanceIndex = INDEX_NONE;
 					}
 					break;
 				}
@@ -184,39 +202,47 @@ void FChaosVDMeshDataInstanceHandle::HandleInstanceIndexUpdated(TArrayView<const
 	}
 }
 
-void FChaosVDGeometryComponentUtils::UpdateCollisionDataFromShapeArray(const TArray<FChaosVDShapeCollisionData>& InShapeArray, const TSharedPtr<FChaosVDMeshDataInstanceHandle>& InInstanceHandle)
+void FChaosVDGeometryComponentUtils::UpdateCollisionDataFromShapeArray(const TArray<FChaosVDShapeCollisionData>& InShapeArray, const TSharedRef<FChaosVDMeshDataInstanceHandle>& InInstanceHandle)
 {
-	if (!ensureMsgf(InInstanceHandle && InInstanceHandle->GetGeometryHandle() && InInstanceHandle->GetGeometryHandle()->GetRootImplicitObject(), TEXT("Tried to Update Collision Data without a valid Implicit Object")))
+	if (InShapeArray.IsEmpty())
 	{
 		return;
 	}
 
-	FChaosVDShapeCollisionData CollisionDataToUpdate = InInstanceHandle->GetGeometryCollisionData();
+	const TSharedPtr<FChaosVDExtractedGeometryDataHandle>& ExtractedGeometryHandle = InInstanceHandle->GetGeometryHandle();
 
-	InInstanceHandle->GetGeometryHandle()->GetRootImplicitObject()->VisitObjects([InInstanceHandle, &CollisionDataToUpdate, &InShapeArray] (const Chaos::FImplicitObject* ImplicitA, const Chaos::FRigidTransform3& RelativeTransformA, const int32 RootObjectIndexA, const int32 ObjectIndex, const int32 LeafObjectIndexA)
+	if (!ExtractedGeometryHandle)
 	{
-		if (!InShapeArray.IsValidIndex(LeafObjectIndexA))
-		{
-			return true;
-		}
+		return;
+	}
 
-		if (ImplicitA == InInstanceHandle->GetGeometryHandle()->GetImplicitObject())
-		{
-			
-			CollisionDataToUpdate = InShapeArray[LeafObjectIndexA];
-			CollisionDataToUpdate.bIsComplex = FChaosVDGeometryBuilder::DoesImplicitContainType(ImplicitA, Chaos::ImplicitObjectType::HeightField) || FChaosVDGeometryBuilder::DoesImplicitContainType(ImplicitA, Chaos::ImplicitObjectType::TriangleMesh);
-			CollisionDataToUpdate.bIsValid = true;
-		}
+	const int32 ShapeInstanceIndex = ExtractedGeometryHandle->GetShapeInstanceIndex();
+	if (!InShapeArray.IsValidIndex(ShapeInstanceIndex))
+	{
+		const FName ImplicitObjectTypeName = InInstanceHandle->GetState().ImplicitObjectInfo.ImplicitObjectType;
+		const Chaos::FImplicitObject* RootImplicitObject = ExtractedGeometryHandle->GetRootImplicitObject();
+		const FName RootImplicitObjectTypeName = !InInstanceHandle->GetState().ImplicitObjectInfo.bIsRootObject && RootImplicitObject ? Chaos::GetImplicitObjectTypeName(Chaos::GetInnerType(RootImplicitObject->GetType())) : TEXT("None");
+		UE_LOG(LogChaosVDEditor, Warning, TEXT("[%s] Failed to find shape instance data at Index [%d] | Particle ID[%d] | Available Shape instance Data Num [%d] | Implicit Type [%s] - Root Implicit Type [%s] | This geometry will be hidden..."), ANSI_TO_TCHAR(__FUNCTION__), ShapeInstanceIndex, InInstanceHandle->GetOwningParticleID(), InShapeArray.Num(), *ImplicitObjectTypeName.ToString(), *RootImplicitObjectTypeName.ToString());
+
+		InInstanceHandle->bFailedToUpdateShapeInstanceData = true;
+		return;
+	}
+	else if (InInstanceHandle->bFailedToUpdateShapeInstanceData)
+	{
+		InInstanceHandle->bFailedToUpdateShapeInstanceData = false;
+		UE_LOG(LogChaosVDEditor, Warning, TEXT("[%s] Recovered from failing to find shape instance data at Index [%d] | Particle ID[%d] | Available Shape instance Data Num [%d] | This geometry will be shown again..."), ANSI_TO_TCHAR(__FUNCTION__), ShapeInstanceIndex, InInstanceHandle->GetOwningParticleID(), InShapeArray.Num());
+	}
 
-		return true;
-	});
+	FChaosVDShapeCollisionData CollisionDataToUpdate = InShapeArray[ShapeInstanceIndex];
+	CollisionDataToUpdate.bIsComplex = FChaosVDGeometryBuilder::DoesImplicitContainType(ExtractedGeometryHandle->GetImplicitObject(), Chaos::ImplicitObjectType::HeightField) || FChaosVDGeometryBuilder::DoesImplicitContainType(ExtractedGeometryHandle->GetImplicitObject(), Chaos::ImplicitObjectType::TriangleMesh);
+	CollisionDataToUpdate.bIsValid = true;
 
 	InInstanceHandle->SetGeometryCollisionData(CollisionDataToUpdate);
 }
 
 void FChaosVDGeometryComponentUtils::UpdateMeshColor(const TSharedPtr<FChaosVDMeshDataInstanceHandle>& InInstanceHandle, const FChaosVDParticleDataWrapper& InParticleData, bool bIsServer)
 {
-	const FChaosVDShapeCollisionData ShapeData = InInstanceHandle->GetGeometryCollisionData();
+	const FChaosVDShapeCollisionData& ShapeData = InInstanceHandle->GetGeometryCollisionData();
 	const bool bIsQueryOnly = ShapeData.bQueryCollision && !ShapeData.bSimCollision;
 
 	if (ShapeData.bIsValid)
@@ -269,7 +295,7 @@ void FChaosVDGeometryComponentUtils::UpdateMeshVisibility(const TSharedPtr<FChao
 		}
 		else
 		{
-			const FChaosVDShapeCollisionData InstanceShapeData = InInstanceHandle->GetGeometryCollisionData();
+			const FChaosVDShapeCollisionData& InstanceShapeData = InInstanceHandle->GetGeometryCollisionData();
 
 			if (InstanceShapeData.bIsValid)
 			{
@@ -386,3 +412,14 @@ UMaterialInstanceDynamic* FChaosVDGeometryComponentUtils::CreateMaterialInstance
 {
 	return CreateMaterialInstance(GetBaseMaterialForType(Type));
 }
+
+void Chaos::VisualDebugger::SelectParticleWithGeometryInstance(const TSharedRef<FChaosVDScene>& InScene, IChaosVDGeometryOwnerInterface* GeometryOwner, const TSharedPtr<FChaosVDMeshDataInstanceHandle>& InMeshDataHandle)
+{
+	InScene->SetSelectedObject(nullptr);
+
+	if(GeometryOwner)
+	{
+		GeometryOwner->SetSelectedMeshInstance(InMeshDataHandle);
+		InScene->SetSelectedObject(Cast<UObject>(GeometryOwner));	
+	}
+}
diff --git a/Engine/Plugins/ChaosVD/Source/ChaosVD/Private/ChaosVDGeometryDataComponent.h b/Engine/Plugins/ChaosVD/Source/ChaosVD/Private/ChaosVDGeometryDataComponent.h
--- a/Engine/Plugins/ChaosVD/Source/ChaosVD/Private/ChaosVDGeometryDataComponent.h
+++ b/Engine/Plugins/ChaosVD/Source/ChaosVD/Private/ChaosVDGeometryDataComponent.h
@@ -16,6 +16,8 @@
 #include "Templates/SharedPointer.h"
 #include "Components/MeshComponent.h"
 
+class FChaosVDScene;
+
 #include "ChaosVDGeometryDataComponent.generated.h"
 
 class UChaosVDStaticMeshComponent;
@@ -51,18 +53,87 @@ enum class EChaosVDMeshComponent
 	Dynamic
 };
 
+/** Struct holding the a minimum amount of data about a Implicit object to be shown in the details panel */
+USTRUCT()
+struct FChaosVDImplicitObjectBasicView
+{
+	GENERATED_BODY()
+
+	/** Geometry type name*/
+	UPROPERTY(VisibleAnywhere, Category="Recorded GeometryData")
+	FName ImplicitObjectType;
+
+	/** Index of the Shape Instance data for this geometry in the Shape Instance data array */
+	UPROPERTY(VisibleAnywhere, Category="Recorded GeometryData")
+	int32 ShapeInstanceIndex = INDEX_NONE;
+
+	/** True if this is the root implicit object */
+	UPROPERTY(VisibleAnywhere, Category="Recorded GeometryData")
+	bool bIsRootObject = false;
+
+	/** If this is a transformed implicit, this will contain the recorded relative transform */
+	UPROPERTY(VisibleAnywhere, Category="Recorded GeometryData")
+	FTransform RelativeTransform;
+};
+
+/** Struct holding the state of a mesh instance - Is separated from the Mesh instance class so we can show the data in the Details panel */
+USTRUCT()
+struct FChaosVDMeshDataInstanceState
+{
+	GENERATED_BODY()
+
+	/** Recorded Shape instance Data */
+	UPROPERTY(VisibleAnywhere, Category="Recorded GeometryData", meta=(ShowOnlyInnerProperties))
+	FChaosVDShapeCollisionData CollisionData;
+
+	/** Minimum set of data about the recorded implicit object */
+	UPROPERTY(VisibleAnywhere, Category="Recorded GeometryData", meta=(ShowOnlyInnerProperties))
+	FChaosVDImplicitObjectBasicView ImplicitObjectInfo;
+
+	/* CVD Debug - Current world transform used to render this Mesh */
+	UPROPERTY(VisibleAnywhere, Category="CVD GeometryData")
+	FTransform CurrentWorldTransform;
+
+	/* CVD Debug - Current mesh component type to render this Mesh */
+	UPROPERTY(VisibleAnywhere, Category="CVD GeometryData")
+	EChaosVDMeshComponent MeshComponentType = EChaosVDMeshComponent::Invalid;
+
+	/* CVD Debug - Pointer to the mesh component used to render this Mesh */
+	UPROPERTY(VisibleAnywhere, Category="CVD GeometryData")
+	TObjectPtr<UMeshComponent> MeshComponent;
+
+	/* CVD Debug - Instance index of mesh component used to render this Mesh */
+	UPROPERTY(VisibleAnywhere, Category="CVD GeometryData")
+	int32 MeshInstanceIndex = INDEX_NONE;
+
+	/* CVD Debug - Color used to render this mesh */
+	UPROPERTY(VisibleAnywhere, Category="CVD GeometryData")
+	FLinearColor CurrentGeometryColor = FLinearColor(ForceInitToZero);
+
+	/* CVD Debug - Id of the particle this geometry belongs */
+	UPROPERTY(VisibleAnywhere, Category="CVD GeometryData")
+	int32 OwningParticleID = INDEX_NONE;
+
+	/* CVD Debug - Id of the solver this geometry belongs */
+	UPROPERTY(VisibleAnywhere, Category="CVD GeometryData")
+	int32 OwningSolverID = INDEX_NONE;
+
+	bool bIsVisible = true;
+
+	bool bIsSelected = false;
+};
+
 /** Handle that provides access to a specific mesh instance on a CVD Mesh component (instanced or static)*/
 class FChaosVDMeshDataInstanceHandle : public TSharedFromThis<FChaosVDMeshDataInstanceHandle>
 {
 public:
-
 	explicit FChaosVDMeshDataInstanceHandle(int32 InInstanceIndex, UMeshComponent* InMeshComponent, int32 InParticleID, int32 InSolverID);
 
 	/** Returns the Particle ID of the particle owning this mesh instance */
-	int32 GetOwningParticleID() const { return OwningParticleID; }
+	int32 GetOwningParticleID() const { return InstanceState.OwningParticleID; }
 
 	/** Returns the Solver ID of the particle owning this mesh instance */
-	int32 GetOwningSolverID() const { return OwningSolverID; }
+	int32 GetOwningSolverID() const { return InstanceState.OwningSolverID; }
 
 	/** Applies the provided world transform to the mesh instance this handle represents
 	 * @param InTransform Transform to apply
@@ -70,11 +141,11 @@ public:
 	void SetWorldTransform(const FTransform& InTransform);
 
 	/** Returns the world transform of the mesh instance this handle represents */
-	const FTransform& GetWorldTransform() const { return CurrentWorldTransform; }
+	const FTransform& GetWorldTransform() const { return InstanceState.CurrentWorldTransform; }
 
 	/** Sets the Geometry Handle used to create the mesh instance this handle represents */
-	void SetGeometryHandle(const TSharedPtr<FChaosVDExtractedGeometryDataHandle>& InHandle) { ExtractedGeometryHandle = InHandle; };
-	
+	void SetGeometryHandle(const TSharedPtr<FChaosVDExtractedGeometryDataHandle>& InHandle);
+
 	/** Returns the geometry handle used to create the mesh instance this handle represents */
 	TSharedPtr<FChaosVDExtractedGeometryDataHandle> GetGeometryHandle() const { return ExtractedGeometryHandle; }
 
@@ -82,19 +153,19 @@ public:
 	void SetInstanceColor(const FLinearColor& NewColor);
 
 	/** Returns the current color of the mesh instance this handle represents */
-	FLinearColor GetInstanceColor() const { return CurrentGeometryColor; }
+	FLinearColor GetInstanceColor() const { return InstanceState.CurrentGeometryColor; }
 	
 	/** Applies to the provided shape collision data to the mesh instance this handle represents */
 	void UpdateMeshComponentForCollisionData(const FChaosVDShapeCollisionData& InCollisionData);
 
 	/** Returns the mesh component used to render the mesh instance this handle represents */
-	UMeshComponent* GetMeshComponent() const { return MeshComponent; }
+	UMeshComponent* GetMeshComponent() const { return InstanceState.MeshComponent; }
 
 	/** Returns the instance index of the mesh instance this handle represents */
-	int32 GetMeshInstanceIndex() const { return MeshInstanceIndex; }
+	int32 GetMeshInstanceIndex() const { return InstanceState.MeshInstanceIndex; }
 
 	/** Returns the type of the component used to render the mesh instance this handle represents */
-	EChaosVDMeshComponent GetMeshComponentType() const { return MeshComponentType; }
+	EChaosVDMeshComponent GetMeshComponentType() const { return InstanceState.MeshComponentType; }
 
 	/** Sets a Ptr to the geometry builder used to generate and manage the geometry/mesh components this handle represents */
 	void SetGeometryBuilder(const TWeakPtr<FChaosVDGeometryBuilder>& InGeometryBuilder) { GeometryBuilderInstance = InGeometryBuilder; }
@@ -106,53 +177,41 @@ public:
 	void SetVisibility(bool bInIsVisible);
 
 	/** Returns the current visibility state this mesh instance */
-	bool GetVisibility() const { return bIsVisible; }
+	bool GetVisibility() const { return InstanceState.bIsVisible; }
 
 	/** Applies a new shape collision data to this mesh instance */
 	void SetGeometryCollisionData(const FChaosVDShapeCollisionData& InCollisionData);
 
-	/** Returns a copy of the current shape collision data to this mesh instance */
-	FChaosVDShapeCollisionData GetGeometryCollisionData() const { return CollisionData; }
+	/** Returns the current shape collision data to this mesh instance */
+	FChaosVDShapeCollisionData& GetGeometryCollisionData() { return InstanceState.CollisionData; }
 
 	/** Handles a mesh instance index update reported by the mesh component used to render this mesh instance */
 	void HandleInstanceIndexUpdated(TArrayView<const FInstancedStaticMeshDelegates::FInstanceIndexUpdateData> InIndexUpdates);
 
+	FChaosVDMeshDataInstanceState& GetState() { return InstanceState; }
+
 	friend uint32 GetTypeHash(const FChaosVDMeshDataInstanceHandle& Handle)
 	{
 		const uint32 GeometryHandleHash = Handle.ExtractedGeometryHandle ? GetTypeHash(Handle.ExtractedGeometryHandle->GetGeometryKey()) : 0;
-		const uint32 MeshComponentHandleHash = HashCombine(HashCombine(Handle.OwningParticleID, Handle.MeshInstanceIndex), static_cast<uint32>(Handle.MeshComponentType));
+		const uint32 MeshComponentHandleHash = HashCombine(HashCombine(Handle.GetOwningParticleID(), Handle.GetMeshInstanceIndex()), static_cast<uint32>(Handle.GetMeshComponentType()));
 		return HashCombine(GeometryHandleHash, MeshComponentHandleHash);
 	}
 
+	/** Used only for debugging purposes - It will be set to true if we received new Shape Instance data but the Shape Index for the implicit object we represent is not valid */
+	bool bFailedToUpdateShapeInstanceData = false; 
+
 private:
 
 	/** Sets the mesh component used to render the mesh instance this handle represents */
-	void SetMeshComponent(UMeshComponent* NewComponent) { MeshComponent = NewComponent; }
+	void SetMeshComponent(UMeshComponent* NewComponent) { InstanceState.MeshComponent = NewComponent; }
 
 	/** Sets the mesh instance index of the mesh instance this handle represents */
-	void SetMeshInstanceIndex(int32 NewIndex) { MeshInstanceIndex = NewIndex; }
+	void SetMeshInstanceIndex(int32 NewIndex) { InstanceState.MeshInstanceIndex = NewIndex; }
 
-	int32 OwningParticleID = INDEX_NONE;
-	int32 OwningSolverID = INDEX_NONE;
-
-	int32 MeshInstanceIndex = INDEX_NONE;
-
-	EChaosVDMeshComponent MeshComponentType = EChaosVDMeshComponent::Invalid;
+	FChaosVDMeshDataInstanceState InstanceState;
 
 	TSharedPtr<FChaosVDExtractedGeometryDataHandle> ExtractedGeometryHandle;
 
-	UMeshComponent* MeshComponent;
-
-	FChaosVDShapeCollisionData CollisionData;
-
-	FLinearColor CurrentGeometryColor = FLinearColor(ForceInitToZero);
-
-	FTransform CurrentWorldTransform;
-
-	bool bIsVisible = true;
-
-	bool bIsSelected = false;
-
 	TWeakPtr<FChaosVDGeometryBuilder> GeometryBuilderInstance = nullptr;
 
 	friend UChaosVDInstancedStaticMeshComponent;
@@ -259,7 +318,7 @@ class FChaosVDGeometryComponentUtils
 {
 public:
 	/** Finds and updates the Shape data using the provided array as source*/
-	static void UpdateCollisionDataFromShapeArray(const TArray<FChaosVDShapeCollisionData>& InShapeArray, const TSharedPtr<FChaosVDMeshDataInstanceHandle>& InInstanceHandle);
+	static void UpdateCollisionDataFromShapeArray(const TArray<FChaosVDShapeCollisionData>& InShapeArray, const TSharedRef<FChaosVDMeshDataInstanceHandle>& InInstanceHandle);
 
 	/** Calculates and updates the color used to render the mesh represented by the provided handle, based on the particle state */
 	static void UpdateMeshColor(const TSharedPtr<FChaosVDMeshDataInstanceHandle>& InInstanceHandle, const FChaosVDParticleDataWrapper& InParticleData, bool bIsServer);
@@ -299,3 +358,26 @@ EChaosVDMaterialType FChaosVDGeometryComponentUtils::GetMaterialTypeForComponent
 		return bIsInstancedMeshComponent ? EChaosVDMaterialType::Instanced : EChaosVDMaterialType::SimOnlyMaterial;
 	}
 }
+
+UINTERFACE()
+class UChaosVDGeometryOwnerInterface : public UInterface
+{
+	GENERATED_BODY()
+};
+
+/**
+ * Interface with for any CVD object that owns and handles Geometry generated from recorded data.
+ */
+class IChaosVDGeometryOwnerInterface
+{
+	GENERATED_BODY()
+public:
+	virtual void SetSelectedMeshInstance(const TWeakPtr<FChaosVDMeshDataInstanceHandle>& InMeshInstanceToSelect) PURE_VIRTUAL(IChaosVDGeometryOwnerInterface::SetSelectedMeshInstance);
+	virtual TWeakPtr<FChaosVDMeshDataInstanceHandle> GetSelectedMeshInstance() const PURE_VIRTUAL(IChaosVDGeometryOwnerInterface::GetSelectedMeshInstance, return nullptr;);
+	virtual TConstArrayView<TSharedPtr<FChaosVDMeshDataInstanceHandle>> GetMeshInstances() const PURE_VIRTUAL(IChaosVDGeometryOwnerInterface::GetMeshInstances, return TConstArrayView<TSharedPtr<FChaosVDMeshDataInstanceHandle>>(););
+};
+
+namespace Chaos::VisualDebugger
+{
+	void SelectParticleWithGeometryInstance(const TSharedRef<FChaosVDScene>& InScene, IChaosVDGeometryOwnerInterface* GeometryOwner, const TSharedPtr<FChaosVDMeshDataInstanceHandle>& InMeshDataHandle);
+}
diff --git a/Engine/Plugins/ChaosVD/Source/ChaosVD/Private/ChaosVDObjectDetailsTab.h b/Engine/Plugins/ChaosVD/Source/ChaosVD/Private/ChaosVDObjectDetailsTab.h
--- a/Engine/Plugins/ChaosVD/Source/ChaosVD/Private/ChaosVDObjectDetailsTab.h
+++ b/Engine/Plugins/ChaosVD/Source/ChaosVD/Private/ChaosVDObjectDetailsTab.h
@@ -24,6 +24,8 @@ public:
 	FChaosVDObjectDetailsTab(const FName& InTabID, TSharedPtr<FTabManager> InTabManager, TWeakPtr<SChaosVDMainTab> InOwningTabWidget) : FChaosVDTabSpawnerBase(InTabID, InTabManager, InOwningTabWidget)
 	{
 	}
+	
+	TSharedPtr<SChaosVDDetailsView> GetDetailsPanel() { return DetailsPanelView; }
 
 protected:
 
diff --git a/Engine/Plugins/ChaosVD/Source/ChaosVD/Private/ChaosVDParticleActor.cpp b/Engine/Plugins/ChaosVD/Source/ChaosVD/Private/ChaosVDParticleActor.cpp
--- a/Engine/Plugins/ChaosVD/Source/ChaosVD/Private/ChaosVDParticleActor.cpp
+++ b/Engine/Plugins/ChaosVD/Source/ChaosVD/Private/ChaosVDParticleActor.cpp
@@ -112,6 +112,36 @@ void AChaosVDParticleActor::UpdateFromRecordedParticleData(const TSharedPtr<FCha
 	OnParticleDataUpdated().ExecuteIfBound();
 }
 
+void AChaosVDParticleActor::SetSelectedMeshInstance(const TWeakPtr<FChaosVDMeshDataInstanceHandle>& GeometryInstanceToSelect)
+{
+	if (!ParticleDataPtr)
+	{
+		return;
+	}
+
+	const TSharedPtr<FChaosVDMeshDataInstanceHandle> GeometryInstanceToSelectPtr = GeometryInstanceToSelect.Pin();
+	if (!GeometryInstanceToSelectPtr)
+	{
+		return;
+	}
+
+	if (ensure(ParticleDataPtr->ParticleIndex == GeometryInstanceToSelectPtr->GetOwningParticleID()))
+	{
+		CurrentSelectedGeometryInstance = GeometryInstanceToSelect;
+	}
+}
+
+void AChaosVDParticleActor::HandleSelected()
+{
+	PushSelectionToProxies();
+}
+
+void AChaosVDParticleActor::HandleDeSelected()
+{
+	CurrentSelectedGeometryInstance = nullptr;
+	PushSelectionToProxies();
+}
+
 void AChaosVDParticleActor::ProcessUpdatedAndRemovedHandles(TArray<TSharedPtr<FChaosVDExtractedGeometryDataHandle>>& OutExtractedGeometryDataHandles)
 {
 	for (TArray<TSharedPtr<FChaosVDMeshDataInstanceHandle>>::TIterator MeshDataHandleRemoveIterator = MeshDataHandles.CreateIterator(); MeshDataHandleRemoveIterator; ++MeshDataHandleRemoveIterator)
@@ -405,17 +435,24 @@ FName AChaosVDParticleActor::GetProviderName()
 	return GetFName();
 }
 
-void AChaosVDParticleActor::PushSelectionToProxies()
+void AChaosVDParticleActor::UpdateMeshInstancesSelectionState()
 {
-	Super::PushSelectionToProxies();
-
-	VisitGeometryInstances([this](const TSharedRef<FChaosVDMeshDataInstanceHandle>& MeshDataHandle)
+	TSharedPtr<FChaosVDMeshDataInstanceHandle> CurrentSelectedGeometry = CurrentSelectedGeometryInstance.Pin();
+	const bool bIsOwningParticleSelectedInEditor = IsSelectedInEditor();
+	VisitGeometryInstances([this, bIsOwningParticleSelectedInEditor, CurrentSelectedGeometry](const TSharedRef<FChaosVDMeshDataInstanceHandle>& MeshDataHandle)
 	{
-		const bool bIsSelectedInEditor = IsSelectedInEditor();
-		MeshDataHandle->SetIsSelected(bIsSelectedInEditor);
+		const bool bShouldSelectInstance = bIsOwningParticleSelectedInEditor ? (CurrentSelectedGeometry ? CurrentSelectedGeometryInstance == MeshDataHandle : true) : false;
+
+		MeshDataHandle->SetIsSelected(bShouldSelectInstance);
 	});
 }
 
+void AChaosVDParticleActor::PushSelectionToProxies()
+{
+	UpdateMeshInstancesSelectionState();
+	Super::PushSelectionToProxies();
+}
+
 const TArray<TSharedPtr<FChaosVDParticlePairMidPhase>>* AChaosVDParticleActor::GetCollisionMidPhasesArray() const
 {
 	if (!ParticleDataPtr.IsValid())
diff --git a/Engine/Plugins/ChaosVD/Source/ChaosVD/Private/ChaosVDParticleActor.h b/Engine/Plugins/ChaosVD/Source/ChaosVD/Private/ChaosVDParticleActor.h
--- a/Engine/Plugins/ChaosVD/Source/ChaosVD/Private/ChaosVDParticleActor.h
+++ b/Engine/Plugins/ChaosVD/Source/ChaosVD/Private/ChaosVDParticleActor.h
@@ -8,6 +8,7 @@
 #include "DataWrappers/ChaosVDParticleDataWrapper.h"
 #include "DataWrappers/ChaosVDCollisionDataWrappers.h"
 #include "GameFramework/Actor.h"
+#include "Interfaces/ChaosVDSelectableObject.h"
 #include "Visualizers/IChaosVDParticleVisualizationDataProvider.h"
 
 #include "ChaosVDParticleActor.generated.h"
@@ -52,7 +53,7 @@ ENUM_CLASS_FLAGS(EChaosVDHideParticleFlags)
 /** Actor used to represent a Chaos Particle in the Visual Debugger's world */
 UCLASS(HideCategories=(Transform))
 class AChaosVDParticleActor : public AActor, public IChaosVDParticleVisualizationDataProvider,
-								public FChaosVDSceneObjectBase, public IChaosVDCollisionDataProviderInterface
+								public FChaosVDSceneObjectBase, public IChaosVDCollisionDataProviderInterface, public IChaosVDGeometryOwnerInterface, public IChaosVDSelectableObject
 {
 	GENERATED_BODY()
 
@@ -100,20 +101,30 @@ public:
 
 	virtual FBox GetComponentsBoundingBox(bool bNonColliding, bool bIncludeFromChildActors) const override;
 
-	//BEGIN IChaosVDCollisionDataProvider Interface
+	// BEGIN IChaosVDCollisionDataProvider Interface
 	virtual void GetCollisionData(TArray<TSharedPtr<FChaosVDCollisionDataFinder>>& OutCollisionDataFound) override;
 	virtual bool HasCollisionData() override;
 	virtual FName GetProviderName() override;
-	//END IChaosVDCollisionDataProvider Interface
+	// END IChaosVDCollisionDataProvider Interface
 
 	void SetIsServerParticle(bool bNewIsServer) { bIsServer = bNewIsServer; }
 	bool GetIsServerParticle() const { return bIsServer; }
 
+	void UpdateMeshInstancesSelectionState();
 	virtual void PushSelectionToProxies() override;
 
 	FChaosVDParticleDataUpdatedDelegate& OnParticleDataUpdated() { return ParticleDataUpdatedDelegate; };
 
-	TConstArrayView<TSharedPtr<FChaosVDMeshDataInstanceHandle>> GetMeshInstances() const { return MeshDataHandles; }
+	// BEGIN IChaosVDGeometryOwner Interface
+	virtual TConstArrayView<TSharedPtr<FChaosVDMeshDataInstanceHandle>> GetMeshInstances() const override { return MeshDataHandles; }
+	virtual void SetSelectedMeshInstance(const TWeakPtr<FChaosVDMeshDataInstanceHandle>& GeometryInstanceToSelect) override;
+	virtual TWeakPtr<FChaosVDMeshDataInstanceHandle> GetSelectedMeshInstance() const override { return CurrentSelectedGeometryInstance; }
+	// END IChaosVDGeometryOwner Interface
+	
+	// BEGIN IChaosVDSelectableObject Interface
+	virtual void HandleSelected() override;
+	virtual void HandleDeSelected() override;
+	// END IChaosVDSelectableObject Interface
 
 protected:
 
@@ -144,6 +155,8 @@ protected:
 
 	EChaosVDHideParticleFlags HideParticleFlags;
 
+	TWeakPtr<FChaosVDMeshDataInstanceHandle> CurrentSelectedGeometryInstance;
+
 	friend FChaosVDParticleActorCustomization;
 };
 
diff --git a/Engine/Plugins/ChaosVD/Source/ChaosVD/Private/ChaosVDParticleActorCustomization.cpp b/Engine/Plugins/ChaosVD/Source/ChaosVD/Private/ChaosVDParticleActorCustomization.cpp
--- a/Engine/Plugins/ChaosVD/Source/ChaosVD/Private/ChaosVDParticleActorCustomization.cpp
+++ b/Engine/Plugins/ChaosVD/Source/ChaosVD/Private/ChaosVDParticleActorCustomization.cpp
@@ -15,8 +15,8 @@
 
 FChaosVDParticleActorCustomization::FChaosVDParticleActorCustomization()
 {
-	AllowedCategories.Add(FChaosVDParticleActorCustomization::ChaosVDCategoryName);
-	AllowedCategories.Add(FChaosVDParticleActorCustomization::ChaosVDVisualizationCategoryName);
+	AllowedCategories.Add(FChaosVDParticleActorCustomization::ParticleDataCategoryName);
+	AllowedCategories.Add(FChaosVDParticleActorCustomization::GeometryCategoryName);
 }
 
 FChaosVDParticleActorCustomization::~FChaosVDParticleActorCustomization()
@@ -34,8 +34,6 @@ TSharedRef<IDetailCustomization> FChaosVDParticleActorCustomization::MakeInstanc
 
 void FChaosVDParticleActorCustomization::CustomizeDetails(IDetailLayoutBuilder& DetailBuilder)
 {
-	DetailBuilder.EditCategory(ChaosVDVisualizationCategoryName, FText::GetEmpty(), ECategoryPriority::Important);
-
 	FChaosVDDetailsCustomizationUtils::HideAllCategories(DetailBuilder, AllowedCategories);
 
 	// We keep the particle data we need to visualize as a shared ptr because copying it each frame we advance/rewind to to an struct that lives in the particle actor it is not cheap.
@@ -44,8 +42,6 @@ void FChaosVDParticleActorCustomization::CustomizeDetails(IDetailLayoutBuilder&
 	// So a middle ground I found is to have a Particle Data struct in this customization instance, which we add as external property. Then each time the particle data is updated we copy the data over.
 	// This allow us to only perform the copy just for the particle that is being inspected and not every particle updated in that frame.
 
-	IDetailCategoryBuilder& CVDMainCategoryBuilder = DetailBuilder.EditCategory(ChaosVDCategoryName).InitiallyCollapsed(false);
-
 	TArray<TWeakObjectPtr<UObject>> SelectedObjects;
 	DetailBuilder.GetObjectsBeingCustomized(SelectedObjects);
 	if (SelectedObjects.Num() > 0)
@@ -58,7 +54,8 @@ void FChaosVDParticleActorCustomization::CustomizeDetails(IDetailLayoutBuilder&
 
 		if (AChaosVDParticleActor* CurrentActor = CurrentObservedActor.Get())
 		{
-			CurrentParticleDataCopy = FChaosVDParticleDataWrapper();
+			CachedParticleData = FChaosVDParticleDataWrapper();
+			CachedGeometryDataInstanceCopy = FChaosVDMeshDataInstanceState();
 			CurrentActor->OnParticleDataUpdated().Unbind();
 			CurrentActor = nullptr;
 		}
@@ -70,23 +67,50 @@ void FChaosVDParticleActorCustomization::CustomizeDetails(IDetailLayoutBuilder&
 
 			HandleParticleDataUpdated();
 
-			const TSharedPtr<FStructOnScope> ParticleDataView = MakeShared<FStructOnScope>(FChaosVDParticleDataWrapper::StaticStruct(), reinterpret_cast<uint8*>(&CurrentParticleDataCopy));
-			TArray<TSharedPtr<IPropertyHandle>> Handles = CVDMainCategoryBuilder.AddAllExternalStructureProperties(ParticleDataView.ToSharedRef(), EPropertyLocation::Default, nullptr);
-
-			FChaosVDDetailsCustomizationUtils::HideInvalidCVDDataWrapperProperties(Handles);
+			TSharedPtr<IPropertyHandle> InspectedDataPropertyHandlePtr;
+
+			if (TSharedPtr<FChaosVDMeshDataInstanceHandle> SelectedGeometryInstance = ParticleActor->GetSelectedMeshInstance().Pin())
+			{
+				InspectedDataPropertyHandlePtr = AddExternalStructure(CachedGeometryDataInstanceCopy, DetailBuilder, GeometryCategoryName, LOCTEXT("ParticleDataStructName", "Geometry Shape Data"));
+			}
+			else
+			{
+				InspectedDataPropertyHandlePtr = AddExternalStructure(CachedParticleData, DetailBuilder, ParticleDataCategoryName, LOCTEXT("ParticleDataStructName", "Particle Data"));
+			}
+
+			if (InspectedDataPropertyHandlePtr)
+			{
+				TSharedRef<IPropertyHandle> InspectedDataPropertyHandleRef = InspectedDataPropertyHandlePtr.ToSharedRef();
+				FChaosVDDetailsCustomizationUtils::HideInvalidCVDDataWrapperProperties({&InspectedDataPropertyHandleRef, 1}, DetailBuilder);
+			}
 		}
 	}
 }
 
 void FChaosVDParticleActorCustomization::HandleParticleDataUpdated()
-{	
-	if (const FChaosVDParticleDataWrapper* ParticleDataPtr = CurrentObservedActor.Get() ? CurrentObservedActor->GetParticleData() : nullptr)
+{
+	AChaosVDParticleActor* ParticleActor = CurrentObservedActor.Get();
+	if (!ParticleActor)
 	{
-		CurrentParticleDataCopy = *ParticleDataPtr;	
+		CachedParticleData = FChaosVDParticleDataWrapper();
+		CachedGeometryDataInstanceCopy = FChaosVDMeshDataInstanceState();
+	}
+
+	// If we have selected a mesh instance, the only data being added to the details panel is the Shape Instance data, so can just update that data here
+	if (TSharedPtr<FChaosVDMeshDataInstanceHandle> SelectedGeometryInstance = ParticleActor->GetSelectedMeshInstance().Pin())
+	{
+		ParticleActor->VisitGeometryInstances([this, SelectedGeometryInstance](const TSharedRef<FChaosVDMeshDataInstanceHandle>& MeshDataHandle)
+		{
+			if (MeshDataHandle == SelectedGeometryInstance)
+			{
+				CachedGeometryDataInstanceCopy = MeshDataHandle->GetState();
+			}
+		});
 	}
 	else
 	{
-		CurrentParticleDataCopy = FChaosVDParticleDataWrapper();
+		const FChaosVDParticleDataWrapper* ParticleDataPtr = CurrentObservedActor.Get() ? CurrentObservedActor->GetParticleData() : nullptr;
+		CachedParticleData = ParticleDataPtr ? *ParticleDataPtr : FChaosVDParticleDataWrapper();
 	}
 }
 
diff --git a/Engine/Plugins/ChaosVD/Source/ChaosVD/Private/ChaosVDParticleActorCustomization.h b/Engine/Plugins/ChaosVD/Source/ChaosVD/Private/ChaosVDParticleActorCustomization.h
--- a/Engine/Plugins/ChaosVD/Source/ChaosVD/Private/ChaosVDParticleActorCustomization.h
+++ b/Engine/Plugins/ChaosVD/Source/ChaosVD/Private/ChaosVDParticleActorCustomization.h
@@ -2,8 +2,12 @@
 
 #pragma once
 
+#include "ChaosVDGeometryDataComponent.h"
+#include "DetailCategoryBuilder.h"
+#include "DetailLayoutBuilder.h"
 #include "Templates/SharedPointer.h"
 #include "IDetailCustomization.h"
+#include "PropertyHandle.h"
 #include "DataWrappers/ChaosVDParticleDataWrapper.h"
 
 struct FChaosVDParticleDataWrapper;
@@ -16,18 +20,46 @@ public:
 	FChaosVDParticleActorCustomization();
 	virtual ~FChaosVDParticleActorCustomization() override;
 
-	inline static FName ChaosVDCategoryName = FName("Particle Data");
-	inline static FName ChaosVDVisualizationCategoryName = FName("Viewport Visualization Flags");
+	inline static FName ParticleDataCategoryName = FName("Particle Data");
+	inline static FName GeometryCategoryName = FName("Geometry Shape Data");
 
 	static TSharedRef<IDetailCustomization> MakeInstance();
 
 	virtual void CustomizeDetails(IDetailLayoutBuilder& DetailBuilder) override;
 
 private:
+
+	template<typename TStruct>
+	TSharedPtr<IPropertyHandle> AddExternalStructure(TStruct& CachedStruct, IDetailLayoutBuilder& DetailBuilder, FName CategoryName, const FText& InPropertyName);
+
 	TSet<FName> AllowedCategories;
 	TWeakObjectPtr<AChaosVDParticleActor> CurrentObservedActor;
 
 	void HandleParticleDataUpdated();
 
-	FChaosVDParticleDataWrapper CurrentParticleDataCopy;
+	/* Copied of the last known particle data structure of a selected particle -  Used to avoid rebuild the layout every time we change frame in CVD */
+	FChaosVDParticleDataWrapper CachedParticleData;
+	/* Copied of the last known particle data structure of a selected particle -  Used to avoid rebuild the layout every time we change frame in CVD */
+	FChaosVDMeshDataInstanceState CachedGeometryDataInstanceCopy;
 };
+
+template <typename TStruct>
+TSharedPtr<IPropertyHandle> FChaosVDParticleActorCustomization::AddExternalStructure(TStruct& CachedStruct, IDetailLayoutBuilder& DetailBuilder, FName CategoryName, const FText& InPropertyName)
+{
+	DetailBuilder.EditCategory(CategoryName, FText::GetEmpty(), ECategoryPriority::Important);
+	IDetailCategoryBuilder& CVDMainCategoryBuilder = DetailBuilder.EditCategory(CategoryName).InitiallyCollapsed(false);
+
+	const TSharedPtr<FStructOnScope> DataView = MakeShared<FStructOnScope>(TStruct::StaticStruct(), reinterpret_cast<uint8*>(&CachedStruct));
+
+	FAddPropertyParams AddParams;
+	AddParams.CreateCategoryNodes(true);
+
+	if (IDetailPropertyRow* PropertyRow = CVDMainCategoryBuilder.AddExternalStructureProperty(DataView, NAME_None, EPropertyLocation::Default, AddParams))
+	{
+		PropertyRow->ShouldAutoExpand(true);
+		PropertyRow->DisplayName(InPropertyName);
+		return PropertyRow->GetPropertyHandle();
+	}
+
+	return nullptr;
+}
diff --git a/Engine/Plugins/ChaosVD/Source/ChaosVD/Private/ChaosVDPlaybackViewportClient.cpp b/Engine/Plugins/ChaosVD/Source/ChaosVD/Private/ChaosVDPlaybackViewportClient.cpp
--- a/Engine/Plugins/ChaosVD/Source/ChaosVD/Private/ChaosVDPlaybackViewportClient.cpp
+++ b/Engine/Plugins/ChaosVD/Source/ChaosVD/Private/ChaosVDPlaybackViewportClient.cpp
@@ -5,10 +5,12 @@
 #include "ChaosVDEditorSettings.h"
 #include "ChaosVDEngine.h"
 #include "ChaosVDModule.h"
+#include "ChaosVDObjectDetailsTab.h"
 #include "ChaosVDParticleActor.h"
 #include "ChaosVDPlaybackController.h"
 #include "ChaosVDScene.h"
 #include "ChaosVDSkySphereInterface.h"
+#include "ChaosVDTabsIDs.h"
 #include "Components/ChaosVDSceneQueryDataComponent.h"
 #include "ComponentVisualizer.h"
 #include "EditorModeManager.h"
@@ -25,6 +27,7 @@
 #include "Components/ChaosVDSolverJointConstraintDataComponent.h"
 #include "Components/InstancedStaticMeshComponent.h"
 #include "Visualizers/ChaosVDDebugDrawUtils.h"
+#include "Widgets/SChaosVDDetailsView.h"
 #include "Widgets/SChaosVDMainTab.h"
 
 FChaosVDPlaybackViewportClient::FChaosVDPlaybackViewportClient(const TSharedPtr<FEditorModeTools>& InModeTools, const TSharedPtr<SEditorViewport>& InEditorViewportWidget) : FEditorViewportClient(InModeTools.Get(), nullptr, InEditorViewportWidget), CVDWorld(nullptr)
@@ -83,6 +86,8 @@ void FChaosVDPlaybackViewportClient::ProcessClick(FSceneView& View, HHitProxy* H
 	{
 		return;
 	}
+	
+	const bool bIsShiftKeyDown = Viewport->KeyState(EKeys::LeftShift) || Viewport->KeyState(EKeys::RightShift);
 
 	const FViewportClick Click(&View, this, Key, Event, HitX, HitY);
 
@@ -129,7 +134,7 @@ void FChaosVDPlaybackViewportClient::ProcessClick(FSceneView& View, HHitProxy* H
 			{
 				if (AChaosVDParticleActor* ClickedActor = ScenePtr->GetParticleActor(MeshDataHandle->GetOwningSolverID(), MeshDataHandle->GetOwningParticleID()))
 				{
-					ScenePtr->SetSelectedObject(ClickedActor);
+					Chaos::VisualDebugger::SelectParticleWithGeometryInstance(ScenePtr.ToSharedRef(), ClickedActor, bIsShiftKeyDown ? MeshDataHandle : nullptr);
 					bClickHandled = true;
 				}
 			}
diff --git a/Engine/Plugins/ChaosVD/Source/ChaosVD/Private/ChaosVDScene.cpp b/Engine/Plugins/ChaosVD/Source/ChaosVD/Private/ChaosVDScene.cpp
--- a/Engine/Plugins/ChaosVD/Source/ChaosVD/Private/ChaosVDScene.cpp
+++ b/Engine/Plugins/ChaosVD/Source/ChaosVD/Private/ChaosVDScene.cpp
@@ -10,6 +10,7 @@
 #include "ChaosVDParticleActor.h"
 #include "Chaos/ImplicitObject.h"
 #include "ChaosVDRecording.h"
+#include "ChaosVDSelectionCustomization.h"
 #include "ChaosVDSkySphereInterface.h"
 #include "Components/ChaosVDSceneQueryDataComponent.h"
 #include "Components/ChaosVDSolverCollisionDataComponent.h"
@@ -28,6 +29,9 @@
 #include "UObject/Package.h"
 #include "WorldPersistentFolders.h"
 #include "Components/ChaosVDSolverJointConstraintDataComponent.h"
+#include "Elements/Actor/ActorElementData.h"
+#include "Elements/Component/ComponentElementData.h"
+#include "Elements/Object/ObjectElementData.h"
 #include "Engine/Level.h"
 
 #define LOCTEXT_NAMESPACE "ChaosVisualDebugger"
@@ -552,7 +556,7 @@ FTypedElementHandle FChaosVDScene::GetSelectionHandleForObject(const UObject* Ob
 	return Handle;
 }
 
-void FChaosVDScene::UpdateSelectionProxiesForActors(const TArray<AActor*>& SelectedActors)
+void FChaosVDScene::UpdateSelectionProxiesForActors(TArrayView<AActor*> SelectedActors)
 {
 	for (AActor* SelectedActor : SelectedActors)
 	{
@@ -563,19 +567,30 @@ void FChaosVDScene::UpdateSelectionProxiesForActors(const TArray<AActor*>& Selec
 	}
 }
 
-void FChaosVDScene::HandlePreSelectionChange(const UTypedElementSelectionSet* PreChangeSelectionSet)
+void FChaosVDScene::HandleDeSelectElement(const TTypedElement<ITypedElementSelectionInterface>& InElementSelectionHandle, FTypedElementListRef InSelectionSet, const FTypedElementSelectionOptions& InSelectionOptions)
 {
-	PendingActorsToUpdateSelectionProxy.Append(PreChangeSelectionSet->GetSelectedObjects<AActor>());
+	if (AActor* DeselectedActor = ActorElementDataUtil::GetActorFromHandle(InElementSelectionHandle))
+	{
+		if (IChaosVDSelectableObject* SelectionAwareActor = Cast<IChaosVDSelectableObject>(DeselectedActor))
+		{
+			SelectionAwareActor->HandleDeSelected();
+		}
+	}
+
+	// TODO: Add support for Component and Object Selection Events - This will be needed when we move away from using actors to represent particles
 }
 
-void FChaosVDScene::HandlePostSelectionChange(const UTypedElementSelectionSet* PreChangeSelectionSet)
+void FChaosVDScene::HandleSelectElement(const TTypedElement<ITypedElementSelectionInterface>& InElementSelectionHandle, FTypedElementListRef InSelectionSet, const FTypedElementSelectionOptions& InSelectionOptions)
 {
-	TArray<AActor*> SelectedActors = PreChangeSelectionSet->GetSelectedObjects<AActor>();
-
-	SelectedActors.Append(PendingActorsToUpdateSelectionProxy);
-	UpdateSelectionProxiesForActors(SelectedActors);
+	if (AActor* SelectedActor = ActorElementDataUtil::GetActorFromHandle(InElementSelectionHandle))
+	{
+		if (IChaosVDSelectableObject* SelectionAwareActor = Cast<IChaosVDSelectableObject>(SelectedActor))
+		{
+			SelectionAwareActor->HandleSelected();
+		}
+	}
 
-	PendingActorsToUpdateSelectionProxy.Reset();
+	// TODO: Add support for Component and Object Selection Events - This will be needed when we move away from using actors to represent particles
 }
 
 void FChaosVDScene::ClearSelectionAndNotify()
@@ -616,6 +631,10 @@ void FChaosVDScene::InitializeSelectionSets()
 	SelectionSet = NewObject<UTypedElementSelectionSet>(GetTransientPackage(), NAME_None, RF_Transactional);
 	SelectionSet->AddToRoot();
 
+	SelectionSet->RegisterInterfaceCustomizationByTypeName(NAME_Actor, MakeUnique<FChaosVDSelectionCustomization>(AsShared()));
+	SelectionSet->RegisterInterfaceCustomizationByTypeName(NAME_Components, MakeUnique<FChaosVDSelectionCustomization>(AsShared()));
+	SelectionSet->RegisterInterfaceCustomizationByTypeName(NAME_Object, MakeUnique<FChaosVDSelectionCustomization>(AsShared()));
+
 	FString ActorSelectionObjectName = FString::Printf(TEXT("CVDSelectedActors-%s"), *FGuid::NewGuid().ToString());
 	ActorSelection = USelection::CreateActorSelection(GetTransientPackage(), *ActorSelectionObjectName, RF_Transactional);
 	ActorSelection->SetElementSelectionSet(SelectionSet);
@@ -627,9 +646,6 @@ void FChaosVDScene::InitializeSelectionSets()
 	FString ObjectSelectionObjectName = FString::Printf(TEXT("CVDSelectedObjects-%s"), *FGuid::NewGuid().ToString());
 	ObjectSelection = USelection::CreateObjectSelection(GetTransientPackage(), *ObjectSelectionObjectName, RF_Transactional);
 	ObjectSelection->SetElementSelectionSet(SelectionSet);
-
-	SelectionSet->OnPreChange().AddRaw(this, &FChaosVDScene::HandlePreSelectionChange);
-	SelectionSet->OnChanged().AddRaw(this, &FChaosVDScene::HandlePostSelectionChange);
 }
 
 void FChaosVDScene::DeInitializeSelectionSets()
diff --git a/Engine/Plugins/ChaosVD/Source/ChaosVD/Private/ChaosVDScene.h b/Engine/Plugins/ChaosVD/Source/ChaosVD/Private/ChaosVDScene.h
--- a/Engine/Plugins/ChaosVD/Source/ChaosVD/Private/ChaosVDScene.h
+++ b/Engine/Plugins/ChaosVD/Source/ChaosVD/Private/ChaosVDScene.h
@@ -5,10 +5,15 @@
 #include "ChaosVDRecording.h"
 #include "Containers/UnrealString.h"
 #include "Containers/Map.h"
+#include "Elements/Framework/TypedElementListFwd.h"
+#include "Elements/Interfaces/TypedElementAssetDataInterface.h"
 #include "UObject/GCObject.h"
 #include "UObject/ObjectMacros.h"
 #include "UObject/ObjectPtr.h"
 
+class FChaosVDSelectionCustomization;
+class ITypedElementSelectionInterface;
+struct FTypedElementSelectionOptions;
 class AChaosVDSceneQueryDataContainer;
 class AChaosVDSolverInfoActor;
 class AChaosVDSceneCollisionContainer;
@@ -119,6 +124,9 @@ public:
 
 	FChaosVDSolverVisibilityChangedDelegate& OnSolverVisibilityUpdated() { return SolverVisibilityChangedDelegate; }
 
+	/** Updates the render state of the hit proxies of an array of actors. This used to update the selection outline state */
+	void UpdateSelectionProxiesForActors(TArrayView<AActor*> SelectedActors);
+
 	TSharedPtr<FChaosVDRecording> LoadedRecording;
 
 private:
@@ -147,11 +155,8 @@ private:
 	/** Returns the correct TypedElementHandle based on an object type so it can be used with the selection set object */
 	FTypedElementHandle GetSelectionHandleForObject(const UObject* Object) const;
 
-	/** Updates the render state of the hit proxies of an array of actors. This used to update the selection outline state */
-	void UpdateSelectionProxiesForActors(const TArray<AActor*>& SelectedActors);
-
-	void HandlePreSelectionChange(const UTypedElementSelectionSet* PreChangeSelectionSet);
-	void HandlePostSelectionChange(const UTypedElementSelectionSet* PreChangeSelectionSet);
+	void HandleDeSelectElement(const TTypedElement<ITypedElementSelectionInterface>& InElementSelectionHandle, FTypedElementListRef InSelectionSet, const FTypedElementSelectionOptions& InSelectionOptions);
+	void HandleSelectElement(const TTypedElement<ITypedElementSelectionInterface>& InElementSelectionHandle, FTypedElementListRef InSelectionSet, const FTypedElementSelectionOptions& InSelectionOptions);
 
 	void ClearSelectionAndNotify();
 
@@ -201,4 +206,6 @@ private:
 	FChaosVDSolverInfoActorCreatedDelegate SolverInfoActorCreatedDelegate;
 
 	FChaosVDSolverVisibilityChangedDelegate SolverVisibilityChangedDelegate;
+
+	friend FChaosVDSelectionCustomization;
 };
diff --git a/Engine/Plugins/ChaosVD/Source/ChaosVD/Private/ChaosVDSelectionCustomization.cpp b/Engine/Plugins/ChaosVD/Source/ChaosVD/Private/ChaosVDSelectionCustomization.cpp
new file mode 100644
--- /dev/null
+++ b/Engine/Plugins/ChaosVD/Source/ChaosVD/Private/ChaosVDSelectionCustomization.cpp
@@ -0,0 +1,31 @@
+// Copyright Epic Games, Inc. All Rights Reserved.
+
+#include "ChaosVDSelectionCustomization.h"
+
+#include "ChaosVDScene.h"
+
+bool FChaosVDSelectionCustomization::DeselectElement(const TTypedElement<ITypedElementSelectionInterface>& InElementSelectionHandle, FTypedElementListRef InSelectionSet, const FTypedElementSelectionOptions& InSelectionOptions)
+{
+	const bool bSelectionChanged = FTypedElementSelectionCustomization::DeselectElement(InElementSelectionHandle, InSelectionSet, InSelectionOptions);
+
+	if (TSharedPtr<FChaosVDScene> ScenePtr = SceneWeakPtr.Pin())
+	{
+		ScenePtr->HandleDeSelectElement(InElementSelectionHandle, InSelectionSet, InSelectionOptions);
+	}
+
+	return bSelectionChanged;
+}
+
+
+bool FChaosVDSelectionCustomization::SelectElement(const TTypedElement<ITypedElementSelectionInterface>& InElementSelectionHandle, FTypedElementListRef InSelectionSet, const FTypedElementSelectionOptions& InSelectionOptions)
+{
+	const bool bSelectionChanged = FTypedElementSelectionCustomization::SelectElement(InElementSelectionHandle, InSelectionSet, InSelectionOptions);
+
+	// Currently we need to know if an element was tried to be de-selected or selected regardless if the selection set changed, in case we are re-selecting the object but in a different way (like selecting a particle but to show a specific geometry instance)
+	if (TSharedPtr<FChaosVDScene> ScenePtr = SceneWeakPtr.Pin())
+	{
+		ScenePtr->HandleSelectElement(InElementSelectionHandle, InSelectionSet, InSelectionOptions);
+	}
+
+	return bSelectionChanged;
+}
diff --git a/Engine/Plugins/ChaosVD/Source/ChaosVD/Private/ChaosVDSelectionCustomization.h b/Engine/Plugins/ChaosVD/Source/ChaosVD/Private/ChaosVDSelectionCustomization.h
new file mode 100644
--- /dev/null
+++ b/Engine/Plugins/ChaosVD/Source/ChaosVD/Private/ChaosVDSelectionCustomization.h
@@ -0,0 +1,22 @@
+// Copyright Epic Games, Inc. All Rights Reserved.
+
+#pragma once
+#include "Elements/Framework/TypedElementSelectionSet.h"
+#include "Templates/SharedPointer.h"
+
+class FChaosVDScene;
+
+/** Customization class used to intercept selection actions and route Selection Events into CVD's Scene*/
+class FChaosVDSelectionCustomization : public FTypedElementSelectionCustomization
+{
+public:
+	FChaosVDSelectionCustomization(const TWeakPtr<FChaosVDScene>& InScene) : SceneWeakPtr(InScene)
+	{
+	}
+	
+	virtual bool DeselectElement(const TTypedElement<ITypedElementSelectionInterface>& InElementSelectionHandle, FTypedElementListRef InSelectionSet, const FTypedElementSelectionOptions& InSelectionOptions) override;
+	virtual bool SelectElement(const TTypedElement<ITypedElementSelectionInterface>& InElementSelectionHandle, FTypedElementListRef InSelectionSet, const FTypedElementSelectionOptions& InSelectionOptions) override;
+
+protected:
+	TWeakPtr<FChaosVDScene> SceneWeakPtr;
+};
diff --git a/Engine/Plugins/ChaosVD/Source/ChaosVD/Private/ChaosVDWorldOutlinerMode.cpp b/Engine/Plugins/ChaosVD/Source/ChaosVD/Private/ChaosVDWorldOutlinerMode.cpp
--- a/Engine/Plugins/ChaosVD/Source/ChaosVD/Private/ChaosVDWorldOutlinerMode.cpp
+++ b/Engine/Plugins/ChaosVD/Source/ChaosVD/Private/ChaosVDWorldOutlinerMode.cpp
@@ -133,6 +133,11 @@ FChaosVDWorldOutlinerMode::~FChaosVDWorldOutlinerMode()
 
 void FChaosVDWorldOutlinerMode::OnItemSelectionChanged(FSceneOutlinerTreeItemPtr Item, ESelectInfo::Type SelectionType, const FSceneOutlinerItemSelection& Selection)
 {
+	if (SelectionType == ESelectInfo::Direct)
+	{
+		return;
+	}
+
 	TSharedPtr<FChaosVDScene> ScenePtr = CVDScene.Pin();
 	if (!ScenePtr.IsValid())
 	{
@@ -285,7 +290,7 @@ void FChaosVDWorldOutlinerMode::HandlePostSelectionChange(const UTypedElementSel
 		if (FSceneOutlinerTreeItemPtr TreeItem = SceneOutliner->GetTreeItem(SelectedActor, false))
 		{
 			SceneOutliner->ScrollItemIntoView(TreeItem);
-			SceneOutliner->SetItemSelection(TreeItem, true, ESelectInfo::OnMouseClick);
+			SceneOutliner->SetItemSelection(TreeItem, true, ESelectInfo::Direct);
 		}
 		else
 		{
diff --git a/Engine/Plugins/ChaosVD/Source/ChaosVD/Private/DetailsCustomizations/ChaosVDDetailsCustomizationUtils.cpp b/Engine/Plugins/ChaosVD/Source/ChaosVD/Private/DetailsCustomizations/ChaosVDDetailsCustomizationUtils.cpp
--- a/Engine/Plugins/ChaosVD/Source/ChaosVD/Private/DetailsCustomizations/ChaosVDDetailsCustomizationUtils.cpp
+++ b/Engine/Plugins/ChaosVD/Source/ChaosVD/Private/DetailsCustomizations/ChaosVDDetailsCustomizationUtils.cpp
@@ -51,6 +51,15 @@ void FChaosVDDetailsCustomizationUtils::HideInvalidCVDDataWrapperProperties(TCon
 				DetailBuilder.HideProperty(PropertyHandle);
 			}
 		}
+
+		uint32 NumChildren = 0;
+		PropertyHandle->GetNumChildren(NumChildren);
+
+		for (uint32 ChildIndex = 0; ChildIndex < NumChildren; ++ChildIndex)
+		{
+			TSharedRef<IPropertyHandle> ChildHandle = PropertyHandle->GetChildHandle(ChildIndex).ToSharedRef();
+			HideInvalidCVDDataWrapperProperties({ &ChildHandle,1 }, DetailBuilder);
+		}	
 	}
 }
 
diff --git a/Engine/Plugins/ChaosVD/Source/ChaosVD/Private/Interfaces/ChaosVDSelectableObject.h b/Engine/Plugins/ChaosVD/Source/ChaosVD/Private/Interfaces/ChaosVDSelectableObject.h
new file mode 100644
--- /dev/null
+++ b/Engine/Plugins/ChaosVD/Source/ChaosVD/Private/Interfaces/ChaosVDSelectableObject.h
@@ -0,0 +1,24 @@
+// Copyright Epic Games, Inc. All Rights Reserved.
+
+#pragma once
+
+#include "UObject/Interface.h"
+#include "ChaosVDSelectableObject.generated.h"
+
+UINTERFACE()
+class UChaosVDSelectableObject : public UInterface
+{
+	GENERATED_BODY()
+};
+
+/**
+ * Interface used for any object that can be selected in CVD and need to process selection events performed on it.
+ */
+class CHAOSVD_API IChaosVDSelectableObject
+{
+	GENERATED_BODY()
+
+public:
+	virtual void HandleSelected() PURE_VIRTUAL(IChaosVDSelectionAware::HandleSelected);
+	virtual void HandleDeSelected() PURE_VIRTUAL(IChaosVDSelectionAware::HandleUnSelected);
+};
diff --git a/Engine/Plugins/ChaosVD/Source/ChaosVD/Private/Widgets/SChaosVDDetailsView.cpp b/Engine/Plugins/ChaosVD/Source/ChaosVD/Private/Widgets/SChaosVDDetailsView.cpp
--- a/Engine/Plugins/ChaosVD/Source/ChaosVD/Private/Widgets/SChaosVDDetailsView.cpp
+++ b/Engine/Plugins/ChaosVD/Source/ChaosVD/Private/Widgets/SChaosVDDetailsView.cpp
@@ -3,15 +3,42 @@
 #include "Widgets/SChaosVDDetailsView.h"
 
 #include "ChaosVDParticleActor.h"
-#include "GameFramework/Actor.h"
 #include "PropertyEditorModule.h"
-#include "SSubobjectInstanceEditor.h"
 #include "SSubobjectEditor.h"
 #include "SSubobjectEditorModule.h"
-#include "Visualizers/IChaosVDParticleVisualizationDataProvider.h"
-
 
 void SChaosVDDetailsView::Construct(const FArguments& InArgs)
+{
+	DetailsView = CreateObjectDetailsView();
+	StructDetailsView = CreateStructureDataDetailsView();
+
+	ChildSlot
+	[
+		SNew(SVerticalBox)
+		+SVerticalBox::Slot()
+		[
+			SNew(SVerticalBox)
+			.Visibility_Raw(this, &SChaosVDDetailsView::GetObjectDetailsVisibility)
+			+SVerticalBox::Slot()
+			.AutoHeight()
+			+SVerticalBox::Slot()
+			[
+				DetailsView.ToSharedRef()
+			]
+		]
+		+SVerticalBox::Slot()
+	    [
+			SNew(SVerticalBox)
+			.Visibility_Raw(this, &SChaosVDDetailsView::GetStructDetailsVisibility)
+			+SVerticalBox::Slot()
+			[
+				StructDetailsView->GetWidget().ToSharedRef()
+			]
+	    ]
+	];
+}
+
+TSharedPtr<IDetailsView> SChaosVDDetailsView::CreateObjectDetailsView()
 {
 	FPropertyEditorModule& PropertyEditorModule = FModuleManager::GetModuleChecked<FPropertyEditorModule>("PropertyEditor");
 	FDetailsViewArgs DetailsViewArgs;
@@ -19,126 +46,48 @@ void SChaosVDDetailsView::Construct(const FArguments& InArgs)
 	DetailsViewArgs.bLockable = true;
 	DetailsViewArgs.bAllowFavoriteSystem = false;
 	DetailsViewArgs.NameAreaSettings = FDetailsViewArgs::ObjectsUseNameArea | FDetailsViewArgs::ComponentsAndActorsUseNameArea;
-	DetailsViewArgs.bCustomNameAreaLocation = true;
 	DetailsViewArgs.bCustomFilterAreaLocation = false;
 	DetailsViewArgs.bShowSectionSelector = false;
 	DetailsViewArgs.bShowScrollBar = false;
 
-	DetailsView = PropertyEditorModule.CreateDetailView(DetailsViewArgs);
+	return PropertyEditorModule.CreateDetailView(DetailsViewArgs);
+}
 
-	FModuleManager::LoadModuleChecked<FSubobjectEditorModule>("SubobjectEditor");
+TSharedPtr<IStructureDetailsView> SChaosVDDetailsView::CreateStructureDataDetailsView()
+{
+	FPropertyEditorModule& PropertyEditorModule = FModuleManager::GetModuleChecked<FPropertyEditorModule>("PropertyEditor");
 
-	SubobjectEditor = SNew(SSubobjectInstanceEditor)
-		.ObjectContext(this, &SChaosVDDetailsView::GetRootContextObject)
-		.AllowEditing(false)
-		.OnSelectionUpdated(this, &SChaosVDDetailsView::OnSelectedSubobjectsChanged);
+	const FStructureDetailsViewArgs StructDetailsViewArgs;
+	FDetailsViewArgs DetailsViewArgs;
+	DetailsViewArgs.bShowOptions = false;
+	DetailsViewArgs.bAllowFavoriteSystem = false;
+	DetailsViewArgs.bAllowSearch = false;
+	DetailsViewArgs.bShowScrollBar = false;
 
-	ChildSlot
-	[
-		SNew(SVerticalBox)
-		+SVerticalBox::Slot()
-		.Padding(10.f, 4.f, 0.f, 0.f)
-		.AutoHeight()
-		[
-			DetailsView->GetNameAreaWidget().ToSharedRef()
-		]
-		+SVerticalBox::Slot()
-		[
-			SNew(SSplitter)
-			.MinimumSlotHeight(40.0f)
-			.Orientation(Orient_Vertical)
-			.Style(FAppStyle::Get(), "SplitterDark")
-			.PhysicalSplitterHandleSize(2.0f)
-			+SSplitter::Slot()
-			.Value(0.2f)
-			[
-				SNew(SVerticalBox)
-				+SVerticalBox::Slot()
-				.Padding(FMargin(2.0f, 0.0f, 2.0f, 0.0f))
-				[
-					SubobjectEditor.ToSharedRef()
-				]
-			]
-			+SSplitter::Slot()
-			[
-				SNew(SVerticalBox)
-				+ SVerticalBox::Slot()
-				[
-					DetailsView.ToSharedRef()
-				]
-			]
-		]
-	];
+	return PropertyEditorModule.CreateStructureDetailView(DetailsViewArgs,StructDetailsViewArgs, nullptr);
 }
 
-void SChaosVDDetailsView::SetSelectedObject(UObject* NewObject)
+EVisibility SChaosVDDetailsView::GetStructDetailsVisibility() const
 {
-	if (DetailsView->IsLocked())
-	{
-		return;
-	}
-
-	CurrentObjectInView = NewObject;
+	return CurrentStructInView.IsValid() ? EVisibility::Visible : EVisibility::Collapsed;
+}
 
-	DetailsView->SetObject(NewObject, true);
-	SubobjectEditor->UpdateTree();
+EVisibility SChaosVDDetailsView::GetObjectDetailsVisibility() const
+{
+	return CurrentObjectInView.IsValid() ? EVisibility::Visible : EVisibility::Collapsed;
 }
 
-void SChaosVDDetailsView::OnSelectedSubobjectsChanged(const TArray<TSharedPtr<FSubobjectEditorTreeNode>>& SelectedNodes)
+void SChaosVDDetailsView::SetSelectedObject(UObject* NewObject)
 {
-	if (SelectedNodes.Num() == 0)
+	// Even if the object is not valid, clear any active structure view
+	StructDetailsView->SetStructureData(nullptr);
+
+	if (DetailsView->IsLocked())
 	{
 		return;
 	}
 
-	if (!DetailsView->IsLocked())
-	{
-		const bool bContainsRootActor = SelectedNodes.ContainsByPredicate([](const FSubobjectEditorTreeNodePtrType& Node)
-		{
-			if (Node.IsValid())
-			{
-				if (FSubobjectData* Data = Node->GetDataSource())
-				{
-					return Data->IsRootActor();
-				}
-			}
-
-			return false;
-		});
-
-		if (bContainsRootActor)
-		{
-			DetailsView->SetObject(GetRootContextObject());
-		}
-		else
-		{
-			TArray<UObject*> Components;
-		
-			if (AActor* ContextAsActor = Cast<AActor>(GetRootContextObject()))
-			{
-				Components.Reserve(SelectedNodes.Num());
-                for (const FSubobjectEditorTreeNodePtrType& Node : SelectedNodes)
-                {
-                	if (Node.IsValid())
-                	{
-                		if (FSubobjectData* Data = Node->GetDataSource())
-                		{
-                			if (Data->IsComponent())
-                			{
-                				if (const UActorComponent* Component = Data->FindComponentInstanceInActor(ContextAsActor))
-                				{
-                					if (Component != ContextAsActor->GetRootComponent())
-                					{
-                						Components.Add(const_cast<UActorComponent*>(Component));
-                					}
-                				}
-                			}
-                		}
-                	}
-                }
-			}
+	CurrentObjectInView = NewObject;
 
-			DetailsView->SetObjects(Components);
-		}
-	}
+	DetailsView->SetObject(NewObject, true);
 }
diff --git a/Engine/Plugins/ChaosVD/Source/ChaosVD/Private/Widgets/SChaosVDDetailsView.h b/Engine/Plugins/ChaosVD/Source/ChaosVD/Private/Widgets/SChaosVDDetailsView.h
--- a/Engine/Plugins/ChaosVD/Source/ChaosVD/Private/Widgets/SChaosVDDetailsView.h
+++ b/Engine/Plugins/ChaosVD/Source/ChaosVD/Private/Widgets/SChaosVDDetailsView.h
@@ -2,10 +2,13 @@
 
 #pragma once
 
+#include "IStructureDetailsView.h"
 #include "Templates/SharedPointer.h"
+#include "UObject/StructOnScope.h"
 #include "UObject/WeakObjectPtrTemplates.h"
 #include "Widgets/SCompoundWidget.h"
 
+class IStructureDetailsView;
 class FSubobjectEditorTreeNode;
 class AActor;
 class SBox;
@@ -15,7 +18,7 @@ class SSubobjectEditor;
 class UObject;
 
 /**
- * Simple details view with SubObject Editor
+ * Simple details for CVD objects and structures
  */
 class CHAOSVD_API SChaosVDDetailsView : public SCompoundWidget
 {
@@ -31,15 +34,40 @@ public:
 	/** Updates the current object this view details is viewing */
 	void SetSelectedObject(UObject* NewObject);
 
+	/** Updates the current object this view details is viewing */
+	template<typename TStruct>
+	void SetSelectedStruct(TStruct* NewStruct);
+
 protected:
-	UObject* GetRootContextObject() const { return CurrentObjectInView.IsValid() ? CurrentObjectInView.Get() : nullptr; };
-	
-	void OnSelectedSubobjectsChanged(const TArray<TSharedPtr<FSubobjectEditorTreeNode>>& SelectedNodes);
+
+	TSharedPtr<IDetailsView> CreateObjectDetailsView();
+	TSharedPtr<IStructureDetailsView> CreateStructureDataDetailsView();
+
+	EVisibility GetStructDetailsVisibility() const;
+	EVisibility GetObjectDetailsVisibility() const;
 
 	TWeakObjectPtr<UObject> CurrentObjectInView;
+
+	TWeakPtr<FStructOnScope> CurrentStructInView;
 	
 	TSharedPtr<IDetailsView> DetailsView;
 
-	// The subobject editor provides a tree widget that allows for editing of subobjects
-	TSharedPtr<SSubobjectEditor> SubobjectEditor;
+	TSharedPtr<IStructureDetailsView> StructDetailsView;
 };
+
+template <typename TStruct>
+void SChaosVDDetailsView::SetSelectedStruct(TStruct* NewStruct)
+{
+	SetSelectedObject(nullptr);
+
+	TSharedPtr<FStructOnScope> StructDataView = nullptr;
+
+	if (NewStruct)
+	{
+		StructDataView = MakeShared<FStructOnScope>(TStruct::StaticStruct(), reinterpret_cast<uint8*>(NewStruct));
+	}
+
+	CurrentStructInView = StructDataView;
+
+	StructDetailsView->SetStructureData(StructDataView);
+}
diff --git a/Engine/Plugins/ChaosVD/Source/ChaosVD/Private/Widgets/SChaosVDSceneQueryDataInspector.cpp b/Engine/Plugins/ChaosVD/Source/ChaosVD/Private/Widgets/SChaosVDSceneQueryDataInspector.cpp
--- a/Engine/Plugins/ChaosVD/Source/ChaosVD/Private/Widgets/SChaosVDSceneQueryDataInspector.cpp
+++ b/Engine/Plugins/ChaosVD/Source/ChaosVD/Private/Widgets/SChaosVDSceneQueryDataInspector.cpp
@@ -411,25 +411,18 @@ FReply SChaosVDSceneQueryDataInspector::SelectParticleForCurrentQueryData() cons
 
 	if (AChaosVDParticleActor* ParticleActor = ScenePtr->GetParticleActor(QueryDataBeingInspected->WorldSolverID, QueryDataBeingInspected->SQVisitData[QueryDataBeingInspected->CurrentVisitIndex].ParticleIndex))
 	{
-		ScenePtr->SetSelectedObject(ParticleActor);
+		const int32 ShapeInstanceIndexToSelect = QueryDataBeingInspected->SQVisitData[QueryDataBeingInspected->CurrentVisitIndex].ShapeIndex;
 
-		// When a particle actor is selected, all their mesh instances are selected, so we need to override that selection here.
-		// TODO: The API for selection is intentionally generic as we use it for more things that particle actors, but maybe we can have a new API method to provide a "context", so in this case
-		// we can tell the system how we want the selection to be visualized
-
-		const TConstArrayView<TSharedPtr<FChaosVDMeshDataInstanceHandle>> MeshHandles = ParticleActor->GetMeshInstances();
-		for (const TSharedPtr<FChaosVDMeshDataInstanceHandle>& Handle : MeshHandles)
+		// TODO : This will not work properly when the visited shape was a Union within a union
+		// CVD currently doesn't support multi selection, so we can't easily select all mesh instances that represent a union within a union
+		// We need to revisit this when multi selection support is added. Jira for tracking UE-212733
+		const TConstArrayView<TSharedPtr<FChaosVDMeshDataInstanceHandle>> AvailableMeshInstances = ParticleActor->GetMeshInstances();
+		for (const TSharedPtr<FChaosVDMeshDataInstanceHandle>& MeshInstance : AvailableMeshInstances)
 		{
-			if (const TSharedPtr<FChaosVDExtractedGeometryDataHandle> GeometryHandle = Handle->GetGeometryHandle())
+			if (MeshInstance->GetState().ImplicitObjectInfo.ShapeInstanceIndex == ShapeInstanceIndexToSelect)
 			{
-				if (GeometryHandle->GetImplicitObjectIndex() == QueryDataBeingInspected->SQVisitData[QueryDataBeingInspected->CurrentVisitIndex].ShapeIndex)
-				{
-					Handle->SetIsSelected(true);
-				}
-				else
-				{
-					Handle->SetIsSelected(false);
-				}
+				Chaos::VisualDebugger::SelectParticleWithGeometryInstance(ScenePtr.ToSharedRef(), ParticleActor, MeshInstance);
+				break;
 			}
 		}
 		
diff --git a/Engine/Source/Runtime/Core/Public/UObject/FortniteMainBranchObjectVersion.h b/Engine/Source/Runtime/Core/Public/UObject/FortniteMainBranchObjectVersion.h
--- a/Engine/Source/Runtime/Core/Public/UObject/FortniteMainBranchObjectVersion.h
+++ b/Engine/Source/Runtime/Core/Public/UObject/FortniteMainBranchObjectVersion.h
@@ -444,6 +444,9 @@ struct FFortniteMainBranchObjectVersion
 		// Adding a new flag in RBAN Solver Setting to be able to use manifolds
 		PhysicsAssetUseManifoldFlags,
 
+		// Added support for to record sim and query data of Shape Instance data in CVD
+		SimAndQueryDataSupportInChaosVisualDebugger,
+
 		// -----<new versions can be added above this line>-------------------------------------------------
 		VersionPlusOne,
 		LatestVersion = VersionPlusOne - 1
diff --git a/Engine/Source/Runtime/Experimental/ChaosVisualDebugger/Private/DataWrappers/ChaosVDCollisionDataWrappers.cpp b/Engine/Source/Runtime/Experimental/ChaosVisualDebugger/Private/DataWrappers/ChaosVDCollisionDataWrappers.cpp
--- a/Engine/Source/Runtime/Experimental/ChaosVisualDebugger/Private/DataWrappers/ChaosVDCollisionDataWrappers.cpp
+++ b/Engine/Source/Runtime/Experimental/ChaosVisualDebugger/Private/DataWrappers/ChaosVDCollisionDataWrappers.cpp
@@ -2,6 +2,8 @@
 
 #include "DataWrappers/ChaosVDCollisionDataWrappers.h"
 #include "DataWrappers/ChaosVDDataSerializationMacros.h"
+#include "UObject/FortniteMainBranchObjectVersion.h"
+
 
 bool FChaosVDContactPoint::Serialize(FArchive& Ar)
 {
@@ -196,6 +198,8 @@ bool FChaosVDCollisionFilterData::Serialize(FArchive& Ar)
 
 bool FChaosVDShapeCollisionData::Serialize(FArchive& Ar)
 {
+	Ar.UsingCustomVersion(FFortniteMainBranchObjectVersion::GUID);
+
 	Ar << CollisionTraceType;
 
 	EChaosVDCollisionShapeDataFlags PackedFlags = EChaosVDCollisionShapeDataFlags::None;
@@ -218,6 +222,12 @@ bool FChaosVDShapeCollisionData::Serialize(FArchive& Ar)
 		Ar << PackedFlags;
 	}
 
+	if (Ar.CustomVer(FFortniteMainBranchObjectVersion::GUID) >= FFortniteMainBranchObjectVersion::SimAndQueryDataSupportInChaosVisualDebugger)
+	{
+		Ar << SimData;
+		Ar << QueryData;
+	}
+
 	return true;
 }
 
diff --git a/Engine/Source/Runtime/Experimental/ChaosVisualDebugger/Public/DataWrappers/ChaosVDCollisionDataWrappers.h b/Engine/Source/Runtime/Experimental/ChaosVisualDebugger/Public/DataWrappers/ChaosVDCollisionDataWrappers.h
--- a/Engine/Source/Runtime/Experimental/ChaosVisualDebugger/Public/DataWrappers/ChaosVDCollisionDataWrappers.h
+++ b/Engine/Source/Runtime/Experimental/ChaosVisualDebugger/Public/DataWrappers/ChaosVDCollisionDataWrappers.h
@@ -498,8 +498,10 @@ struct CHAOSVDRUNTIME_API FChaosVDShapeCollisionData
 	UPROPERTY(VisibleAnywhere, Category=SimData)
 	FChaosVDCollisionFilterData SimData;
 
+	UPROPERTY(VisibleAnywhere, Category="CVD Data")
 	bool bIsComplex = false;
 
+	UPROPERTY(VisibleAnywhere, Category="CVD Data")
 	bool bIsValid = false;
 
 	bool Serialize(FArchive& Ar);
diff --git a/Engine/Source/Runtime/Experimental/ChaosVisualDebugger/Public/DataWrappers/ChaosVDParticleDataWrapper.h b/Engine/Source/Runtime/Experimental/ChaosVisualDebugger/Public/DataWrappers/ChaosVDParticleDataWrapper.h
--- a/Engine/Source/Runtime/Experimental/ChaosVisualDebugger/Public/DataWrappers/ChaosVDParticleDataWrapper.h
+++ b/Engine/Source/Runtime/Experimental/ChaosVisualDebugger/Public/DataWrappers/ChaosVDParticleDataWrapper.h
@@ -556,34 +556,34 @@ struct FChaosVDParticleDataWrapper : public FChaosVDWrapperDataBase
 	{
 	}
 
-	UPROPERTY(VisibleAnywhere, Category= "Particle Non Frequent Data")
+	UPROPERTY(VisibleAnywhere, Category= "General")
 	uint32 GeometryHash = 0;
 
-	UPROPERTY(VisibleAnywhere, Category= "Particle Non Frequent Data")
+	UPROPERTY(VisibleAnywhere, Category= "General")
 	FString DebugName;
 
-	UPROPERTY(VisibleAnywhere, Category= "Particle Non Frequent Data")
+	UPROPERTY(VisibleAnywhere, Category= "General")
 	int32 ParticleIndex = INDEX_NONE;
 
-	UPROPERTY(VisibleAnywhere, Category= "Particle Non Frequent Data")
+	UPROPERTY(VisibleAnywhere, Category= "General")
 	int32 SolverID = INDEX_NONE;
 
-	UPROPERTY(VisibleAnywhere, Category= "Particle Non Frequent Data")
+	UPROPERTY(VisibleAnywhere, Category= "General")
 	EChaosVDParticleType Type = EChaosVDParticleType::Unknown;
 
 	UPROPERTY(VisibleAnywhere, Category= "Particle Position Rotation")
 	FChaosVDParticlePositionRotation ParticlePositionRotation;
 
-	UPROPERTY(VisibleAnywhere, Category= "Particle Particle Velocities")
+	UPROPERTY(VisibleAnywhere, Category= "Particle Velocities")
 	FChaosVDParticleVelocities ParticleVelocities;
 
-	UPROPERTY(VisibleAnywhere, Category= "Particle Particle Dynamics")
+	UPROPERTY(VisibleAnywhere, Category= "Particle Dynamics")
 	FChaosVDParticleDynamics ParticleDynamics;
 
-	UPROPERTY(VisibleAnywhere, Category= "Particle Particle Dynamics Misc")
+	UPROPERTY(VisibleAnywhere, Category= "Particle Dynamics Misc")
 	FChaosVDParticleDynamicMisc ParticleDynamicsMisc;
 
-	UPROPERTY(VisibleAnywhere, Category= "Particle Particle Mass Props")
+	UPROPERTY(VisibleAnywhere, Category= "Particle Mass Props")
 	FChaosVDParticleMassProps ParticleMassProps;
 
 	UPROPERTY(VisibleAnywhere, Category= "Particle Cluster Data")
