[ChaosVD] Fix for CIS warning  in ChaosVDParticleActorCustomization.cpp

Adding missing return call after checking a ptr is null

diff --git a/Engine/Plugins/ChaosVD/Source/ChaosVD/Private/ChaosVDParticleActorCustomization.cpp b/Engine/Plugins/ChaosVD/Source/ChaosVD/Private/ChaosVDParticleActorCustomization.cpp
--- a/Engine/Plugins/ChaosVD/Source/ChaosVD/Private/ChaosVDParticleActorCustomization.cpp
+++ b/Engine/Plugins/ChaosVD/Source/ChaosVD/Private/ChaosVDParticleActorCustomization.cpp
@@ -94,6 +94,7 @@ void FChaosVDParticleActorCustomization::HandleParticleDataUpdated()
 	{
 		CachedParticleData = FChaosVDParticleDataWrapper();
 		CachedGeometryDataInstanceCopy = FChaosVDMeshDataInstanceState();
+		return;
 	}
 
 	// If we have selected a mesh instance, the only data being added to the details panel is the Shape Instance data, so can just update that data here
@@ -109,7 +110,7 @@ void FChaosVDParticleActorCustomization::HandleParticleDataUpdated()
 	}
 	else
 	{
-		const FChaosVDParticleDataWrapper* ParticleDataPtr = CurrentObservedActor.Get() ? CurrentObservedActor->GetParticleData() : nullptr;
+		const FChaosVDParticleDataWrapper* ParticleDataPtr = ParticleActor->GetParticleData();
 		CachedParticleData = ParticleDataPtr ? *ParticleDataPtr : FChaosVDParticleDataWrapper();
 	}
 }
