[ChaosVD] Fix Localization CIS warning in ChaosVDParticleActorCustomization

Also, fixed a duplicated tooltip and comment

diff --git a/Engine/Plugins/ChaosVD/Source/ChaosVD/Private/ChaosVDEditorSettings.h b/Engine/Plugins/ChaosVD/Source/ChaosVD/Private/ChaosVDEditorSettings.h
--- a/Engine/Plugins/ChaosVD/Source/ChaosVD/Private/ChaosVDEditorSettings.h
+++ b/Engine/Plugins/ChaosVD/Source/ChaosVD/Private/ChaosVDEditorSettings.h
@@ -315,7 +315,7 @@ public:
 	UPROPERTY(EditAnywhere, Category = "Viewport Visualization Flags", meta = (Bitmask, BitmaskEnum = "/Script/ChaosVD.EChaosVDSceneQueryVisualizationFlags"))
 	uint32 GlobalSceneQueriesVisualizationFlags = static_cast<uint32>(EChaosVDSceneQueryVisualizationFlags::DrawHits | EChaosVDSceneQueryVisualizationFlags::DrawLineTraceQueries);
 
-	/** Set of flags to enable/disable visualization of specific scene queries data as debug draw */
+	/** Set of flags to enable/disable visualization of specific joint constraints data as debug draw */
 	UPROPERTY(EditAnywhere, Category = "Viewport Visualization Flags", meta = (Bitmask, BitmaskEnum = "/Script/ChaosVD.EChaosVDJointsDataVisualizationFlags"))
 	uint32 GlobalJointsDataVisualizationFlags = static_cast<uint32>(EChaosVDJointsDataVisualizationFlags::ActorConnector | EChaosVDJointsDataVisualizationFlags::DrawKinematic);
 
diff --git a/Engine/Plugins/ChaosVD/Source/ChaosVD/Private/ChaosVDParticleActorCustomization.cpp b/Engine/Plugins/ChaosVD/Source/ChaosVD/Private/ChaosVDParticleActorCustomization.cpp
--- a/Engine/Plugins/ChaosVD/Source/ChaosVD/Private/ChaosVDParticleActorCustomization.cpp
+++ b/Engine/Plugins/ChaosVD/Source/ChaosVD/Private/ChaosVDParticleActorCustomization.cpp
@@ -71,7 +71,7 @@ void FChaosVDParticleActorCustomization::CustomizeDetails(IDetailLayoutBuilder&
 
 			if (TSharedPtr<FChaosVDMeshDataInstanceHandle> SelectedGeometryInstance = ParticleActor->GetSelectedMeshInstance().Pin())
 			{
-				InspectedDataPropertyHandlePtr = AddExternalStructure(CachedGeometryDataInstanceCopy, DetailBuilder, GeometryCategoryName, LOCTEXT("ParticleDataStructName", "Geometry Shape Data"));
+				InspectedDataPropertyHandlePtr = AddExternalStructure(CachedGeometryDataInstanceCopy, DetailBuilder, GeometryCategoryName, LOCTEXT("GeometryShapeDataStructName", "Geometry Shape Data"));
 			}
 			else
 			{
diff --git a/Engine/Plugins/ChaosVD/Source/ChaosVD/Private/ChaosVDParticleActorCustomization.h b/Engine/Plugins/ChaosVD/Source/ChaosVD/Private/ChaosVDParticleActorCustomization.h
--- a/Engine/Plugins/ChaosVD/Source/ChaosVD/Private/ChaosVDParticleActorCustomization.h
+++ b/Engine/Plugins/ChaosVD/Source/ChaosVD/Private/ChaosVDParticleActorCustomization.h
@@ -37,9 +37,9 @@ private:
 
 	void HandleParticleDataUpdated();
 
-	/* Copied of the last known particle data structure of a selected particle -  Used to avoid rebuild the layout every time we change frame in CVD */
+	/* Copy of the last known geometry shape data structure of a selected particle and mesh instance -  Used to avoid rebuild the layout every time we change frame in CVD */
 	FChaosVDParticleDataWrapper CachedParticleData;
-	/* Copied of the last known particle data structure of a selected particle -  Used to avoid rebuild the layout every time we change frame in CVD */
+	/* Copy of the last known particle data structure of a selected particle -  Used to avoid rebuild the layout every time we change frame in CVD */
 	FChaosVDMeshDataInstanceState CachedGeometryDataInstanceCopy;
 };
 
