[ChaosVD] Adding Viewport option to go to specific locations by using a vector in string format (similar to how you can go to locations using the BugItGo console command)

diff --git a/Engine/Plugins/ChaosVD/Source/ChaosVD/Private/ChaosVDPlaybackViewportClient.cpp b/Engine/Plugins/ChaosVD/Source/ChaosVD/Private/ChaosVDPlaybackViewportClient.cpp
--- a/Engine/Plugins/ChaosVD/Source/ChaosVD/Private/ChaosVDPlaybackViewportClient.cpp
+++ b/Engine/Plugins/ChaosVD/Source/ChaosVD/Private/ChaosVDPlaybackViewportClient.cpp
@@ -241,6 +241,14 @@ void FChaosVDPlaybackViewportClient::ToggleObjectTrackingIfSelected()
 	}
 }
 
+void FChaosVDPlaybackViewportClient::GoToLocation(const FVector& InLocation)
+{
+	FViewportCameraTransform& ViewTransform = GetViewTransform();
+	ViewTransform.SetLocation(InLocation);
+
+	Invalidate();
+}
+
 void FChaosVDPlaybackViewportClient::Draw(const FSceneView* View, FPrimitiveDrawInterface* PDI)
 {
 	if (View)
diff --git a/Engine/Plugins/ChaosVD/Source/ChaosVD/Private/ChaosVDPlaybackViewportClient.h b/Engine/Plugins/ChaosVD/Source/ChaosVD/Private/ChaosVDPlaybackViewportClient.h
--- a/Engine/Plugins/ChaosVD/Source/ChaosVD/Private/ChaosVDPlaybackViewportClient.h
+++ b/Engine/Plugins/ChaosVD/Source/ChaosVD/Private/ChaosVDPlaybackViewportClient.h
@@ -31,6 +31,8 @@ public:
 
 	void ToggleObjectTrackingIfSelected();
 
+	void GoToLocation(const FVector& InLocation);
+
 private:
 
 	void TrackSelectedObject();
diff --git a/Engine/Plugins/ChaosVD/Source/ChaosVD/Private/Widgets/SChaosVDPlaybackViewport.cpp b/Engine/Plugins/ChaosVD/Source/ChaosVD/Private/Widgets/SChaosVDPlaybackViewport.cpp
--- a/Engine/Plugins/ChaosVD/Source/ChaosVD/Private/Widgets/SChaosVDPlaybackViewport.cpp
+++ b/Engine/Plugins/ChaosVD/Source/ChaosVD/Private/Widgets/SChaosVDPlaybackViewport.cpp
@@ -138,6 +138,14 @@ EVisibility SChaosVDPlaybackViewport::GetTransformToolbarVisibility() const
 	return EVisibility::Visible;
 }
 
+void SChaosVDPlaybackViewport::GoToLocation(const FVector& InLocation) const
+{
+	if (PlaybackViewportClient)
+	{
+		PlaybackViewportClient->GoToLocation(InLocation);
+	}
+}
+
 TSharedRef<FEditorViewportClient> SChaosVDPlaybackViewport::MakeEditorViewportClient()
 {
 	TSharedPtr<FChaosVDPlaybackViewportClient> NewViewport = MakeShared<FChaosVDPlaybackViewportClient>(EditorModeTools, GetViewportWidget());
diff --git a/Engine/Plugins/ChaosVD/Source/ChaosVD/Private/Widgets/SChaosVDPlaybackViewport.h b/Engine/Plugins/ChaosVD/Source/ChaosVD/Private/Widgets/SChaosVDPlaybackViewport.h
--- a/Engine/Plugins/ChaosVD/Source/ChaosVD/Private/Widgets/SChaosVDPlaybackViewport.h
+++ b/Engine/Plugins/ChaosVD/Source/ChaosVD/Private/Widgets/SChaosVDPlaybackViewport.h
@@ -44,6 +44,8 @@ public:
 
 	virtual EVisibility GetTransformToolbarVisibility() const override;
 
+	void GoToLocation(const FVector& InLocation) const;
+
 protected:
 
 	virtual TSharedRef<FEditorViewportClient> MakeEditorViewportClient() override;
diff --git a/Engine/Plugins/ChaosVD/Source/ChaosVD/Private/Widgets/SChaosVDViewportToolbar.cpp b/Engine/Plugins/ChaosVD/Source/ChaosVD/Private/Widgets/SChaosVDViewportToolbar.cpp
--- a/Engine/Plugins/ChaosVD/Source/ChaosVD/Private/Widgets/SChaosVDViewportToolbar.cpp
+++ b/Engine/Plugins/ChaosVD/Source/ChaosVD/Private/Widgets/SChaosVDViewportToolbar.cpp
@@ -2,6 +2,7 @@
 
 #include "Widgets/SChaosVDViewportToolbar.h"
 
+#include "ChaosVDCommands.h"
 #include "ChaosVDEditorSettings.h"
 #include "Framework/MultiBox/MultiBoxBuilder.h"
 #include "PropertyEditorModule.h"
@@ -11,6 +12,10 @@
 #include "PropertyEditorModule.h"
 #include "SCommonEditorViewportToolbarBase.h"
 #include "SEditorViewport.h"
+#include "Styling/AppStyle.h"
+#include "Widgets/Input/SEditableText.h"
+#include "Widgets/Layout/SBorder.h"
+#include "Widgets/Layout/SBox.h"
 #include "Widgets/SChaosVDEditorViewportViewMenu.h"
 #include "Widgets/SChaosVDPlaybackViewport.h"
 #include "Widgets/SChaosVDVisualizationControls.h"
@@ -33,20 +38,36 @@ void SChaosVDViewportToolbar::ExtendOptionsMenu(FMenuBuilder& OptionsMenuBuilder
 	constexpr bool bInShouldCloseWindowAfterMenuSelection = true;
 	FMenuBuilder CVDOptionsMenuBuilder(bInShouldCloseWindowAfterMenuSelection, GetInfoProvider().GetViewportWidget()->GetCommandList());
 
-	TSharedRef<SVerticalBox> OptionsMenuWidget = SNew(SVerticalBox)
-													+SVerticalBox::Slot()
-													.Padding(5.0f)
-													[
-														SNew(STextBlock)
-														.Text(LOCTEXT("Default Options Message", "Nothing to see yet. An options menu will be implemented soon"))
-													];
-
-	// Intentionally leaving the label empty as we only want to show the temp message defined above
-	CVDOptionsMenuBuilder.AddWidget(OptionsMenuWidget, FText());
+	CVDOptionsMenuBuilder.BeginSection("CVDViewportViewportOptions", LOCTEXT("OptionsMenuHeader", "Utils"));
+	
+	CVDOptionsMenuBuilder.AddWidget(GenerateGoToLocationWidget(), LOCTEXT("GoToLocation", "Go to Location"));
+	CVDOptionsMenuBuilder.EndSection();
 
 	OptionsMenuBuilder = CVDOptionsMenuBuilder;
 }
 
+TSharedRef<SWidget> SChaosVDViewportToolbar::GenerateGoToLocationWidget() const
+{
+	return SNew(SBox)
+			.HAlign(HAlign_Right)
+			[
+				SNew(SBox)
+				.Padding(FMargin(4.0f, 0.0f, 0.0f, 0.0f))
+				.WidthOverride(100.0f)
+				[
+					SNew (SBorder)
+					.BorderImage(FAppStyle::Get().GetBrush("Menu.WidgetBorder"))
+					.Padding(FMargin(1.0f))
+					[
+						SNew(SEditableText)
+						.ToolTipText(LOCTEXT("GoToLocationTooltip", "Location to teleport to."))
+						.Font(FAppStyle::GetFontStyle(TEXT("MenuItem.Font")))
+						.OnTextCommitted(FOnTextCommitted::CreateSP(this, &SChaosVDViewportToolbar::HandleGoToLocationCommited))
+					]
+				]
+			];
+}
+
 TSharedRef<SWidget> SChaosVDViewportToolbar::GenerateShowMenu() const
 {
 	FPropertyEditorModule& PropertyEditorModule = FModuleManager::GetModuleChecked<FPropertyEditorModule>("PropertyEditor");
@@ -62,4 +83,20 @@ TSharedRef<SWidget> SChaosVDViewportToolbar::GenerateShowMenu() const
 
 	return SNew(SChaosVDVisualizationControls, ViewportRef->GetObservedController(), ViewportRef);
 }
+
+void SChaosVDViewportToolbar::HandleGoToLocationCommited(const FText& InLocationAsText, ETextCommit::Type Type) const
+{
+	if (Type != ETextCommit::OnEnter)
+	{
+		return;
+	}
+
+	FVector Location;
+	Location.InitFromString(InLocationAsText.ToString());
+
+	TSharedRef<SChaosVDPlaybackViewport> ViewportRef = StaticCastSharedRef<SChaosVDPlaybackViewport>(GetInfoProvider().GetViewportWidget());
+
+	ViewportRef->GoToLocation(Location);
+	
+}
 #undef LOCTEXT_NAMESPACE
diff --git a/Engine/Plugins/ChaosVD/Source/ChaosVD/Private/Widgets/SChaosVDViewportToolbar.h b/Engine/Plugins/ChaosVD/Source/ChaosVD/Private/Widgets/SChaosVDViewportToolbar.h
--- a/Engine/Plugins/ChaosVD/Source/ChaosVD/Private/Widgets/SChaosVDViewportToolbar.h
+++ b/Engine/Plugins/ChaosVD/Source/ChaosVD/Private/Widgets/SChaosVDViewportToolbar.h
@@ -27,7 +27,11 @@ protected:
 	virtual TSharedRef<SEditorViewportViewMenu> MakeViewMenu() override;
 	virtual void ExtendOptionsMenu(FMenuBuilder& OptionsMenuBuilder) const override;
 
+	TSharedRef<SWidget> GenerateGoToLocationWidget() const;
+
 private:
 	virtual TSharedRef<SWidget> GenerateShowMenu() const override;
 
+	void HandleGoToLocationCommited(const FText& InLocationAsText, ETextCommit::Type Type) const;
+
 };
