[ChaosVD] Add Character Ground Constraint

diff --git a/Engine/Plugins/ChaosVD/Source/ChaosVD/Private/Actors/ChaosVDSolverInfoActor.cpp b/Engine/Plugins/ChaosVD/Source/ChaosVD/Private/Actors/ChaosVDSolverInfoActor.cpp
--- a/Engine/Plugins/ChaosVD/Source/ChaosVD/Private/Actors/ChaosVDSolverInfoActor.cpp
+++ b/Engine/Plugins/ChaosVD/Source/ChaosVD/Private/Actors/ChaosVDSolverInfoActor.cpp
@@ -7,6 +7,7 @@
 #include "ChaosVDScene.h"
 #include "EditorActorFolders.h"
 #include "Components/ChaosVDParticleDataComponent.h"
+#include "Components/ChaosVDSolverCharacterGroundConstraintDataComponent.h"
 #include "Components/ChaosVDSolverCollisionDataComponent.h"
 #include "Components/ChaosVDSolverJointConstraintDataComponent.h"
 #include "Elements/Framework/TypedElementSelectionSet.h"
@@ -20,6 +21,7 @@ AChaosVDSolverInfoActor::AChaosVDSolverInfoActor(const FObjectInitializer& Objec
 	CollisionDataComponent = CreateDefaultSubobject<UChaosVDSolverCollisionDataComponent>(TEXT("SolverCollisionDataComponent"));
 	ParticleDataComponent = CreateDefaultSubobject<UChaosVDParticleDataComponent>(TEXT("ParticleCollisionDataComponent"));
 	JointsDataComponent = CreateDefaultSubobject<UChaosVDSolverJointConstraintDataComponent>(TEXT("JointDataComponent"));
+	CharacterGroundConstraintDataComponent = CreateDefaultSubobject<UChaosVDSolverCharacterGroundConstraintDataComponent>(TEXT("CharacterGroundConstraintDataComponent"));
 	bIsServer = false;
 }
 
diff --git a/Engine/Plugins/ChaosVD/Source/ChaosVD/Private/Actors/ChaosVDSolverInfoActor.h b/Engine/Plugins/ChaosVD/Source/ChaosVD/Private/Actors/ChaosVDSolverInfoActor.h
--- a/Engine/Plugins/ChaosVD/Source/ChaosVD/Private/Actors/ChaosVDSolverInfoActor.h
+++ b/Engine/Plugins/ChaosVD/Source/ChaosVD/Private/Actors/ChaosVDSolverInfoActor.h
@@ -8,6 +8,7 @@
 #include "GameFramework/Actor.h"
 #include "ChaosVDSolverInfoActor.generated.h"
 
+class UChaosVDSolverCharacterGroundConstraintDataComponent;
 class UChaosVDSolverJointConstraintDataComponent;
 struct FChaosVDParticleDataWrapper;
 class AChaosVDParticleActor;
@@ -42,6 +43,7 @@ public:
 	UChaosVDSolverCollisionDataComponent* GetCollisionDataComponent() { return CollisionDataComponent; }
 	UChaosVDParticleDataComponent* GetParticleDataComponent() { return ParticleDataComponent; }
 	UChaosVDSolverJointConstraintDataComponent* GetJointsDataComponent() { return JointsDataComponent; }
+	UChaosVDSolverCharacterGroundConstraintDataComponent* GetCharacterGroundConstraintDataComponent() { return CharacterGroundConstraintDataComponent; }
 
 	void RegisterParticleActor(int32 ParticleID, AChaosVDParticleActor* ParticleActor);
 
@@ -108,6 +110,8 @@ protected:
 	UPROPERTY()
 	TObjectPtr<UChaosVDSolverJointConstraintDataComponent> JointsDataComponent;
 
+	UPROPERTY()
+	TObjectPtr<UChaosVDSolverCharacterGroundConstraintDataComponent> CharacterGroundConstraintDataComponent;
 };
 
 template <typename TCallback>
diff --git a/Engine/Plugins/ChaosVD/Source/ChaosVD/Private/ChaosVDCharacterGroundConstraintDataInspectorTab.cpp b/Engine/Plugins/ChaosVD/Source/ChaosVD/Private/ChaosVDCharacterGroundConstraintDataInspectorTab.cpp
new file mode 100644
--- /dev/null
+++ b/Engine/Plugins/ChaosVD/Source/ChaosVD/Private/ChaosVDCharacterGroundConstraintDataInspectorTab.cpp
@@ -0,0 +1,74 @@
+// Copyright Epic Games, Inc. All Rights Reserved.
+
+#include "ChaosVDCharacterGroundConstraintDataInspectorTab.h"
+
+#include "ChaosVDCharacterGroundConstraintDataProviderInterface.h"
+#include "ChaosVDScene.h"
+#include "Elements/Framework/TypedElementSelectionSet.h"
+#include "Widgets/SChaosVDCharacterGroundConstraintDataInspector.h"
+
+#define LOCTEXT_NAMESPACE "ChaosVisualDebugger"
+
+FChaosVDCharacterGroundConstraintDataInspectorTab::~FChaosVDCharacterGroundConstraintDataInspectorTab()
+{
+}
+
+TSharedRef<SDockTab> FChaosVDCharacterGroundConstraintDataInspectorTab::HandleTabSpawnRequest(const FSpawnTabArgs& Args)
+{
+	TSharedRef<SDockTab> DetailsPanelTab =
+	SNew(SDockTab)
+	.TabRole(ETabRole::MajorTab)
+	.Label(LOCTEXT("CharacterGroundConstraintDataInspectorTab", "Character Ground Constraint Inspector"))
+	.ToolTipText(LOCTEXT("CharacterGroundConstraintDataInspectorTabToolTip", "See the details of the any character ground constraint selected in the viewport"));
+
+	if (const TSharedPtr<SChaosVDMainTab> MainTabPtr = OwningTabWidget.Pin())
+	{
+		DetailsPanelTab->SetContent
+		(
+			SAssignNew(ConstraintDataInspector, SChaosVDCharacterGroundConstraintDataInspector, GetChaosVDScene())
+		);
+	}
+	else
+	{
+		DetailsPanelTab->SetContent(GenerateErrorWidget());
+	}
+
+	HandleTabSpawned(DetailsPanelTab);
+
+	if (const TSharedPtr<FChaosVDScene> ScenePtr = GetChaosVDScene().Pin())
+	{
+		RegisterSelectionSetObject(ScenePtr->GetElementSelectionSet());
+	}
+
+	return DetailsPanelTab;
+}
+
+void FChaosVDCharacterGroundConstraintDataInspectorTab::HandleTabClosed(TSharedRef<SDockTab> InTabClosed)
+{
+	FChaosVDTabSpawnerBase::HandleTabClosed(InTabClosed);
+
+	ConstraintDataInspector.Reset();
+}
+
+void FChaosVDCharacterGroundConstraintDataInspectorTab::HandlePostSelectionChange(const UTypedElementSelectionSet* ChangedSelectionSet)
+{
+	TArray<AActor*> SelectedActors = ChangedSelectionSet->GetSelectedObjects<AActor>();
+
+	if (SelectedActors.Num() > 0)
+	{
+		// We don't support multi selection yet
+		for (AActor* SelectedActor : SelectedActors)
+		{
+			if (IChaosVDCharacterGroundConstraintDataProviderInterface* ConstraintDataProvider = Cast<IChaosVDCharacterGroundConstraintDataProviderInterface>(SelectedActor))
+			{
+				if (ConstraintDataProvider->HasCharacterGroundConstraintData())
+				{
+					ConstraintDataInspector->SetConstraintDataProviderObjectToInspect(ConstraintDataProvider);
+					break;
+				}
+			}
+		}
+	}
+}
+
+#undef LOCTEXT_NAMESPACE
diff --git a/Engine/Plugins/ChaosVD/Source/ChaosVD/Private/ChaosVDCharacterGroundConstraintDataInspectorTab.h b/Engine/Plugins/ChaosVD/Source/ChaosVD/Private/ChaosVDCharacterGroundConstraintDataInspectorTab.h
new file mode 100644
--- /dev/null
+++ b/Engine/Plugins/ChaosVD/Source/ChaosVD/Private/ChaosVDCharacterGroundConstraintDataInspectorTab.h
@@ -0,0 +1,31 @@
+// Copyright Epic Games, Inc. All Rights Reserved.
+
+#pragma once
+
+#include "ChaosVDSceneSelectionObserver.h"
+#include "ChaosVDTabSpawnerBase.h"
+#include "Templates/SharedPointer.h"
+
+class SChaosVDCharacterGroundConstraintDataInspector;
+class SChaosVDSceneQueryDataInspector;
+
+/** Spawns and handles and instance for the visual debugger character ground constraint inspector tab */
+class FChaosVDCharacterGroundConstraintDataInspectorTab final : public FChaosVDTabSpawnerBase, public FChaosVDSceneSelectionObserver
+{
+public:
+	FChaosVDCharacterGroundConstraintDataInspectorTab(const FName& InTabID, const TSharedPtr<FTabManager>& InTabManager, const TWeakPtr<SChaosVDMainTab>& InOwningTabWidget)
+		: FChaosVDTabSpawnerBase(InTabID, InTabManager, InOwningTabWidget)
+	{
+	}
+
+	virtual ~FChaosVDCharacterGroundConstraintDataInspectorTab() override;
+	virtual TSharedRef<SDockTab> HandleTabSpawnRequest(const FSpawnTabArgs& Args) override;
+	virtual void HandleTabClosed(TSharedRef<SDockTab> InTabClosed) override;
+
+	virtual void HandlePostSelectionChange(const UTypedElementSelectionSet* ChangedSelectionSet) override;
+
+	TWeakPtr<SChaosVDCharacterGroundConstraintDataInspector> GetConstraintDataInspectorInstance() const { return ConstraintDataInspector; }
+
+protected:
+	TSharedPtr<SChaosVDCharacterGroundConstraintDataInspector> ConstraintDataInspector;
+};
diff --git a/Engine/Plugins/ChaosVD/Source/ChaosVD/Private/ChaosVDCharacterGroundConstraintDataProviderInterface.h b/Engine/Plugins/ChaosVD/Source/ChaosVD/Private/ChaosVDCharacterGroundConstraintDataProviderInterface.h
new file mode 100644
--- /dev/null
+++ b/Engine/Plugins/ChaosVD/Source/ChaosVD/Private/ChaosVDCharacterGroundConstraintDataProviderInterface.h
@@ -0,0 +1,42 @@
+// Copyright Epic Games, Inc. All Rights Reserved.
+
+#pragma once
+
+#include "UObject/Interface.h"
+#include "ChaosVDCharacterGroundConstraintDataProviderInterface.generated.h"
+
+struct FChaosVDCharacterGroundConstraint;
+
+// This class does not need to be modified.
+UINTERFACE()
+class UChaosVDCharacterGroundConstraintDataProviderInterface : public UInterface
+{
+	GENERATED_BODY()
+};
+
+/**
+ * Interface for any object that is able to provide CVD Character Ground Constraint Data
+ */
+class IChaosVDCharacterGroundConstraintDataProviderInterface
+{
+	GENERATED_BODY()
+
+public:
+	
+	/**
+	 * Gathers and populates the provided array with any existing character ground constraint data for this object (if any) 
+	 * @param OutConstraintDataFound Array to populate with constraint data
+	 */
+	virtual void GetCharacterGroundConstraintData(TArray<TSharedPtr<FChaosVDCharacterGroundConstraint>>& OutConstraintsFound) PURE_VIRTUAL(IChaosVDCharacterGroundConstraintDataProviderInterface::GetCharacterGroundConstraintData)
+
+	/**
+	 * Checks if the object implementing the interface has any character ground constraint data available
+	 * @return True if any constraint data is found 
+	 */
+	virtual bool HasCharacterGroundConstraintData() PURE_VIRTUAL(IChaosVDCharacterGroundConstraintDataProviderInterface::HasCharacterGroundConstraintData, return false;)
+
+	/**
+	 * Returns the name of the object providing the constraint data
+	 */
+	virtual FName GetProviderName() PURE_VIRTUAL(IChaosVDCharacterGroundConstraintDataProviderInterface::GetProviderName, return FName();)
+};
diff --git a/Engine/Plugins/ChaosVD/Source/ChaosVD/Private/ChaosVDEditorSettings.h b/Engine/Plugins/ChaosVD/Source/ChaosVD/Private/ChaosVDEditorSettings.h
--- a/Engine/Plugins/ChaosVD/Source/ChaosVD/Private/ChaosVDEditorSettings.h
+++ b/Engine/Plugins/ChaosVD/Source/ChaosVD/Private/ChaosVDEditorSettings.h
@@ -7,6 +7,7 @@
 #include "Engine/StaticMesh.h"
 #include "UObject/Object.h"
 #include "UObject/SoftObjectPtr.h"
+#include "Visualizers/ChaosVDCharacterGroundConstraintsDataComponentVisualizer.h"
 #include "Visualizers/ChaosVDJointConstraintsDataComponentVisualizer.h"
 #include "Visualizers/ChaosVDParticleDataComponentVisualizer.h"
 #include "Visualizers/ChaosVDSceneQueryDataComponentVisualizer.h"
@@ -210,6 +211,44 @@ struct FChaosVDJointsDebugDrawSettings
     float ConstraintAxisLength = 10.0f;
 };
 
+USTRUCT()
+struct FChaosVDCharacterGroundConstraintDebugDrawSettings
+{
+	GENERATED_BODY()
+
+	/** The depth priority used for while drawing data. Can be World or Foreground (with this one the shapes will be drawn on top of the geometry and be always visible) */
+	UPROPERTY(EditAnywhere, Category=DebugDraw)
+	TEnumAsByte<ESceneDepthPriorityGroup> DepthPriority = ESceneDepthPriorityGroup::SDPG_Foreground;
+
+	/** Scale to apply to the force vector before draw it. */
+	UPROPERTY(EditAnywhere, Category=DebugDraw)
+	float ForceScale = 0.01f;
+
+	/** Scale to apply to the torque vector before draw it. */
+	UPROPERTY(EditAnywhere, Category=DebugDraw)
+	float TorqueScale = 0.01f;
+
+	/** Scale to apply to anything that does not have a dedicated scale setting before draw it. */
+	UPROPERTY(EditAnywhere, Category=DebugDraw)
+	float GeneralScale = 1.0f;
+
+	/** Line thickness to use as a base to calculate the different line thickness values used to debug draw the data. */
+	UPROPERTY(EditAnywhere, Category=DebugDraw)
+	float BaseLineThickness = 2.0f;
+
+	/** Color used for normal force vector */
+	UPROPERTY(EditAnywhere, Category = DebugDraw)
+	FColor NormalForceColor = FColor(255, 0, 0);
+
+	/** Color used for radial force vector */
+	UPROPERTY(EditAnywhere, Category = DebugDraw)
+	FColor RadialForceColor = FColor(255, 255, 0);
+
+	/** Color used for torque vector */
+	UPROPERTY(EditAnywhere, Category = DebugDraw)
+	FColor TorqueColor = FColor(255, 0, 255);
+};
+
 /** Structure holding the settings using to debug draw Particles shape based on their shape type on the Chaos Visual Debugger */
 USTRUCT()
 struct FChaosDebugDrawColorsByShapeType
@@ -319,6 +358,10 @@ public:
 	UPROPERTY(EditAnywhere, Category = "Viewport Visualization Flags", meta = (Bitmask, BitmaskEnum = "/Script/ChaosVD.EChaosVDJointsDataVisualizationFlags"))
 	uint32 GlobalJointsDataVisualizationFlags = static_cast<uint32>(EChaosVDJointsDataVisualizationFlags::ActorConnector | EChaosVDJointsDataVisualizationFlags::DrawKinematic);
 
+	/** Set of flags to enable/disable visualization of specific character ground constraint data as debug draw */
+	UPROPERTY(EditAnywhere, Category = "Viewport Visualization Flags", meta = (Bitmask, BitmaskEnum = "/Script/ChaosVD.EChaosVDCharacterGroundConstraintDataVisualizationFlags"))
+	uint32 GlobalCharacterGroundConstraintDataVisualizationFlags = static_cast<uint32>(EChaosVDCharacterGroundConstraintDataVisualizationFlags::EnableDraw | EChaosVDCharacterGroundConstraintDataVisualizationFlags::GroundQueryDistance | EChaosVDCharacterGroundConstraintDataVisualizationFlags::GroundQueryNormal | EChaosVDCharacterGroundConstraintDataVisualizationFlags::TargetDeltaPosition);
+
 	/** If true, text information (if available) will be drawn alongside any other debug draw shape */
 	UPROPERTY(EditAnywhere, Category = "Viewport Visualization")
 	bool bShowDebugText = false;
@@ -332,6 +375,9 @@ public:
 	UPROPERTY(EditAnywhere, Category = "Viewport Visualization",  meta=(EditCondition = "GlobalParticleDataVisualizationFlags != 0", EditConditionHides))
 	FChaosVDJointsDebugDrawSettings JointsDataDebugDrawSettings;
 
+	UPROPERTY(EditAnywhere, Category = "Viewport Visualization", meta = (EditCondition = "GlobalParticleDataVisualizationFlags != 0", EditConditionHides))
+	FChaosVDCharacterGroundConstraintDebugDrawSettings CharacterGroundConstraintDataDebugDrawSettings;
+
 	UPROPERTY(EditAnywhere, Category = "Viewport Visualization")
 	EChaosVDParticleDebugColorMode ParticleColorMode;
 	
diff --git a/Engine/Plugins/ChaosVD/Source/ChaosVD/Private/ChaosVDParticleActor.cpp b/Engine/Plugins/ChaosVD/Source/ChaosVD/Private/ChaosVDParticleActor.cpp
--- a/Engine/Plugins/ChaosVD/Source/ChaosVD/Private/ChaosVDParticleActor.cpp
+++ b/Engine/Plugins/ChaosVD/Source/ChaosVD/Private/ChaosVDParticleActor.cpp
@@ -447,6 +447,29 @@ void AChaosVDParticleActor::UpdateMeshInstancesSelectionState()
 	});
 }
 
+void AChaosVDParticleActor::GetCharacterGroundConstraintData(TArray<TSharedPtr<FChaosVDCharacterGroundConstraint>>& OutConstraintsFound)
+{
+	if (const TArray<TSharedPtr<FChaosVDCharacterGroundConstraint>>* Constraints = GetCharacterGroundConstraintArray())
+	{
+		OutConstraintsFound.Reserve(Constraints->Num());
+
+		for (const TSharedPtr<FChaosVDCharacterGroundConstraint>& Constraint : *Constraints)
+		{
+			OutConstraintsFound.Add(Constraint);
+		}
+	}
+}
+
+bool AChaosVDParticleActor::HasCharacterGroundConstraintData()
+{
+	if (const TArray<TSharedPtr<FChaosVDCharacterGroundConstraint>>* Constraints = GetCharacterGroundConstraintArray())
+	{
+		return Constraints->Num() > 0;
+	}
+
+	return false;
+}
+
 void AChaosVDParticleActor::PushSelectionToProxies()
 {
 	UpdateMeshInstancesSelectionState();
@@ -477,6 +500,30 @@ const TArray<TSharedPtr<FChaosVDParticlePairMidPhase>>* AChaosVDParticleActor::G
 	return nullptr;
 }
 
+const TArray<TSharedPtr<FChaosVDCharacterGroundConstraint>>* AChaosVDParticleActor::GetCharacterGroundConstraintArray() const
+{
+	if (!ParticleDataPtr.IsValid())
+	{
+		return nullptr;
+	}
+
+	const TSharedPtr<FChaosVDScene> ScenePtr = SceneWeakPtr.Pin();
+	if (!ScenePtr.IsValid())
+	{
+		return nullptr;
+	}
+
+	if (AChaosVDSolverInfoActor* SolverInfoActor = ScenePtr->GetSolverInfoActor(ParticleDataPtr->SolverID))
+	{
+		if (const UChaosVDSolverCharacterGroundConstraintDataComponent* ConstraintDataComponent = SolverInfoActor->GetCharacterGroundConstraintDataComponent())
+		{
+			return ConstraintDataComponent->GetConstraintsForParticle(ParticleDataPtr->ParticleIndex);
+		}
+	}
+
+	return nullptr;
+}
+
 void AChaosVDParticleActor::UpdateShapeDataComponents()
 {
 	VisitGeometryInstances([this](const TSharedRef<FChaosVDMeshDataInstanceHandle>& MeshDataHandle)
diff --git a/Engine/Plugins/ChaosVD/Source/ChaosVD/Private/ChaosVDParticleActor.h b/Engine/Plugins/ChaosVD/Source/ChaosVD/Private/ChaosVDParticleActor.h
--- a/Engine/Plugins/ChaosVD/Source/ChaosVD/Private/ChaosVDParticleActor.h
+++ b/Engine/Plugins/ChaosVD/Source/ChaosVD/Private/ChaosVDParticleActor.h
@@ -1,6 +1,8 @@
 // Copyright Epic Games, Inc. All Rights Reserved.
 
 #pragma once
+
+#include "ChaosVDCharacterGroundConstraintDataProviderInterface.h"
 #include "ChaosVDCollisionDataProviderInterface.h"
 #include "ChaosVDGeometryDataComponent.h"
 #include "ChaosVDSceneObjectBase.h"
@@ -53,7 +55,8 @@ ENUM_CLASS_FLAGS(EChaosVDHideParticleFlags)
 /** Actor used to represent a Chaos Particle in the Visual Debugger's world */
 UCLASS(HideCategories=(Transform))
 class AChaosVDParticleActor : public AActor, public IChaosVDParticleVisualizationDataProvider,
-								public FChaosVDSceneObjectBase, public IChaosVDCollisionDataProviderInterface, public IChaosVDGeometryOwnerInterface, public IChaosVDSelectableObject
+								public FChaosVDSceneObjectBase, public IChaosVDCollisionDataProviderInterface, public IChaosVDGeometryOwnerInterface, public IChaosVDSelectableObject,
+								public IChaosVDCharacterGroundConstraintDataProviderInterface
 {
 	GENERATED_BODY()
 
@@ -107,6 +110,11 @@ public:
 	virtual FName GetProviderName() override;
 	// END IChaosVDCollisionDataProvider Interface
 
+	// BEGIN IChaosVDCharacterGroundConstraintDataProviderInterface
+	virtual void GetCharacterGroundConstraintData(TArray<TSharedPtr<FChaosVDCharacterGroundConstraint>>& OutConstraintsFound) override;
+	virtual bool HasCharacterGroundConstraintData() override;
+	// END IChaosVDCharacterGroundConstraintDataProviderInterface
+
 	void SetIsServerParticle(bool bNewIsServer) { bIsServer = bNewIsServer; }
 	bool GetIsServerParticle() const { return bIsServer; }
 
@@ -132,6 +140,8 @@ protected:
 
 	const TArray<TSharedPtr<FChaosVDParticlePairMidPhase>>* GetCollisionMidPhasesArray() const;
 
+	const TArray<TSharedPtr<FChaosVDCharacterGroundConstraint>>* GetCharacterGroundConstraintArray() const;
+
 	void UpdateShapeDataComponents();
 
 	template<typename TTaskCallback>
diff --git a/Engine/Plugins/ChaosVD/Source/ChaosVD/Private/ChaosVDScene.cpp b/Engine/Plugins/ChaosVD/Source/ChaosVD/Private/ChaosVDScene.cpp
--- a/Engine/Plugins/ChaosVD/Source/ChaosVD/Private/ChaosVDScene.cpp
+++ b/Engine/Plugins/ChaosVD/Source/ChaosVD/Private/ChaosVDScene.cpp
@@ -28,6 +28,7 @@
 #include "Selection.h"
 #include "UObject/Package.h"
 #include "WorldPersistentFolders.h"
+#include "Components/ChaosVDSolverCharacterGroundConstraintDataComponent.h"
 #include "Components/ChaosVDSolverJointConstraintDataComponent.h"
 #include "Elements/Actor/ActorElementData.h"
 #include "Elements/Component/ComponentElementData.h"
@@ -340,6 +341,17 @@ void FChaosVDScene::HandleEnterNewGameFrame(int32 FrameNumber, const TArray<int3
 			QueryDataComponent->UpdateQueriesFromFrameData(InNewGameFrameData);
 		}
 	}
+
+	for (int32 SolverID : AvailableSolversIds)
+	{
+		if (AChaosVDSolverInfoActor* SolverDataInfoContainer = SolverDataContainerBySolverID.FindChecked(SolverID))
+		{
+			if (UChaosVDSolverCharacterGroundConstraintDataComponent* DataContainer = SolverDataInfoContainer->GetCharacterGroundConstraintDataComponent())
+			{
+				DataContainer->UpdateConstraintData(InNewGameFrameData.RecordedCharacterGroundConstraints);
+			}
+		}
+	}
 }
 
 void FChaosVDScene::CleanUpScene()
diff --git a/Engine/Plugins/ChaosVD/Source/ChaosVD/Private/ChaosVDTabsIDs.cpp b/Engine/Plugins/ChaosVD/Source/ChaosVD/Private/ChaosVDTabsIDs.cpp
--- a/Engine/Plugins/ChaosVD/Source/ChaosVD/Private/ChaosVDTabsIDs.cpp
+++ b/Engine/Plugins/ChaosVD/Source/ChaosVD/Private/ChaosVDTabsIDs.cpp
@@ -13,4 +13,5 @@ const FName FChaosVDTabID::CVDEditorSettings = TEXT("CVDEditorSettings");
 const FName FChaosVDTabID::StatusBar = TEXT("StatusBar");
 const FName FChaosVDTabID::CollisionDataDetails = TEXT("CollisionDataDetails");
 const FName FChaosVDTabID::SceneQueryDataDetails = TEXT("SceneQueryDataDetails");
-const FName FChaosVDTabID::JointsDataDetails = TEXT("JointsDataDetails");
\ No newline at end of file
+const FName FChaosVDTabID::JointsDataDetails = TEXT("JointsDataDetails");
+const FName FChaosVDTabID::CharacterGroundConstraintDataDetails = TEXT("CharacterGroundConstraintDataDetails");
\ No newline at end of file
diff --git a/Engine/Plugins/ChaosVD/Source/ChaosVD/Private/ChaosVDTabsIDs.h b/Engine/Plugins/ChaosVD/Source/ChaosVD/Private/ChaosVDTabsIDs.h
--- a/Engine/Plugins/ChaosVD/Source/ChaosVD/Private/ChaosVDTabsIDs.h
+++ b/Engine/Plugins/ChaosVD/Source/ChaosVD/Private/ChaosVDTabsIDs.h
@@ -19,4 +19,5 @@ public:
 	static const FName CollisionDataDetails;
 	static const FName SceneQueryDataDetails;
 	static const FName JointsDataDetails;
+	static const FName CharacterGroundConstraintDataDetails;
 };
diff --git a/Engine/Plugins/ChaosVD/Source/ChaosVD/Private/Components/ChaosVDSolverCharacterGroundConstraintDataComponent.cpp b/Engine/Plugins/ChaosVD/Source/ChaosVD/Private/Components/ChaosVDSolverCharacterGroundConstraintDataComponent.cpp
new file mode 100644
--- /dev/null
+++ b/Engine/Plugins/ChaosVD/Source/ChaosVD/Private/Components/ChaosVDSolverCharacterGroundConstraintDataComponent.cpp
@@ -0,0 +1,103 @@
+// Copyright Epic Games, Inc. All Rights Reserved.
+
+#include "Components/ChaosVDSolverCharacterGroundConstraintDataComponent.h"
+
+#include "Actors/ChaosVDSolverInfoActor.h"
+#include "Components/ChaosVDConstraintDataHelpers.h"
+#include "DataWrappers/ChaosVDCharacterGroundConstraintDataWrappers.h"
+
+void FChaosVDCharacterGroundConstraintSelectionHandle::SetIsSelected(bool bNewSelected)
+{
+	if (const TSharedPtr<FChaosVDCharacterGroundConstraint> ConstraintDataPtr = ConstraintData.Pin())
+	{
+		ConstraintDataPtr->bIsSelectedInEditor = bNewSelected;
+	}
+}
+
+bool FChaosVDCharacterGroundConstraintSelectionHandle::IsSelected() const
+{
+	if (const TSharedPtr<FChaosVDCharacterGroundConstraint> ConstraintDataPtr = ConstraintData.Pin())
+	{
+		return ConstraintDataPtr->bIsSelectedInEditor;
+	}
+
+	return false;
+}
+
+UChaosVDSolverCharacterGroundConstraintDataComponent::UChaosVDSolverCharacterGroundConstraintDataComponent()
+{
+	SetCanEverAffectNavigation(false);
+	bNavigationRelevant = false;
+	PrimaryComponentTick.bCanEverTick = false;
+}
+
+void UChaosVDSolverCharacterGroundConstraintDataComponent::UpdateConstraintData(const TArray<TSharedPtr<FChaosVDCharacterGroundConstraint>>& InConstraintData)
+{
+	ClearData();
+
+	AllConstraints = InConstraintData;
+
+	ConstraintByCharacterParticle.Reserve(InConstraintData.Num());
+	ConstraintByConstraintIndex.Reserve(InConstraintData.Num());
+
+	for (const TSharedPtr<FChaosVDCharacterGroundConstraint>& Constraint : InConstraintData)
+	{
+		if (!Constraint)
+		{
+			continue;
+		}
+
+		ConstraintByConstraintIndex.Add(Constraint->ConstraintIndex, Constraint);
+		Chaos::VisualDebugger::Utils::AddDataDataToParticleIDMap(ConstraintByCharacterParticle, Constraint, Constraint->CharacterParticleIndex);
+	}
+}
+
+const FChaosVDCharacterGroundDataArray* UChaosVDSolverCharacterGroundConstraintDataComponent::GetConstraintsForParticle(int32 ParticleID) const
+{
+	return ConstraintByCharacterParticle.Find(ParticleID);
+}
+
+void UChaosVDSolverCharacterGroundConstraintDataComponent::SelectConstraint(int32 ConstraintIndex)
+{
+	CurrentConstraintSelectionHandle.SetIsSelected(false);
+	CurrentConstraintSelectionHandle = FChaosVDCharacterGroundConstraintSelectionHandle(GetConstraintByIndex(ConstraintIndex));
+	CurrentConstraintSelectionHandle.SetIsSelected(true);
+
+	SelectionChangeDelegate.Broadcast(CurrentConstraintSelectionHandle);
+}
+
+void UChaosVDSolverCharacterGroundConstraintDataComponent::SelectConstraint(const FChaosVDCharacterGroundConstraintSelectionHandle& SelectionHandle)
+{
+	CurrentConstraintSelectionHandle.SetIsSelected(false);
+	CurrentConstraintSelectionHandle = SelectionHandle;
+	CurrentConstraintSelectionHandle.SetIsSelected(true);
+	
+	SelectionChangeDelegate.Broadcast(CurrentConstraintSelectionHandle);
+}
+
+bool UChaosVDSolverCharacterGroundConstraintDataComponent::IsConstraintSelected(int32 ConstraintIndex) const
+{
+	if (const TSharedPtr<FChaosVDCharacterGroundConstraint> ConstraintDataPtr = CurrentConstraintSelectionHandle.GetData().Pin())
+	{
+		return ConstraintDataPtr->bIsSelectedInEditor && ConstraintDataPtr->ConstraintIndex == ConstraintIndex;
+	}
+
+	return false;
+}
+
+TSharedPtr<FChaosVDCharacterGroundConstraint> UChaosVDSolverCharacterGroundConstraintDataComponent::GetConstraintByIndex(int32 ConstraintIndex)
+{
+	if (TSharedPtr<FChaosVDCharacterGroundConstraint>* ConstraintPtrPtr = ConstraintByConstraintIndex.Find(ConstraintIndex))
+	{
+		return *ConstraintPtrPtr;
+	}
+
+	return nullptr;
+}
+
+void UChaosVDSolverCharacterGroundConstraintDataComponent::ClearData()
+{
+	AllConstraints.Empty();
+	ConstraintByCharacterParticle.Empty();
+	ConstraintByConstraintIndex.Empty();
+}
diff --git a/Engine/Plugins/ChaosVD/Source/ChaosVD/Private/Components/ChaosVDSolverCharacterGroundConstraintDataComponent.h b/Engine/Plugins/ChaosVD/Source/ChaosVD/Private/Components/ChaosVDSolverCharacterGroundConstraintDataComponent.h
new file mode 100644
--- /dev/null
+++ b/Engine/Plugins/ChaosVD/Source/ChaosVD/Private/Components/ChaosVDSolverCharacterGroundConstraintDataComponent.h
@@ -0,0 +1,79 @@
+// Copyright Epic Games, Inc. All Rights Reserved.
+
+#pragma once
+
+#include "Components/ActorComponent.h"
+#include "ChaosVDSolverCharacterGroundConstraintDataComponent.generated.h"
+
+struct FChaosVDCharacterGroundConstraint;
+
+typedef TMap<int32, TArray<TSharedPtr<FChaosVDCharacterGroundConstraint>>> FChaosVDCharacterGroundDataByParticleMap;
+typedef TMap<int32, TSharedPtr<FChaosVDCharacterGroundConstraint>> FChaosVDCharacterGroundDataByConstraintIndexMap;
+typedef TArray<TSharedPtr<FChaosVDCharacterGroundConstraint>> FChaosVDCharacterGroundDataArray;
+
+enum class EChaosVDParticlePairSlot : uint8;
+
+/** Struct used to pass data about a specific character ground constraint to other objects */
+struct FChaosVDCharacterGroundConstraintSelectionHandle
+{
+	FChaosVDCharacterGroundConstraintSelectionHandle()
+	{		
+	}
+
+	FChaosVDCharacterGroundConstraintSelectionHandle(const TWeakPtr<FChaosVDCharacterGroundConstraint>& InConstraintData)
+		: ConstraintData(InConstraintData)
+	{
+	}
+
+	void SetIsSelected(bool bNewSelected);
+	bool IsSelected() const;
+
+	bool IsValid();
+
+	TWeakPtr<FChaosVDCharacterGroundConstraint> GetData() const { return ConstraintData; }
+
+private:
+	TWeakPtr<FChaosVDCharacterGroundConstraint> ConstraintData = nullptr;
+};
+
+DECLARE_MULTICAST_DELEGATE_OneParam(FChaosVDCharacterGroundSelectionChangedDelegate, const FChaosVDCharacterGroundConstraintSelectionHandle& SelectionHandle)
+
+UCLASS()
+class CHAOSVD_API UChaosVDSolverCharacterGroundConstraintDataComponent : public UActorComponent
+{
+	GENERATED_BODY()
+
+public:
+
+	UChaosVDSolverCharacterGroundConstraintDataComponent();
+
+	void UpdateConstraintData(const TArray<TSharedPtr<FChaosVDCharacterGroundConstraint>>& InConstraintData);
+	
+	const FChaosVDCharacterGroundDataArray& GetAllConstraints() const { return AllConstraints; }
+	const FChaosVDCharacterGroundDataArray* GetConstraintsForParticle(int32 ParticleID) const;
+
+	void SelectConstraint(int32 ConstraintIndex);
+	void SelectConstraint(const FChaosVDCharacterGroundConstraintSelectionHandle& SelectionHandle);
+
+	bool IsConstraintSelected(int32 ConstraintIndex) const;
+
+	FChaosVDCharacterGroundSelectionChangedDelegate& OnSelectionChanged() { return SelectionChangeDelegate; }
+
+	const FChaosVDCharacterGroundConstraintSelectionHandle& GetCurrentSelectionHandle() const { return CurrentConstraintSelectionHandle; }
+
+protected:
+
+	TSharedPtr<FChaosVDCharacterGroundConstraint> GetConstraintByIndex(int32 ConstraintIndex);
+	
+	void ClearData();
+
+	FChaosVDCharacterGroundDataArray AllConstraints;
+	
+	FChaosVDCharacterGroundDataByParticleMap ConstraintByCharacterParticle;
+
+	FChaosVDCharacterGroundDataByConstraintIndexMap ConstraintByConstraintIndex;
+
+	FChaosVDCharacterGroundSelectionChangedDelegate SelectionChangeDelegate;
+
+	FChaosVDCharacterGroundConstraintSelectionHandle CurrentConstraintSelectionHandle;
+};
diff --git a/Engine/Plugins/ChaosVD/Source/ChaosVD/Private/Trace/ChaosVDTraceProvider.cpp b/Engine/Plugins/ChaosVD/Source/ChaosVD/Private/Trace/ChaosVDTraceProvider.cpp
--- a/Engine/Plugins/ChaosVD/Source/ChaosVD/Private/Trace/ChaosVDTraceProvider.cpp
+++ b/Engine/Plugins/ChaosVD/Source/ChaosVD/Private/Trace/ChaosVDTraceProvider.cpp
@@ -9,6 +9,7 @@
 #include "ChaosVisualDebugger/ChaosVisualDebuggerTrace.h"
 
 #include "Compression/OodleDataCompressionUtil.h"
+#include "Trace/DataProcessors/ChaosVDCharacterGroundConstraintDataProcessor.h"
 #include "Trace/DataProcessors/ChaosVDConstraintDataProcessor.h"
 #include "Trace/DataProcessors/ChaosVDJointConstraintDataProcessor.h"
 #include "Trace/DataProcessors/ChaosVDMidPhaseDataProcessor.h"
@@ -229,6 +230,10 @@ void FChaosVDTraceProvider::RegisterDefaultDataProcessorsIfNeeded()
 	JointConstraintDataProcessor->SetTraceProvider(AsShared());
 	RegisterDataProcessor(JointConstraintDataProcessor);
 
+	TSharedPtr<FChaosVDCharacterGroundConstraintDataProcessor> CharacterGroundConstraintDataProcessor = MakeShared<FChaosVDCharacterGroundConstraintDataProcessor>();
+	CharacterGroundConstraintDataProcessor->SetTraceProvider(AsShared());
+	RegisterDataProcessor(CharacterGroundConstraintDataProcessor);
+
 	TSharedPtr<FChaosVDArchiveHeaderProcessor> ArchiveHeaderDataProcessor = MakeShared<FChaosVDArchiveHeaderProcessor>();
 	ArchiveHeaderDataProcessor->SetTraceProvider(AsShared());
 	RegisterDataProcessor(ArchiveHeaderDataProcessor);
diff --git a/Engine/Plugins/ChaosVD/Source/ChaosVD/Private/Trace/DataProcessors/ChaosVDCharacterGroundConstraintDataProcessor.cpp b/Engine/Plugins/ChaosVD/Source/ChaosVD/Private/Trace/DataProcessors/ChaosVDCharacterGroundConstraintDataProcessor.cpp
new file mode 100644
--- /dev/null
+++ b/Engine/Plugins/ChaosVD/Source/ChaosVD/Private/Trace/DataProcessors/ChaosVDCharacterGroundConstraintDataProcessor.cpp
@@ -0,0 +1,34 @@
+// Copyright Epic Games, Inc. All Rights Reserved.
+
+#include "Trace/DataProcessors/ChaosVDCharacterGroundConstraintDataProcessor.h"
+
+#include "DataWrappers/ChaosVDCharacterGroundConstraintDataWrappers.h"
+#include "Trace/ChaosVDTraceProvider.h"
+
+
+FChaosVDCharacterGroundConstraintDataProcessor::FChaosVDCharacterGroundConstraintDataProcessor()
+	: IChaosVDDataProcessor(FChaosVDCharacterGroundConstraint::WrapperTypeName)
+{
+}
+
+bool FChaosVDCharacterGroundConstraintDataProcessor::ProcessRawData(const TArray<uint8>& InData)
+{
+	TSharedPtr<FChaosVDTraceProvider> ProviderSharedPtr = TraceProvider.Pin();
+	if (!ensure(ProviderSharedPtr.IsValid()))
+	{
+		return false;
+	}
+
+	TSharedPtr<FChaosVDCharacterGroundConstraint> Constraint = MakeShared<FChaosVDCharacterGroundConstraint>();
+	const bool bSuccess = Chaos::VisualDebugger::ReadDataFromBuffer(InData, *Constraint, ProviderSharedPtr.ToSharedRef());
+
+	if (bSuccess)
+	{
+		if (const TSharedPtr<FChaosVDGameFrameData> CurrentFrameData = ProviderSharedPtr->GetCurrentGameFrame().Pin())
+		{
+			CurrentFrameData->RecordedCharacterGroundConstraints.Add(Constraint);
+		}
+	}
+
+	return bSuccess;
+}
diff --git a/Engine/Plugins/ChaosVD/Source/ChaosVD/Private/Trace/DataProcessors/ChaosVDCharacterGroundConstraintDataProcessor.h b/Engine/Plugins/ChaosVD/Source/ChaosVD/Private/Trace/DataProcessors/ChaosVDCharacterGroundConstraintDataProcessor.h
new file mode 100644
--- /dev/null
+++ b/Engine/Plugins/ChaosVD/Source/ChaosVD/Private/Trace/DataProcessors/ChaosVDCharacterGroundConstraintDataProcessor.h
@@ -0,0 +1,22 @@
+// Copyright Epic Games, Inc. All Rights Reserved.
+
+#pragma once
+
+#pragma once
+#include "IChaosVDDataProcessor.h"
+
+struct FChaosVDSolverFrameData;
+struct FChaosVDConstraint;
+
+/**
+ * Data processor implementation that is able to deserialize traced character ground constraints
+ */
+class FChaosVDCharacterGroundConstraintDataProcessor final : public IChaosVDDataProcessor
+{
+public:
+	explicit FChaosVDCharacterGroundConstraintDataProcessor();
+	
+	virtual bool ProcessRawData(const TArray<uint8>& InData) override;
+};
+
+
diff --git a/Engine/Plugins/ChaosVD/Source/ChaosVD/Private/Visualizers/ChaosVDCharacterGroundConstraintsDataComponentVisualizer.cpp b/Engine/Plugins/ChaosVD/Source/ChaosVD/Private/Visualizers/ChaosVDCharacterGroundConstraintsDataComponentVisualizer.cpp
new file mode 100644
--- /dev/null
+++ b/Engine/Plugins/ChaosVD/Source/ChaosVD/Private/Visualizers/ChaosVDCharacterGroundConstraintsDataComponentVisualizer.cpp
@@ -0,0 +1,246 @@
+// Copyright Epic Games, Inc. All Rights Reserved.
+
+#include "Visualizers/ChaosVDCharacterGroundConstraintsDataComponentVisualizer.h"
+
+#include "ChaosVDScene.h"
+#include "ChaosVDTabsIDs.h"
+#include "EditorModeManager.h"
+#include "EditorViewportClient.h"
+#include "SceneManagement.h"
+#include "SceneView.h"
+#include "Actors/ChaosVDSolverInfoActor.h"
+#include "Components/ChaosVDSolverCharacterGroundConstraintDataComponent.h"
+#include "Visualizers/ChaosVDDebugDrawUtils.h"
+#include "Widgets/SChaosVDMainTab.h"
+
+IMPLEMENT_HIT_PROXY(HChaosVDCharacterGroundConstraintProxy, HComponentVisProxy)
+
+
+void FChaosVDCharacterGroundConstraintDataComponentVisualizer::DrawVisualization(const UActorComponent* Component, const FSceneView* View, FPrimitiveDrawInterface* PDI)
+{
+	const UChaosVDSolverCharacterGroundConstraintDataComponent* ConstraintDataComponent = Cast<UChaosVDSolverCharacterGroundConstraintDataComponent>(Component);
+	if (!ConstraintDataComponent)
+	{
+		return;
+	}
+	
+	AChaosVDSolverInfoActor* SolverInfoActor = Cast<AChaosVDSolverInfoActor>(Component->GetOwner());
+	if (!SolverInfoActor)
+	{
+		return;
+	}
+
+	if (!SolverInfoActor->IsVisible())
+	{
+		return;
+	}
+
+	const TSharedPtr<FChaosVDScene> CVDScene = SolverInfoActor->GetScene().Pin();
+	if (!CVDScene)
+	{
+		return;
+	}
+
+	const TSharedPtr<FChaosVDRecording> CVDRecording = CVDScene->LoadedRecording;
+	if (!CVDRecording)
+	{
+		return;
+	}
+
+	FChaosVDCharacterGroundConstraintVisualizationDataContext VisualizationContext;
+	VisualizationContext.CVDScene = CVDScene;
+	VisualizationContext.SpaceTransform = SolverInfoActor->GetSimulationTransform();
+	VisualizationContext.SolverInfoActor = SolverInfoActor;
+
+	if (const UChaosVDEditorSettings* EditorSettings = GetDefault<UChaosVDEditorSettings>())
+	{
+		VisualizationContext.VisualizationFlags = EditorSettings->GlobalCharacterGroundConstraintDataVisualizationFlags;
+		VisualizationContext.bShowDebugText = EditorSettings->bShowDebugText;
+		VisualizationContext.DebugDrawSettings = &EditorSettings->CharacterGroundConstraintDataDebugDrawSettings;
+	}
+
+	if (!VisualizationContext.IsVisualizationFlagEnabled(EChaosVDCharacterGroundConstraintDataVisualizationFlags::EnableDraw))
+	{
+		return;
+	}
+
+	// If nothing is selected, fallback to draw all character ground constraints
+	const bool bDrawOnlySelected = VisualizationContext.IsVisualizationFlagEnabled(EChaosVDCharacterGroundConstraintDataVisualizationFlags::OnlyDrawSelected) && ConstraintDataComponent->GetCurrentSelectionHandle().IsSelected();
+	if (bDrawOnlySelected)
+	{
+		if (const TSharedPtr<FChaosVDCharacterGroundConstraint> Constraint = ConstraintDataComponent->GetCurrentSelectionHandle().GetData().Pin())
+		{
+			VisualizationContext.DataSelectionHandle = FChaosVDCharacterGroundConstraintSelectionHandle(Constraint);
+			DrawConstraint(Component, *Constraint, VisualizationContext, View, PDI);
+		}
+	}
+	else
+	{
+		for (const TSharedPtr<FChaosVDCharacterGroundConstraint>& Constraint : ConstraintDataComponent->GetAllConstraints())
+		{
+			if (Constraint)
+			{
+				VisualizationContext.DataSelectionHandle = FChaosVDCharacterGroundConstraintSelectionHandle(Constraint);
+				DrawConstraint(Component, *Constraint, VisualizationContext, View, PDI);
+			}
+		}
+	}
+}
+
+bool FChaosVDCharacterGroundConstraintDataComponentVisualizer::VisProxyHandleClick(FEditorViewportClient* InViewportClient, HComponentVisProxy* VisProxy, const FViewportClick& Click)
+{
+	const HChaosVDCharacterGroundConstraintProxy* ConstraintDataProxy = HitProxyCast<HChaosVDCharacterGroundConstraintProxy>(VisProxy);
+	if (ConstraintDataProxy == nullptr)
+	{
+		return false;
+	}
+	
+	if (const UChaosVDSolverCharacterGroundConstraintDataComponent* ConstraintDataComponent = Cast<UChaosVDSolverCharacterGroundConstraintDataComponent>(ConstraintDataProxy->Component.Get()))
+	{
+		// Bring the constraint details tab into focus if available
+		const TSharedPtr<SChaosVDMainTab> MainTabToolkitHost = InViewportClient->GetModeTools() ? StaticCastSharedPtr<SChaosVDMainTab>(InViewportClient->GetModeTools()->GetToolkitHost()) : nullptr;
+		if (const TSharedPtr<FTabManager> TabManager = MainTabToolkitHost ? MainTabToolkitHost->GetTabManager() : nullptr)
+		{
+			TabManager->TryInvokeTab(FChaosVDTabID::CharacterGroundConstraintDataDetails);
+		}
+
+		const_cast<UChaosVDSolverCharacterGroundConstraintDataComponent*>(ConstraintDataComponent)->SelectConstraint(ConstraintDataProxy->DataSelectionHandle);
+
+		return true;
+	}
+
+	return false;
+}
+
+void FChaosVDCharacterGroundConstraintDataComponentVisualizer::DrawConstraint(const UActorComponent* Component, const FChaosVDCharacterGroundConstraint& InConstraintData, FChaosVDCharacterGroundConstraintVisualizationDataContext& VisualizationContext, const FSceneView* View, FPrimitiveDrawInterface* PDI)
+{
+	using namespace Chaos::VisualDebugger::Utils;
+
+	if (!VisualizationContext.IsVisualizationFlagEnabled(EChaosVDCharacterGroundConstraintDataVisualizationFlags::DrawDisabled))
+	{
+		if (InConstraintData.State.bDisabled)
+		{
+			return;
+		}
+	}
+
+	if (!Component)
+	{
+		return;
+	}
+	
+	if (!VisualizationContext.DebugDrawSettings)
+	{
+		return;
+	}
+
+	if (!PDI)
+	{
+		return;
+	}
+
+	if (!View)
+	{
+		return;
+	}
+
+	const FChaosVDParticleDataWrapper* CharacterParticleData = nullptr;
+	const FChaosVDParticleDataWrapper* GroundParticleData = nullptr;
+
+	if (AChaosVDParticleActor* CharacterParticle = VisualizationContext.SolverInfoActor->GetParticleActor(InConstraintData.CharacterParticleIndex))
+	{
+		CharacterParticleData = CharacterParticle->GetParticleData();
+	}
+
+	if (AChaosVDParticleActor* GroundParticle = VisualizationContext.SolverInfoActor->GetParticleActor(InConstraintData.GroundParticleIndex))
+	{
+		GroundParticleData = GroundParticle->GetParticleData();
+	}
+
+	if (!CharacterParticleData)
+	{
+		return;
+	}
+
+
+	if (!CharacterParticleData->ParticleMassProps.HasValidData())
+	{
+		// If we don't have mass data, all the following calculations will be off
+		// TODO: Should we draw just a line between the two particles as fallback?
+		return;
+	}
+
+	PDI->SetHitProxy(new HChaosVDCharacterGroundConstraintProxy(Component, VisualizationContext.DataSelectionHandle));
+
+	const float LineThickness = InConstraintData.bIsSelectedInEditor ?  VisualizationContext.DebugDrawSettings->BaseLineThickness * 1.5f :  VisualizationContext.DebugDrawSettings->BaseLineThickness;
+
+	const FVector CharacterPos = CharacterParticleData->ParticlePositionRotation.MX;
+	const FVector UpDir = InConstraintData.Settings.VerticalAxis;
+
+	const float GroundDistance = InConstraintData.Data.GroundDistance;
+	const float TargetHeight = InConstraintData.Settings.TargetHeight;
+
+	// TODO: The target data is valid for pre-sim positions and the force/torque for post sim
+	// but we don't have a way of differentiating here which state the particle is in so just draw
+	// everything for now and leave it up to the user to interpret what they're seeing
+
+	if (VisualizationContext.IsVisualizationFlagEnabled(EChaosVDCharacterGroundConstraintDataVisualizationFlags::TargetDeltaPosition))
+	{
+		FChaosVDDebugDrawUtils::DrawArrowVector(PDI, CharacterPos, CharacterPos + InConstraintData.Data.TargetDeltaPosition, FText::GetEmpty(), FColor::Blue, VisualizationContext.DebugDrawSettings->DepthPriority, 0.5f * LineThickness);
+	}
+
+	if (VisualizationContext.IsVisualizationFlagEnabled(EChaosVDCharacterGroundConstraintDataVisualizationFlags::TargetDeltaFacing))
+	{
+		const FVector Forward = CharacterParticleData->ParticlePositionRotation.MR * FVector::XAxisVector * VisualizationContext.DebugDrawSettings->GeneralScale * 10.0f;
+		const FVector TargetForward = FQuat(UpDir, InConstraintData.Data.TargetDeltaFacing) * Forward;
+		FChaosVDDebugDrawUtils::DrawArrowVector(PDI, CharacterPos, CharacterPos + Forward, FText::GetEmpty(), FColor::Silver, VisualizationContext.DebugDrawSettings->DepthPriority, 0.25f * LineThickness);
+		FChaosVDDebugDrawUtils::DrawArrowVector(PDI, CharacterPos, CharacterPos + TargetForward, FText::GetEmpty(), FColor::White, VisualizationContext.DebugDrawSettings->DepthPriority, 0.25f * LineThickness);
+	}
+
+	if (VisualizationContext.IsVisualizationFlagEnabled(EChaosVDCharacterGroundConstraintDataVisualizationFlags::GroundQueryDistance))
+	{
+		if (GroundDistance >= TargetHeight)
+		{
+			if (GroundDistance <= 4.0f * TargetHeight)
+			{
+				FChaosVDDebugDrawUtils::DrawLine(PDI, CharacterPos, CharacterPos - UpDir * TargetHeight, FColor::Green, FText::GetEmpty(), VisualizationContext.DebugDrawSettings->DepthPriority, LineThickness);
+				FChaosVDDebugDrawUtils::DrawLine(PDI, CharacterPos - UpDir * TargetHeight, CharacterPos - UpDir * GroundDistance, FColor::Silver, FText::GetEmpty(), VisualizationContext.DebugDrawSettings->DepthPriority, LineThickness);
+			}
+		}
+		else
+		{
+			FChaosVDDebugDrawUtils::DrawLine(PDI, CharacterPos, CharacterPos - UpDir * GroundDistance, FColor::Green, FText::GetEmpty(), VisualizationContext.DebugDrawSettings->DepthPriority, LineThickness);
+			FChaosVDDebugDrawUtils::DrawLine(PDI, CharacterPos - UpDir * GroundDistance, CharacterPos - UpDir * TargetHeight, FColor::Red, FText::GetEmpty(), VisualizationContext.DebugDrawSettings->DepthPriority, LineThickness);
+		}
+	}
+
+	if (VisualizationContext.IsVisualizationFlagEnabled(EChaosVDCharacterGroundConstraintDataVisualizationFlags::GroundQueryNormal))
+	{
+		if (GroundDistance < 4.0f * TargetHeight)
+		{
+			const FVector ScaledGroundNormal = 10.0f * InConstraintData.Data.GroundNormal * VisualizationContext.DebugDrawSettings->GeneralScale;
+			const FVector GroundPos = CharacterPos - UpDir * GroundDistance;
+			FChaosVDDebugDrawUtils::DrawArrowVector(PDI, GroundPos, GroundPos + ScaledGroundNormal, FText::GetEmpty(), FColor::Cyan, VisualizationContext.DebugDrawSettings->DepthPriority, 0.25f * LineThickness);
+		}
+	}
+	
+	if (VisualizationContext.IsVisualizationFlagEnabled(EChaosVDCharacterGroundConstraintDataVisualizationFlags::AppliedNormalForce))
+	{
+		const FVector NormalForce = VisualizationContext.DebugDrawSettings->ForceScale * InConstraintData.Data.GroundNormal.Dot(InConstraintData.State.SolverAppliedForce) * InConstraintData.Data.GroundNormal;
+		FChaosVDDebugDrawUtils::DrawArrowVector(PDI, CharacterPos, CharacterPos + NormalForce, FText::GetEmpty(), VisualizationContext.DebugDrawSettings->NormalForceColor, VisualizationContext.DebugDrawSettings->DepthPriority, LineThickness);
+	}
+
+	if (VisualizationContext.IsVisualizationFlagEnabled(EChaosVDCharacterGroundConstraintDataVisualizationFlags::AppliedRadialForce))
+	{
+		const FVector RadialForce = VisualizationContext.DebugDrawSettings->ForceScale * (InConstraintData.State.SolverAppliedForce - InConstraintData.Data.GroundNormal.Dot(InConstraintData.State.SolverAppliedForce) * InConstraintData.Data.GroundNormal);
+		FChaosVDDebugDrawUtils::DrawArrowVector(PDI, CharacterPos, CharacterPos + RadialForce, FText::GetEmpty(), VisualizationContext.DebugDrawSettings->NormalForceColor, VisualizationContext.DebugDrawSettings->DepthPriority, LineThickness);
+	}
+
+	if (VisualizationContext.IsVisualizationFlagEnabled(EChaosVDCharacterGroundConstraintDataVisualizationFlags::AppliedTorque))
+	{
+		const FVector Torque = VisualizationContext.DebugDrawSettings->TorqueScale * InConstraintData.State.SolverAppliedTorque;
+		FChaosVDDebugDrawUtils::DrawArrowVector(PDI, CharacterPos, CharacterPos + Torque, FText::GetEmpty(), VisualizationContext.DebugDrawSettings->TorqueColor, VisualizationContext.DebugDrawSettings->DepthPriority, LineThickness);
+	}
+
+	PDI->SetHitProxy(nullptr);
+}
diff --git a/Engine/Plugins/ChaosVD/Source/ChaosVD/Private/Visualizers/ChaosVDCharacterGroundConstraintsDataComponentVisualizer.h b/Engine/Plugins/ChaosVD/Source/ChaosVD/Private/Visualizers/ChaosVDCharacterGroundConstraintsDataComponentVisualizer.h
new file mode 100644
--- /dev/null
+++ b/Engine/Plugins/ChaosVD/Source/ChaosVD/Private/Visualizers/ChaosVDCharacterGroundConstraintsDataComponentVisualizer.h
@@ -0,0 +1,92 @@
+// Copyright Epic Games, Inc. All Rights Reserved.
+
+#pragma once
+#include "ComponentVisualizer.h"
+#include "IChaosVDParticleVisualizationDataProvider.h"
+#include "Chaos/Core.h"
+#include "Components/ChaosVDSolverCharacterGroundConstraintDataComponent.h"
+
+struct FChaosVDCharacterGroundConstraintDebugDrawSettings;
+class AChaosVDSolverInfoActor;
+struct FChaosVDCharacterGroundConstraint;
+
+UENUM(meta = (Bitflags, UseEnumValuesAsMaskValuesInEditor = "true"))
+enum class EChaosVDCharacterGroundConstraintDataVisualizationFlags : uint32
+{
+	None					= 0 UMETA(Hidden),
+	/** Draw the target movement vector */
+	TargetDeltaPosition		= 1 << 0,
+	/** Draw the target orientation facing vector */
+	TargetDeltaFacing		= 1 << 1,
+	/** Draw the ground query distance based on the constraint's data */
+	GroundQueryDistance		= 1 << 2,
+	/** Draw the ground query normal based on the constraint's data */
+	GroundQueryNormal		= 1 << 3,
+	/** Draw the applied force vector */
+	AppliedRadialForce		= 1 << 4,
+	/** Draw the applied force vector */
+	AppliedNormalForce		= 1 << 5,
+	/** Draw the applied force vector */
+	AppliedTorque			= 1 << 6,
+	/** Draw the constraint even if it is disabled */
+	DrawDisabled			= 1 << 7,
+	/** Only debugs draw data for a selected constraint */
+	OnlyDrawSelected		= 1 << 8,
+	/** Enables debug draw for constraint data from any solver */
+	EnableDraw				= 1 << 9,
+};
+ENUM_CLASS_FLAGS(EChaosVDCharacterGroundConstraintDataVisualizationFlags);
+
+/** Visualization context structure specific for character ground constraint visualizations */
+struct FChaosVDCharacterGroundConstraintVisualizationDataContext : public FChaosVDVisualizationContext
+{
+	FChaosVDCharacterGroundConstraintSelectionHandle DataSelectionHandle = FChaosVDCharacterGroundConstraintSelectionHandle(nullptr);
+	bool bIsServerVisualizationEnabled = false;
+	
+	AChaosVDSolverInfoActor* SolverInfoActor = nullptr;
+
+	const FChaosVDCharacterGroundConstraintDebugDrawSettings* DebugDrawSettings = nullptr;
+
+	bool bShowDebugText = false;
+
+	bool IsVisualizationFlagEnabled(EChaosVDCharacterGroundConstraintDataVisualizationFlags Flag) const
+	{
+		const EChaosVDCharacterGroundConstraintDataVisualizationFlags FlagsAsParticleFlags = static_cast<EChaosVDCharacterGroundConstraintDataVisualizationFlags>(VisualizationFlags);
+		return EnumHasAnyFlags(FlagsAsParticleFlags, Flag);
+	}
+};
+
+/** Custom Hit Proxy for debug drawn scene queries */
+struct HChaosVDCharacterGroundConstraintProxy : public HComponentVisProxy
+{
+	DECLARE_HIT_PROXY()
+	
+	HChaosVDCharacterGroundConstraintProxy(const UActorComponent* Component, const FChaosVDCharacterGroundConstraintSelectionHandle& InConstraintFinderData) : HComponentVisProxy(Component, HPP_UI), DataSelectionHandle(InConstraintFinderData)
+	{	
+	}
+
+	virtual EMouseCursor::Type GetMouseCursor() override
+	{
+		return EMouseCursor::Crosshairs;
+	}
+
+	FChaosVDCharacterGroundConstraintSelectionHandle DataSelectionHandle;
+};
+
+/**
+ * Component visualizer in charge of generating debug draw visualizations for character ground constraints in a UChaosVDSolverCharacterGroundConstraintDataComponent
+ */
+class FChaosVDCharacterGroundConstraintDataComponentVisualizer final : public FComponentVisualizer
+{
+public:
+	FChaosVDCharacterGroundConstraintDataComponentVisualizer()
+	{
+	}
+
+	virtual void DrawVisualization(const UActorComponent* Component, const FSceneView* View, FPrimitiveDrawInterface* PDI) override;
+	virtual bool VisProxyHandleClick(FEditorViewportClient* InViewportClient, HComponentVisProxy* VisProxy, const FViewportClick& Click) override;
+
+protected:
+
+	void DrawConstraint(const UActorComponent* Component, const FChaosVDCharacterGroundConstraint& InConstraintData, FChaosVDCharacterGroundConstraintVisualizationDataContext& VisualizationContext, const FSceneView* View, FPrimitiveDrawInterface* PDI);
+};
diff --git a/Engine/Plugins/ChaosVD/Source/ChaosVD/Private/Widgets/SChaosVDCharacterGroundConstraintDataInspector.cpp b/Engine/Plugins/ChaosVD/Source/ChaosVD/Private/Widgets/SChaosVDCharacterGroundConstraintDataInspector.cpp
new file mode 100644
--- /dev/null
+++ b/Engine/Plugins/ChaosVD/Source/ChaosVD/Private/Widgets/SChaosVDCharacterGroundConstraintDataInspector.cpp
@@ -0,0 +1,229 @@
+// Copyright Epic Games, Inc. All Rights Reserved.
+
+#include "Widgets/SChaosVDCharacterGroundConstraintDataInspector.h"
+
+#include "ChaosVDScene.h"
+#include "IStructureDetailsView.h"
+#include "PropertyEditorModule.h"
+#include "Actors/ChaosVDSolverInfoActor.h"
+#include "Widgets/SChaosVDWarningMessageBox.h"
+#include "Components/ChaosVDSolverCharacterGroundConstraintDataComponent.h"
+#include "DataWrappers/ChaosVDCharacterGroundConstraintDataWrappers.h"
+#include "Modules/ModuleManager.h"
+#include "Selection.h"
+#include "Widgets/Input/SButton.h"
+#include "Widgets/Layout/SBox.h"
+#include "Widgets/Layout/SScrollBox.h"
+#include "Widgets/Layout/SUniformGridPanel.h"
+#include "Widgets/Text/STextBlock.h"
+
+#define LOCTEXT_NAMESPACE "ChaosVisualDebugger"
+
+SChaosVDCharacterGroundConstraintDataInspector::~SChaosVDCharacterGroundConstraintDataInspector()
+{
+	UnregisterSceneEvents();
+}
+
+void SChaosVDCharacterGroundConstraintDataInspector::Construct(const FArguments& InArgs, const TWeakPtr<FChaosVDScene>& InScenePtr)
+{
+	SceneWeakPtr = InScenePtr;
+
+	RegisterSceneEvents();
+
+	ConstraintDataDetailsView = CreateDataDetailsView();
+
+	constexpr float NoPadding = 0.0f;
+	constexpr float OuterBoxPadding = 2.0f;
+	constexpr float OuterInnerPadding = 5.0f;
+	constexpr float TagTitleBoxHorizontalPadding = 10.0f;
+	constexpr float TagTitleBoxVerticalPadding = 5.0f;
+	constexpr float InnerDetailsPanelsHorizontalPadding = 15.0f;
+	constexpr float InnerDetailsPanelsVerticalPadding = 15.0f;
+
+	ChildSlot
+	[
+		SNew(SVerticalBox)
+		+SVerticalBox::Slot()
+		.AutoHeight()
+		.Padding(OuterInnerPadding)
+		[
+			SNew(SBox)
+			.Visibility_Raw(this, &SChaosVDCharacterGroundConstraintDataInspector::GetOutOfDateWarningVisibility)
+			.Padding(OuterBoxPadding, OuterBoxPadding,OuterBoxPadding,NoPadding)
+			[
+				SNew(SChaosVDWarningMessageBox)
+				.WarningText(LOCTEXT("CharacterGroundConstraintDataOutOfDate", "Scene change detected!. Selected constraint data is out of date..."))
+			]
+		]
+		+SVerticalBox::Slot()
+		.Padding(TagTitleBoxHorizontalPadding, TagTitleBoxVerticalPadding, TagTitleBoxHorizontalPadding, NoPadding)
+		.AutoHeight()
+		[
+			SNew(STextBlock)
+			.Visibility_Raw(this, &SChaosVDCharacterGroundConstraintDataInspector::GetNothingSelectedMessageVisibility)
+			.Justification(ETextJustify::Center)
+			.TextStyle(FAppStyle::Get(), "DetailsView.BPMessageTextStyle")
+			.Text(LOCTEXT("CharacterGroundConstraintDataNoSelectedMessage", "Select a Character Ground Constraint in the viewport to see its details..."))
+			.AutoWrapText(true)
+		]
+		+SVerticalBox::Slot()
+		.Padding(OuterInnerPadding)
+		[
+			SNew(SScrollBox)
+			.Visibility_Raw(this, &SChaosVDCharacterGroundConstraintDataInspector::GetDetailsSectionVisibility)
+			+SScrollBox::Slot()
+			.Padding(InnerDetailsPanelsHorizontalPadding,NoPadding,InnerDetailsPanelsHorizontalPadding,InnerDetailsPanelsVerticalPadding)
+			[
+				ConstraintDataDetailsView->GetWidget().ToSharedRef()
+			]
+		]
+	];
+}
+
+void SChaosVDCharacterGroundConstraintDataInspector::RegisterSelectionEventsForSolver(AChaosVDSolverInfoActor* SolverInfo)
+{
+	if (UChaosVDSolverCharacterGroundConstraintDataComponent* ConstraintDataComponent = SolverInfo ? SolverInfo->GetCharacterGroundConstraintDataComponent() : nullptr)
+	{
+		ConstraintDataComponent->OnSelectionChanged().AddRaw(this, &SChaosVDCharacterGroundConstraintDataInspector::SetConstraintDataToInspect);
+	}
+}
+
+void SChaosVDCharacterGroundConstraintDataInspector::UnregisterSelectionEventsForSolver(AChaosVDSolverInfoActor* SolverInfo)
+{
+	if (UChaosVDSolverCharacterGroundConstraintDataComponent* ConstraintDataComponent = SolverInfo ? SolverInfo->GetCharacterGroundConstraintDataComponent() : nullptr)
+	{
+		ConstraintDataComponent->OnSelectionChanged().RemoveAll(this);
+	}
+}
+
+void SChaosVDCharacterGroundConstraintDataInspector::RegisterSceneEvents()
+{
+	if (const TSharedPtr<FChaosVDScene> ScenePtr = SceneWeakPtr.Pin())
+	{
+		ScenePtr->OnSceneUpdated().AddRaw(this, &SChaosVDCharacterGroundConstraintDataInspector::HandleSceneUpdated);
+		ScenePtr->OnSolverInfoActorCreated().AddRaw(this, &SChaosVDCharacterGroundConstraintDataInspector::RegisterSelectionEventsForSolver);
+
+		for (const TPair<int32, AChaosVDSolverInfoActor*> SolverInfos : ScenePtr->GetSolverInfoActorsMap())
+		{
+			RegisterSelectionEventsForSolver(SolverInfos.Value);
+		}
+	}
+}
+
+void SChaosVDCharacterGroundConstraintDataInspector::UnregisterSceneEvents()
+{
+	if (const TSharedPtr<FChaosVDScene> ScenePtr = SceneWeakPtr.Pin())
+	{
+		ScenePtr->OnSceneUpdated().RemoveAll(this);
+		ScenePtr->OnSolverInfoActorCreated().RemoveAll(this);
+
+		for (const TPair<int32, AChaosVDSolverInfoActor*> SolverInfos : ScenePtr->GetSolverInfoActorsMap())
+		{
+			UnregisterSelectionEventsForSolver(SolverInfos.Value);
+		}
+	}
+}
+
+void SChaosVDCharacterGroundConstraintDataInspector::SetConstraintDataToInspect(const FChaosVDCharacterGroundConstraintSelectionHandle& InDataSelectionHandle)
+{
+	ClearInspector();
+
+	if (const TSharedPtr<FChaosVDCharacterGroundConstraint> QueryDataToInspect = InDataSelectionHandle.GetData().Pin())
+	{
+		CurrentDataSelectionHandle = InDataSelectionHandle;
+		const TSharedPtr<FStructOnScope> QueryDataView = MakeShared<FStructOnScope>(FChaosVDCharacterGroundConstraint::StaticStruct(), reinterpret_cast<uint8*>(QueryDataToInspect.Get()));
+		ConstraintDataDetailsView->SetStructureData(QueryDataView);
+	}
+
+	bIsUpToDate = true;
+}
+
+void SChaosVDCharacterGroundConstraintDataInspector::SetConstraintDataProviderObjectToInspect(IChaosVDCharacterGroundConstraintDataProviderInterface* ConstraintDataProvider)
+{
+	if (ConstraintDataProvider == nullptr)
+	{
+		return;
+	}
+
+	// Currently only support viewing first found data
+	if (ConstraintDataProvider->HasCharacterGroundConstraintData())
+	{
+		TArray<TSharedPtr<FChaosVDCharacterGroundConstraint>> FoundConstraintData;
+		ConstraintDataProvider->GetCharacterGroundConstraintData(FoundConstraintData);
+
+		FChaosVDCharacterGroundConstraintSelectionHandle NewDataSelectionHandle = FChaosVDCharacterGroundConstraintSelectionHandle(FoundConstraintData[0]);
+		if (NewDataSelectionHandle.GetData() != CurrentDataSelectionHandle.GetData())
+		{
+			ClearInspector();
+
+			const TSharedPtr<FStructOnScope> QueryDataView = MakeShared<FStructOnScope>(FChaosVDCharacterGroundConstraint::StaticStruct(), reinterpret_cast<uint8*>(FoundConstraintData[0].Get()));
+			ConstraintDataDetailsView->SetStructureData(QueryDataView);
+			CurrentDataSelectionHandle = NewDataSelectionHandle;
+		}
+	}
+	else
+	{
+		ClearInspector();
+	}
+
+	bIsUpToDate = true;
+}
+
+EVisibility SChaosVDCharacterGroundConstraintDataInspector::GetOutOfDateWarningVisibility() const
+{
+	return !bIsUpToDate && CurrentDataSelectionHandle.GetData().IsValid() ? EVisibility::Visible : EVisibility::Collapsed;
+}
+
+EVisibility SChaosVDCharacterGroundConstraintDataInspector::GetDetailsSectionVisibility() const
+{
+	return CurrentDataSelectionHandle.GetData().IsValid() ? EVisibility::Visible : EVisibility::Collapsed;
+}
+
+EVisibility SChaosVDCharacterGroundConstraintDataInspector::GetNothingSelectedMessageVisibility() const
+{
+	return !CurrentDataSelectionHandle.GetData().IsValid() ? EVisibility::Visible : EVisibility::Collapsed;
+}
+
+void SChaosVDCharacterGroundConstraintDataInspector::HandleSceneUpdated()
+{
+	// TODO: Optimize this
+
+	if (const TSharedPtr<FChaosVDScene> ScenePtr = SceneWeakPtr.Pin())
+	{
+		if (USelection* Selection = ScenePtr->GetActorSelectionObject())
+		{
+			if (Selection->Num() > 0)
+			{
+				if (IChaosVDCharacterGroundConstraintDataProviderInterface* DataProvider = Cast<IChaosVDCharacterGroundConstraintDataProviderInterface>(Selection->GetSelectedObject(0)))
+				{
+					SetConstraintDataProviderObjectToInspect(DataProvider);
+					return;
+				}
+			}
+		}
+	}
+
+	ClearInspector();
+}
+
+void SChaosVDCharacterGroundConstraintDataInspector::ClearInspector()
+{
+	ConstraintDataDetailsView->SetStructureData(nullptr);
+	CurrentDataSelectionHandle = FChaosVDCharacterGroundConstraintSelectionHandle();
+}
+
+TSharedPtr<IStructureDetailsView> SChaosVDCharacterGroundConstraintDataInspector::CreateDataDetailsView()
+{
+	FPropertyEditorModule& PropertyEditorModule = FModuleManager::GetModuleChecked<FPropertyEditorModule>("PropertyEditor");
+
+	const FStructureDetailsViewArgs StructDetailsViewArgs;
+	FDetailsViewArgs DetailsViewArgs;
+	DetailsViewArgs.bShowOptions = false;
+	DetailsViewArgs.bAllowFavoriteSystem = false;
+	DetailsViewArgs.bAllowSearch = true;
+	DetailsViewArgs.bShowScrollBar = false;
+
+	return PropertyEditorModule.CreateStructureDetailView(DetailsViewArgs, StructDetailsViewArgs, nullptr);
+}
+
+#undef LOCTEXT_NAMESPACE
diff --git a/Engine/Plugins/ChaosVD/Source/ChaosVD/Private/Widgets/SChaosVDCharacterGroundConstraintDataInspector.h b/Engine/Plugins/ChaosVD/Source/ChaosVD/Private/Widgets/SChaosVDCharacterGroundConstraintDataInspector.h
new file mode 100644
--- /dev/null
+++ b/Engine/Plugins/ChaosVD/Source/ChaosVD/Private/Widgets/SChaosVDCharacterGroundConstraintDataInspector.h
@@ -0,0 +1,58 @@
+// Copyright Epic Games, Inc. All Rights Reserved.
+
+#pragma once
+#include "Actors/ChaosVDSolverInfoActor.h"
+#include "Components/ChaosVDSolverCharacterGroundConstraintDataComponent.h"
+#include "Widgets/SCompoundWidget.h"
+
+struct FChaosVDCharacterGroundConstraintSelectionHandle;
+class FChaosVDScene;
+class IChaosVDCharacterGroundConstraintDataProviderInterface;
+class IStructureDetailsView;
+
+class SChaosVDCharacterGroundConstraintDataInspector : public SCompoundWidget
+{
+public:
+
+	SLATE_BEGIN_ARGS(SChaosVDCharacterGroundConstraintDataInspector)
+	{
+	}
+
+	SLATE_END_ARGS()
+
+	virtual ~SChaosVDCharacterGroundConstraintDataInspector() override;
+
+	/** Constructs this widget with InArgs */
+	void Construct(const FArguments& InArgs, const TWeakPtr<FChaosVDScene>& InScenePtr);
+
+	/** Sets a new constraint data to be inspected */
+	void SetConstraintDataToInspect(const FChaosVDCharacterGroundConstraintSelectionHandle& InDataSelectionHandle);
+
+	/** Set a constraint provider*/
+	void SetConstraintDataProviderObjectToInspect(IChaosVDCharacterGroundConstraintDataProviderInterface* ConstraintDataProvider);
+
+protected:
+	TSharedPtr<IStructureDetailsView> CreateDataDetailsView();
+
+	EVisibility GetOutOfDateWarningVisibility() const;
+	EVisibility GetDetailsSectionVisibility() const;
+	EVisibility GetNothingSelectedMessageVisibility() const;
+
+	void RegisterSelectionEventsForSolver(AChaosVDSolverInfoActor* SolverInfo);
+	void UnregisterSelectionEventsForSolver(AChaosVDSolverInfoActor* SolverInfo);
+
+	void RegisterSceneEvents();
+	void UnregisterSceneEvents();
+
+	void HandleSceneUpdated();
+
+	void ClearInspector();
+
+	TSharedPtr<IStructureDetailsView> ConstraintDataDetailsView;
+	TWeakPtr<FChaosVDScene> SceneWeakPtr;
+	FChaosVDCharacterGroundConstraintSelectionHandle CurrentDataSelectionHandle;
+
+	FName CurrentObjectBeingInspectedName;
+
+	bool bIsUpToDate = true;
+};
diff --git a/Engine/Plugins/ChaosVD/Source/ChaosVD/Private/Widgets/SChaosVDMainTab.cpp b/Engine/Plugins/ChaosVD/Source/ChaosVD/Private/Widgets/SChaosVDMainTab.cpp
--- a/Engine/Plugins/ChaosVD/Source/ChaosVD/Private/Widgets/SChaosVDMainTab.cpp
+++ b/Engine/Plugins/ChaosVD/Source/ChaosVD/Private/Widgets/SChaosVDMainTab.cpp
@@ -12,6 +12,7 @@
 #include "ChaosVDPlaybackViewportTab.h"
 #include "ChaosVDScene.h"
 #include "ChaosVDSolversTracksTab.h"
+#include "ChaosVDCharacterGroundConstraintDataInspectorTab.h"
 #include "ChaosVDCollisionDataDetailsTab.h"
 #include "ChaosVDConstraintDataInspectorTab.h"
 #include "ChaosVDSceneQueryDataInspectorTab.h"
@@ -31,11 +32,13 @@
 #include "ToolMenu.h"
 #include "ToolMenus.h"
 #include "Components/ChaosVDParticleDataComponent.h"
+#include "Components/ChaosVDSolverCharacterGroundConstraintDataComponent.h"
 #include "Components/ChaosVDSolverCollisionDataComponent.h"
 #include "Components/ChaosVDSolverJointConstraintDataComponent.h"
 #include "Styling/StyleColors.h"
 #include "Styling/ToolBarStyle.h"
 #include "Trace/ChaosVDTraceManager.h"
+#include "Visualizers/ChaosVDCharacterGroundConstraintsDataComponentVisualizer.h"
 #include "Visualizers/ChaosVDJointConstraintsDataComponentVisualizer.h"
 #include "Visualizers/ChaosVDParticleDataComponentVisualizer.h"
 #include "Visualizers/ChaosVDSceneQueryDataComponentVisualizer.h"
@@ -62,6 +65,7 @@ void SChaosVDMainTab::Construct(const FArguments& InArgs, TSharedPtr<FChaosVDEng
 	RegisterComponentVisualizer(UChaosVDSceneQueryDataComponent::StaticClass()->GetFName(), MakeShared<FChaosVDSceneQueryDataComponentVisualizer>());
 	RegisterComponentVisualizer(UChaosVDParticleDataComponent::StaticClass()->GetFName(), MakeShared<FChaosVDParticleDataComponentVisualizer>());
 	RegisterComponentVisualizer(UChaosVDSolverJointConstraintDataComponent::StaticClass()->GetFName(), MakeShared<FChaosVDJointConstraintsDataComponentVisualizer>());
+	RegisterComponentVisualizer(UChaosVDSolverCharacterGroundConstraintDataComponent::StaticClass()->GetFName(), MakeShared<FChaosVDCharacterGroundConstraintDataComponentVisualizer>());
 
 	TabManager = FGlobalTabmanager::Get()->NewTabManager(InArgs._OwnerTab.ToSharedRef()).ToSharedPtr();
 
@@ -74,6 +78,7 @@ void SChaosVDMainTab::Construct(const FArguments& InArgs, TSharedPtr<FChaosVDEng
 	RegisterTabSpawner<FChaosVDCollisionDataDetailsTab>(FChaosVDTabID::CollisionDataDetails);
 	RegisterTabSpawner<FChaosVDSceneQueryDataInspectorTab>(FChaosVDTabID::SceneQueryDataDetails);
 	RegisterTabSpawner<FChaosVDConstraintDataInspectorTab>(FChaosVDTabID::JointsDataDetails);
+	RegisterTabSpawner<FChaosVDCharacterGroundConstraintDataInspectorTab>(FChaosVDTabID::CharacterGroundConstraintDataDetails);
 
 	StatusBarID = FName(FChaosVDTabID::StatusBar.ToString() + InChaosVDEngine->GetInstanceGuid().ToString());
 	
@@ -282,6 +287,7 @@ TSharedRef<FTabManager::FLayout> SChaosVDMainTab::GenerateMainLayout()
 					->AddTab(FChaosVDTabID::CollisionDataDetails, ETabState::OpenedTab)
 					->AddTab(FChaosVDTabID::SceneQueryDataDetails, ETabState::OpenedTab)
 					->AddTab(FChaosVDTabID::JointsDataDetails, ETabState::ClosedTab)
+					->AddTab(FChaosVDTabID::CharacterGroundConstraintDataDetails, ETabState::ClosedTab)
 					->AddTab(FChaosVDTabID::CVDEditorSettings, ETabState::ClosedTab)
 				)
 			)
diff --git a/Engine/Source/Runtime/Experimental/Chaos/Private/Chaos/PBDRigidsEvolutionGBF.cpp b/Engine/Source/Runtime/Experimental/Chaos/Private/Chaos/PBDRigidsEvolutionGBF.cpp
--- a/Engine/Source/Runtime/Experimental/Chaos/Private/Chaos/PBDRigidsEvolutionGBF.cpp
+++ b/Engine/Source/Runtime/Experimental/Chaos/Private/Chaos/PBDRigidsEvolutionGBF.cpp
@@ -727,6 +727,7 @@ void FPBDRigidsEvolutionGBF::AdvanceOneTimeStepImpl(const FReal Dt, const FSubSt
 		}
 
 		CVD_TRACE_JOINT_CONSTRAINTS(CVDDC_JointConstraints, JointConstraints);
+		CVD_TRACE_CHARACTER_GROUND_CONSTRAINTS(CVDDC_CharacterGroundConstraints, CharacterGroundConstraints);
 
 		CVD_TRACE_PARTICLES_SOA(Particles);
 	}
diff --git a/Engine/Source/Runtime/Experimental/Chaos/Private/ChaosVisualDebugger/ChaosVDDataWrapperUtils.cpp b/Engine/Source/Runtime/Experimental/Chaos/Private/ChaosVisualDebugger/ChaosVDDataWrapperUtils.cpp
--- a/Engine/Source/Runtime/Experimental/Chaos/Private/ChaosVisualDebugger/ChaosVDDataWrapperUtils.cpp
+++ b/Engine/Source/Runtime/Experimental/Chaos/Private/ChaosVisualDebugger/ChaosVDDataWrapperUtils.cpp
@@ -5,14 +5,17 @@
 #if WITH_CHAOS_VISUAL_DEBUGGER
 
 #include "Chaos/PBDJointConstraints.h"
+#include "DataWrappers/ChaosVDCharacterGroundConstraintDataWrappers.h"
 #include "DataWrappers/ChaosVDJointDataWrappers.h"
 
+#include "Chaos/Character/CharacterGroundConstraintContainer.h"
 #include "Chaos/Collision/ParticlePairMidPhase.h"
 #include "Chaos/Collision/CollisionConstraintAllocator.h"
 #include "Chaos/ParticleHandle.h"
 #include "ChaosVisualDebugger/ChaosVDSerializedNameTable.h"
 #include "DataWrappers/ChaosVDCollisionDataWrappers.h"
 #include "DataWrappers/ChaosVDParticleDataWrapper.h"
+#include "Math/UnitConversion.h"
 
 namespace Chaos::VisualDebugger::Utils
 {
@@ -363,6 +366,77 @@ FChaosVDJointConstraint FChaosVDDataWrapperUtils::BuildJointDataWrapper(const Ch
 	return MoveTemp(WrappedJointData);
 }
 
+FChaosVDCharacterGroundConstraint FChaosVDDataWrapperUtils::BuildCharacterGroundConstraintDataWrapper(const Chaos::FCharacterGroundConstraintHandle* ConstraintHandle)
+{
+	FChaosVDCharacterGroundConstraint WrappedConstraintData;
+
+	if (ConstraintHandle)
+	{
+		if (ConstraintHandle->GetCharacterParticle())
+		{
+			WrappedConstraintData.CharacterParticleIndex = ConstraintHandle->GetCharacterParticle()->UniqueIdx().Idx;
+		}
+		else
+		{
+			WrappedConstraintData.CharacterParticleIndex = INDEX_NONE;
+		}
+		
+		WrappedConstraintData.ConstraintIndex = WrappedConstraintData.CharacterParticleIndex; // TODO - add unique index to constraint handle
+
+		if (ConstraintHandle->GetGroundParticle())
+		{
+			WrappedConstraintData.GroundParticleIndex = ConstraintHandle->GetGroundParticle()->UniqueIdx().Idx;
+		}
+		else
+		{
+			WrappedConstraintData.GroundParticleIndex = INDEX_NONE;
+		}
+		WrappedConstraintData.State.bDisabled = !ConstraintHandle->IsEnabled();
+
+		//TODO: Island related data getters are deprecated. We need to see where is best to get that data now and if this should be recorded as part of the CVD Constraint wrapper
+		//WrappedConstraintData.State.Color = ConstraintHandle->GetConstraintColor();
+		//WrappedConstraintData.State.Island = ConstraintHandle->GetConstraintIsland();
+		//WrappedConstraintData.State.IslandSize = ConstraintHandle->GetConstraintIsland();
+
+		WrappedConstraintData.State.SolverAppliedForce = ConstraintHandle->GetSolverAppliedForce();
+		WrappedConstraintData.State.SolverAppliedTorque = ConstraintHandle->GetSolverAppliedTorque();
+
+		WrappedConstraintData.State.SolverAppliedForce.X = FUnitConversion::Convert(ConstraintHandle->GetSolverAppliedForce().X, EUnit::KilogramCentimetersPerSecondSquared, EUnit::Newtons);
+		WrappedConstraintData.State.SolverAppliedForce.Y = FUnitConversion::Convert(ConstraintHandle->GetSolverAppliedForce().Y, EUnit::KilogramCentimetersPerSecondSquared, EUnit::Newtons);
+		WrappedConstraintData.State.SolverAppliedForce.Z = FUnitConversion::Convert(ConstraintHandle->GetSolverAppliedForce().Z, EUnit::KilogramCentimetersPerSecondSquared, EUnit::Newtons);
+		WrappedConstraintData.State.SolverAppliedTorque.X = FUnitConversion::Convert(ConstraintHandle->GetSolverAppliedTorque().X, EUnit::KilogramCentimetersSquaredPerSecondSquared, EUnit::NewtonMeters);
+		WrappedConstraintData.State.SolverAppliedTorque.Y = FUnitConversion::Convert(ConstraintHandle->GetSolverAppliedTorque().Y, EUnit::KilogramCentimetersSquaredPerSecondSquared, EUnit::NewtonMeters);
+		WrappedConstraintData.State.SolverAppliedTorque.Z = FUnitConversion::Convert(ConstraintHandle->GetSolverAppliedTorque().Z, EUnit::KilogramCentimetersSquaredPerSecondSquared, EUnit::NewtonMeters);
+
+		WrappedConstraintData.State.MarkAsValid();
+
+		CVD_COPY_FIELD_TO_WRAPPER(ConstraintHandle->GetSettings(), WrappedConstraintData.Settings, VerticalAxis);
+		CVD_COPY_FIELD_TO_WRAPPER(ConstraintHandle->GetSettings(), WrappedConstraintData.Settings, TargetHeight);
+		WrappedConstraintData.Settings.RadialForceLimit = FUnitConversion::Convert(ConstraintHandle->GetSettings().RadialForceLimit, EUnit::KilogramCentimetersPerSecondSquared, EUnit::Newtons);
+		WrappedConstraintData.Settings.FrictionForceLimit = FUnitConversion::Convert(ConstraintHandle->GetSettings().FrictionForceLimit, EUnit::KilogramCentimetersPerSecondSquared, EUnit::Newtons);
+		WrappedConstraintData.Settings.SwingTorqueLimit = FUnitConversion::Convert(ConstraintHandle->GetSettings().SwingTorqueLimit, EUnit::KilogramCentimetersSquaredPerSecondSquared, EUnit::NewtonMeters);
+		WrappedConstraintData.Settings.TwistTorqueLimit = FUnitConversion::Convert(ConstraintHandle->GetSettings().TwistTorqueLimit, EUnit::KilogramCentimetersSquaredPerSecondSquared, EUnit::NewtonMeters);
+		CVD_COPY_FIELD_TO_WRAPPER(ConstraintHandle->GetSettings(), WrappedConstraintData.Settings, CosMaxWalkableSlopeAngle);
+		CVD_COPY_FIELD_TO_WRAPPER(ConstraintHandle->GetSettings(), WrappedConstraintData.Settings, DampingFactor);
+		CVD_COPY_FIELD_TO_WRAPPER(ConstraintHandle->GetSettings(), WrappedConstraintData.Settings, AssumedOnGroundHeight);
+
+		WrappedConstraintData.Settings.MarkAsValid();
+
+		CVD_COPY_FIELD_TO_WRAPPER(ConstraintHandle->GetData(), WrappedConstraintData.Data, GroundNormal);
+		CVD_COPY_FIELD_TO_WRAPPER(ConstraintHandle->GetData(), WrappedConstraintData.Data, TargetDeltaPosition);
+		CVD_COPY_FIELD_TO_WRAPPER(ConstraintHandle->GetData(), WrappedConstraintData.Data, TargetDeltaFacing);
+		CVD_COPY_FIELD_TO_WRAPPER(ConstraintHandle->GetData(), WrappedConstraintData.Data, GroundDistance);
+		CVD_COPY_FIELD_TO_WRAPPER(ConstraintHandle->GetData(), WrappedConstraintData.Data, CosMaxWalkableSlopeAngle);
+
+		WrappedConstraintData.Data.MarkAsValid();
+
+	}
+
+	WrappedConstraintData.MarkAsValid();
+
+	return MoveTemp(WrappedConstraintData);
+}
+
 void FChaosVDDataWrapperUtils::CopyShapeDataToWrapper(const Chaos::FShapeInstancePtr& ShapeDataPtr, FChaosVDShapeCollisionData& OutCopyTo)
 {
 	const Chaos::FCollisionData& CollisionData = ShapeDataPtr->GetCollisionData();
diff --git a/Engine/Source/Runtime/Experimental/Chaos/Private/ChaosVisualDebugger/ChaosVDOptionalDataChannel.cpp b/Engine/Source/Runtime/Experimental/Chaos/Private/ChaosVisualDebugger/ChaosVDOptionalDataChannel.cpp
--- a/Engine/Source/Runtime/Experimental/Chaos/Private/ChaosVisualDebugger/ChaosVDOptionalDataChannel.cpp
+++ b/Engine/Source/Runtime/Experimental/Chaos/Private/ChaosVisualDebugger/ChaosVDOptionalDataChannel.cpp
@@ -84,5 +84,6 @@ CVD_DEFINE_OPTIONAL_DATA_CHANNEL(EndOfEvolutionCollisionConstraints, EChaosVDDat
 CVD_DEFINE_OPTIONAL_DATA_CHANNEL(EvolutionEnd, EChaosVDDataChannelInitializationFlags::StartEnabled); // Intentionally not allowing changing the enabled state because otherwise there might be nothing to visualize
 CVD_DEFINE_OPTIONAL_DATA_CHANNEL(SceneQueries, EChaosVDDataChannelInitializationFlags::StartEnabled | EChaosVDDataChannelInitializationFlags::CanChangeEnabledState);
 CVD_DEFINE_OPTIONAL_DATA_CHANNEL(JointConstraints, EChaosVDDataChannelInitializationFlags::CanChangeEnabledState);
+CVD_DEFINE_OPTIONAL_DATA_CHANNEL(CharacterGroundConstraints, EChaosVDDataChannelInitializationFlags::CanChangeEnabledState);
 
 #endif
diff --git a/Engine/Source/Runtime/Experimental/Chaos/Private/ChaosVisualDebugger/ChaosVisualDebuggerTrace.cpp b/Engine/Source/Runtime/Experimental/Chaos/Private/ChaosVisualDebugger/ChaosVisualDebuggerTrace.cpp
--- a/Engine/Source/Runtime/Experimental/Chaos/Private/ChaosVisualDebugger/ChaosVisualDebuggerTrace.cpp
+++ b/Engine/Source/Runtime/Experimental/Chaos/Private/ChaosVisualDebugger/ChaosVisualDebuggerTrace.cpp
@@ -4,6 +4,7 @@
 
 #if WITH_CHAOS_VISUAL_DEBUGGER
 
+#include "Chaos/Character/CharacterGroundConstraintContainer.h"
 #include "Chaos/Framework/PhysicsSolverBase.h"
 #include "Chaos/ImplicitObject.h"
 #include "Chaos/ParticleHandle.h"
@@ -14,6 +15,7 @@
 #include "ChaosVisualDebugger/ChaosVDMemWriterReader.h"
 #include "ChaosVisualDebugger/ChaosVDSerializedNameTable.h"
 #include "Compression/OodleDataCompressionUtil.h"
+#include "DataWrappers/ChaosVDCharacterGroundConstraintDataWrappers.h"
 #include "DataWrappers/ChaosVDCollisionDataWrappers.h"
 #include "DataWrappers/ChaosVDImplicitObjectDataWrapper.h"
 #include "DataWrappers/ChaosVDJointDataWrappers.h"
@@ -353,6 +355,39 @@ void FChaosVisualDebuggerTrace::TraceJointsConstraints(Chaos::FPBDJointConstrain
 	});
 }
 
+void FChaosVisualDebuggerTrace::TraceCharacterGroundConstraints(Chaos::FCharacterGroundConstraintContainer& InConstraints)
+{
+	using namespace Chaos::VisualDebugger::Utils;
+
+	if (!IsTracing())
+	{
+		return;
+	}
+
+	const FChaosVDContext* CVDContextData = FChaosVDThreadContext::Get().GetCurrentContext(EChaosVDContextType::Solver);
+
+	if (!IsContextEnabledAndValid(CVDContextData))
+	{
+		return;
+	}
+
+	const Chaos::FCharacterGroundConstraintContainer::FConstConstraints& ConstraintHandles = InConstraints.GetConstConstraints();
+
+	ParallelFor(ConstraintHandles.Num(), [&ConstraintHandles, CopyContext = *CVDContextData](int32 ConstraintIndex)
+	{
+		CVD_SCOPE_CONTEXT(CopyContext);
+
+		FChaosVDCharacterGroundConstraint WrappedConstraintData = FChaosVDDataWrapperUtils::BuildCharacterGroundConstraintDataWrapper(ConstraintHandles[ConstraintIndex]);
+
+		WrappedConstraintData.SolverID = CopyContext.Id;
+
+		FChaosVDScopedTLSBufferAccessor TLSDataBuffer;
+		Chaos::VisualDebugger::WriteDataToBuffer(TLSDataBuffer.BufferRef, WrappedConstraintData);
+
+		TraceBinaryData(TLSDataBuffer.BufferRef, FChaosVDCharacterGroundConstraint::WrapperTypeName);
+	});
+}
+
 void FChaosVisualDebuggerTrace::TraceCollisionConstraint(const Chaos::FPBDCollisionConstraint* CollisionConstraint)
 {
 	using namespace Chaos::VisualDebugger::Utils;
@@ -419,6 +454,10 @@ void FChaosVisualDebuggerTrace::TraceConstraintsContainer(TConstArrayView<Chaos:
 			{
 				CVD_TRACE_STEP_MID_PHASES_FROM_COLLISION_CONSTRAINTS(CVDDC_EndOfEvolutionCollisionConstraints, *static_cast<Chaos::FPBDCollisionConstraints*>(ConstraintContainer));
 			}
+			else if (ConstraintContainer->GetConstraintHandleType().IsA(Chaos::FCharacterGroundConstraintHandle::StaticType()))
+			{
+				CVD_TRACE_CHARACTER_GROUND_CONSTRAINTS(CVDDC_CharacterGroundConstraints, *static_cast<Chaos::FCharacterGroundConstraintContainer*>(ConstraintContainer));
+			}
 		}
 	}
 }
diff --git a/Engine/Source/Runtime/Experimental/Chaos/Public/ChaosVisualDebugger/ChaosVDDataWrapperUtils.h b/Engine/Source/Runtime/Experimental/Chaos/Public/ChaosVisualDebugger/ChaosVDDataWrapperUtils.h
--- a/Engine/Source/Runtime/Experimental/Chaos/Public/ChaosVisualDebugger/ChaosVDDataWrapperUtils.h
+++ b/Engine/Source/Runtime/Experimental/Chaos/Public/ChaosVisualDebugger/ChaosVDDataWrapperUtils.h
@@ -8,10 +8,12 @@
 #include "Chaos/ParticleHandleFwd.h"
 #include "DataWrappers/ChaosVDCollisionDataWrappers.h"
 
+struct FChaosVDCharacterGroundConstraint;
 struct FChaosVDJointConstraint;
 
 namespace Chaos
 {
+	class FCharacterGroundConstraintHandle;
 	class FPBDJointConstraintHandle;
 
 	namespace VisualDebugger
@@ -62,6 +64,9 @@ private:
 	/** Creates and populates a FChaosVDJointConstraint with the data of the provided FPBDJointConstraintHandle */
 	static FChaosVDJointConstraint BuildJointDataWrapper(const Chaos::FPBDJointConstraintHandle* ConstaintHanlde);
 
+	/** Creates and populates a FChaosVDChartacterGroundConstraint with the data of the provided FCharacterGroundConstraintHandle */
+	static FChaosVDCharacterGroundConstraint BuildCharacterGroundConstraintDataWrapper(const Chaos::FCharacterGroundConstraintHandle* ConstaintHanlde);
+
 	/** Converts a Chaos::FVec3 to a FVector. It is worth notice that FVector is double precision and FVec3 is single */
 	static FVector ConvertToFVector(const Chaos::FVec3f& VectorRef) { return FVector(VectorRef); }
 
diff --git a/Engine/Source/Runtime/Experimental/Chaos/Public/ChaosVisualDebugger/ChaosVDOptionalDataChannel.h b/Engine/Source/Runtime/Experimental/Chaos/Public/ChaosVisualDebugger/ChaosVDOptionalDataChannel.h
--- a/Engine/Source/Runtime/Experimental/Chaos/Public/ChaosVisualDebugger/ChaosVDOptionalDataChannel.h
+++ b/Engine/Source/Runtime/Experimental/Chaos/Public/ChaosVisualDebugger/ChaosVDOptionalDataChannel.h
@@ -152,6 +152,7 @@ CVD_DECLARE_OPTIONAL_DATA_CHANNEL(PostConstraintSolve);						// Particle positio
 CVD_DECLARE_OPTIONAL_DATA_CHANNEL(EvolutionEnd);							// Final particle positions, including destruction and user callbacks.
 CVD_DECLARE_OPTIONAL_DATA_CHANNEL(SceneQueries);
 CVD_DECLARE_OPTIONAL_DATA_CHANNEL(JointConstraints);
+CVD_DECLARE_OPTIONAL_DATA_CHANNEL(CharacterGroundConstraints);
 
 #else  //WITH_CHAOS_VISUAL_DEBUGGER
 
diff --git a/Engine/Source/Runtime/Experimental/Chaos/Public/ChaosVisualDebugger/ChaosVDTraceMacros.h b/Engine/Source/Runtime/Experimental/Chaos/Public/ChaosVisualDebugger/ChaosVDTraceMacros.h
--- a/Engine/Source/Runtime/Experimental/Chaos/Public/ChaosVisualDebugger/ChaosVDTraceMacros.h
+++ b/Engine/Source/Runtime/Experimental/Chaos/Public/ChaosVisualDebugger/ChaosVDTraceMacros.h
@@ -153,6 +153,14 @@
 			}
 	#endif
 
+	#ifndef CVD_TRACE_CHARACTER_GROUND_CONSTRAINTS
+		#define CVD_TRACE_CHARACTER_GROUND_CONSTRAINTS(DataChannel, InConstraints) \
+			{ \
+				CVD_SCOPED_DATA_CHANNEL_OVERRIDE(DataChannel) \
+				FChaosVisualDebuggerTrace::TraceCharacterGroundConstraints(InConstraints); \
+			}
+	#endif
+
 	#ifndef CVD_TRACE_CONSTRAINTS_CONTAINER
 		#define CVD_TRACE_CONSTRAINTS_CONTAINER(ContainerView) \
 			FChaosVisualDebuggerTrace::TraceConstraintsContainer(ContainerView);
@@ -246,6 +254,10 @@
 		#define CVD_TRACE_JOINT_CONSTRAINTS(DataChannel, InJointConstraints)
 	#endif
 
+	#ifndef CVD_TRACE_CHARACTER_GROUND_CONSTRAINTS
+		#define CVD_TRACE_CHARACTER_GROUND_CONSTRAINTS(DataChannel, InConstraints)
+	#endif
+
 	#ifndef CVD_TRACE_CONSTRAINTS_CONTAINER
 		#define CVD_TRACE_CONSTRAINTS_CONTAINER(ContainerView)
 	#endif
diff --git a/Engine/Source/Runtime/Experimental/Chaos/Public/ChaosVisualDebugger/ChaosVisualDebuggerTrace.h b/Engine/Source/Runtime/Experimental/Chaos/Public/ChaosVisualDebugger/ChaosVisualDebuggerTrace.h
--- a/Engine/Source/Runtime/Experimental/Chaos/Public/ChaosVisualDebugger/ChaosVisualDebuggerTrace.h
+++ b/Engine/Source/Runtime/Experimental/Chaos/Public/ChaosVisualDebugger/ChaosVisualDebuggerTrace.h
@@ -28,6 +28,7 @@
 
 namespace Chaos
 {
+class FCharacterGroundConstraintContainer;
 class FPBDConstraintContainer;
 class FPBDJointConstraints;
 }
@@ -198,6 +199,9 @@ public:
 	/** Traces all joint constraints in the provided container */
 	static CHAOS_API void TraceJointsConstraints(Chaos::FPBDJointConstraints& InJointConstraints);
 
+	/** Traces all character ground constraints in the provided container */
+	static CHAOS_API void TraceCharacterGroundConstraints(Chaos::FCharacterGroundConstraintContainer& InConstraints);
+
 	/** Traces a Particle pair MidPhase as binary data */
 	static CHAOS_API void TraceCollisionConstraint(const Chaos::FPBDCollisionConstraint* CollisionConstraint);
 
diff --git a/Engine/Source/Runtime/Experimental/ChaosVDData/Public/ChaosVDRecording.h b/Engine/Source/Runtime/Experimental/ChaosVDData/Public/ChaosVDRecording.h
--- a/Engine/Source/Runtime/Experimental/ChaosVDData/Public/ChaosVDRecording.h
+++ b/Engine/Source/Runtime/Experimental/ChaosVDData/Public/ChaosVDRecording.h
@@ -11,6 +11,7 @@
 #include "Chaos/ImplicitFwd.h"
 #include "Chaos/ImplicitObject.h"
 #include "ChaosVisualDebugger/ChaosVDMemWriterReader.h"
+#include "DataWrappers/ChaosVDCharacterGroundConstraintDataWrappers.h"
 #include "DataWrappers/ChaosVDJointDataWrappers.h"
 #include "DataWrappers/ChaosVDQueryDataWrappers.h"
 #include <atomic>
@@ -60,7 +61,7 @@ struct CHAOSVDDATA_API FChaosVDSolverFrameData
 	TSet<int32> ParticlesDestroyedIDs;
 	double StartTime = -1.0;
 	double EndTime = -1.0;
-
+	
 	/** Calculates and returns the frame time for this recorded frame.
 	 * @return Calculated frame time. -1 if it was not recorded
 	 */
@@ -98,6 +99,7 @@ struct FChaosVDGameFrameData
 	TMap<FName, FChaosVDTrackedLocation> RecordedNonSolverLocationsByID;
 	TMap<FName, FChaosVDTrackedTransform> RecordedNonSolverTransformsByID;
 	TMap<int32, TSharedPtr<FChaosVDQueryDataWrapper>> RecordedSceneQueries;
+	TArray<TSharedPtr<FChaosVDCharacterGroundConstraint>> RecordedCharacterGroundConstraints;
 };
 
 /**
diff --git a/Engine/Source/Runtime/Experimental/ChaosVisualDebugger/Private/DataWrappers/ChaosVDCharacterGroundConstraintDataWrappers.cpp b/Engine/Source/Runtime/Experimental/ChaosVisualDebugger/Private/DataWrappers/ChaosVDCharacterGroundConstraintDataWrappers.cpp
new file mode 100644
--- /dev/null
+++ b/Engine/Source/Runtime/Experimental/ChaosVisualDebugger/Private/DataWrappers/ChaosVDCharacterGroundConstraintDataWrappers.cpp
@@ -0,0 +1,86 @@
+// Copyright Epic Games, Inc. All Rights Reserved.
+
+#include "DataWrappers/ChaosVDCharacterGroundConstraintDataWrappers.h"
+
+#include "DataWrappers/ChaosVDDataSerializationMacros.h"
+
+bool FChaosVDCharacterGroundConstraintStateDataWrapper::Serialize(FArchive& Ar)
+{
+	Ar << bHasValidData;
+	if (!bHasValidData)
+	{
+		return !Ar.IsError();
+	}
+
+	Ar << Island;
+	Ar << Level;
+	Ar << Color;
+	Ar << IslandSize;
+
+	Ar << bDisabled;
+
+	Ar << SolverAppliedForce;
+	Ar << SolverAppliedTorque;
+	
+	return !Ar.IsError();
+}
+
+bool FChaosVDCharacterGroundConstraintSettingsDataWrapper::Serialize(FArchive& Ar)
+{
+	Ar << bHasValidData;
+	if (!bHasValidData)
+	{
+		return !Ar.IsError();
+	}
+	
+	Ar << VerticalAxis;
+	Ar << TargetHeight;
+	Ar << RadialForceLimit;
+	Ar << FrictionForceLimit;
+	Ar << TwistTorqueLimit;
+	Ar << SwingTorqueLimit;
+	Ar << CosMaxWalkableSlopeAngle;
+	Ar << DampingFactor;
+	Ar << AssumedOnGroundHeight;
+
+	return !Ar.IsError();
+}
+
+bool FChaosVDCharacterGroundConstraintDataDataWrapper::Serialize(FArchive& Ar)
+{
+	Ar << bHasValidData;
+	if (!bHasValidData)
+	{
+		return !Ar.IsError();
+	}
+
+	Ar << GroundNormal;
+	Ar << TargetDeltaPosition;
+	Ar << TargetDeltaFacing;
+	Ar << GroundDistance;
+	Ar << CosMaxWalkableSlopeAngle;
+
+	return !Ar.IsError();
+}
+
+
+bool FChaosVDCharacterGroundConstraint::Serialize(FArchive& Ar)
+{
+	Ar << bHasValidData;
+	if (!bHasValidData)
+	{
+		return !Ar.IsError();
+	}
+
+	Ar << SolverID;
+	Ar << ConstraintIndex;
+
+	Ar << CharacterParticleIndex;
+	Ar << GroundParticleIndex;
+
+	Ar << State;
+	Ar << Settings;
+	Ar << Data;
+
+	return !Ar.IsError();
+}
diff --git a/Engine/Source/Runtime/Experimental/ChaosVisualDebugger/Public/DataWrappers/ChaosVDCharacterGroundConstraintDataWrappers.h b/Engine/Source/Runtime/Experimental/ChaosVisualDebugger/Public/DataWrappers/ChaosVDCharacterGroundConstraintDataWrappers.h
new file mode 100644
--- /dev/null
+++ b/Engine/Source/Runtime/Experimental/ChaosVisualDebugger/Public/DataWrappers/ChaosVDCharacterGroundConstraintDataWrappers.h
@@ -0,0 +1,150 @@
+// Copyright Epic Games, Inc. All Rights Reserved.
+
+#pragma once
+
+#include "ChaosVDParticleDataWrapper.h"
+#include "UObject/Class.h"
+#include "UObject/ObjectMacros.h"
+
+#include "ChaosVDCharacterGroundConstraintDataWrappers.generated.h"
+
+#ifndef CVD_IMPLEMENT_SERIALIZER
+		#define CVD_IMPLEMENT_SERIALIZER(Type) \
+		inline FArchive& operator<<(FArchive& Ar, Type& Data) \
+		{\
+			Data.Serialize(Ar); \
+			return Ar; \
+		} \
+		template<>\
+		struct TStructOpsTypeTraits<Type> : public TStructOpsTypeTraitsBase2<Type> \
+		{\
+			enum\
+			{\
+				WithSerializer = true,\
+			};\
+		};\
+
+#endif
+
+
+USTRUCT()
+struct FChaosVDCharacterGroundConstraintStateDataWrapper : public FChaosVDWrapperDataBase
+{
+	GENERATED_BODY()
+public:
+
+	CHAOSVDRUNTIME_API bool Serialize(FArchive& Ar);
+
+	//TODO: Make the island data visible when we add support to record that data
+
+	int32 Island = INDEX_NONE;
+	int32 Level = INDEX_NONE;
+	int32 Color = INDEX_NONE;
+	int32 IslandSize = INDEX_NONE;
+
+	UPROPERTY(VisibleAnywhere, Category=ConstraintState)
+	bool bDisabled = false;
+
+	UPROPERTY(VisibleAnywhere, Category= ConstraintState)
+	FVector SolverAppliedForce = FVector::ZeroVector;
+	UPROPERTY(VisibleAnywhere, Category= ConstraintState)
+	FVector SolverAppliedTorque = FVector::ZeroVector;
+};
+CVD_IMPLEMENT_SERIALIZER(FChaosVDCharacterGroundConstraintStateDataWrapper)
+
+
+USTRUCT()
+struct FChaosVDCharacterGroundConstraintSettingsDataWrapper : public FChaosVDWrapperDataBase
+{
+	GENERATED_BODY()
+public:
+	CHAOSVDRUNTIME_API bool Serialize(FArchive& Ar);
+
+	UPROPERTY(VisibleAnywhere, Category = Settings)
+	FVector VerticalAxis = FVector(ForceInitToZero);
+
+	UPROPERTY(VisibleAnywhere, Category = Settings, meta = (ForceUnits = "cm"))
+	double TargetHeight = 0.0;
+
+	UPROPERTY(VisibleAnywhere, Category = Settings, meta = (ForceUnits = "Newtons"))
+	double RadialForceLimit = 0.0;
+
+	UPROPERTY(VisibleAnywhere, Category = Settings, meta = (ForceUnits = "Newtons"))
+	double FrictionForceLimit = 0.0;
+
+	UPROPERTY(VisibleAnywhere, Category = Settings, meta = (ForceUnits = "NewtonMeters"))
+	double TwistTorqueLimit = 0.0;
+
+	UPROPERTY(VisibleAnywhere, Category = Settings, meta = (ForceUnits = "NewtonMeters"))
+	double SwingTorqueLimit = 0.0;
+
+	UPROPERTY(VisibleAnywhere, Category = Settings)
+	double CosMaxWalkableSlopeAngle = 0.0;
+
+	UPROPERTY(VisibleAnywhere, Category = Settings)
+	double DampingFactor = 0.0;
+
+	UPROPERTY(VisibleAnywhere, Category = Settings, meta = (ForceUnits = "cm"))
+	double AssumedOnGroundHeight = 0.0;
+};
+CVD_IMPLEMENT_SERIALIZER(FChaosVDCharacterGroundConstraintSettingsDataWrapper)
+
+USTRUCT()
+struct FChaosVDCharacterGroundConstraintDataDataWrapper : public FChaosVDWrapperDataBase
+{
+	GENERATED_BODY()
+public:
+	CHAOSVDRUNTIME_API bool Serialize(FArchive& Ar);
+
+	UPROPERTY(VisibleAnywhere, Category = Data)
+	FVector GroundNormal = FVector(ForceInitToZero);
+
+	UPROPERTY(VisibleAnywhere, Category = Data, meta = (ForceUnits = "cm"))
+	FVector TargetDeltaPosition = FVector(ForceInitToZero);
+
+	UPROPERTY(VisibleAnywhere, Category = Data, meta = (ForceUnits = "radians"))
+	double TargetDeltaFacing = 0.0;
+
+	UPROPERTY(VisibleAnywhere, Category = Data, meta = (ForceUnits = "cm"))
+	double GroundDistance = 0.0;
+
+	UPROPERTY(VisibleAnywhere, Category = Data)
+	double CosMaxWalkableSlopeAngle = 0.0;
+};
+CVD_IMPLEMENT_SERIALIZER(FChaosVDCharacterGroundConstraintDataDataWrapper)
+
+
+USTRUCT()
+struct FChaosVDCharacterGroundConstraint : public FChaosVDWrapperDataBase
+{
+	GENERATED_BODY()
+public:
+
+	inline static FStringView WrapperTypeName = TEXT("FChaosVDCharacterGroundConstraint");
+
+	CHAOSVDRUNTIME_API bool Serialize(FArchive& Ar);
+
+	int32 SolverID = INDEX_NONE;
+
+	UPROPERTY(VisibleAnywhere, Category="General")
+	int32 ConstraintIndex = INDEX_NONE;
+
+	int32 CharacterParticleIndex = INDEX_NONE;
+	int32 GroundParticleIndex = INDEX_NONE;
+
+	UPROPERTY(VisibleAnywhere, Category=State)
+	FChaosVDCharacterGroundConstraintStateDataWrapper State;
+	
+	UPROPERTY(VisibleAnywhere, Category=Settings)
+	FChaosVDCharacterGroundConstraintSettingsDataWrapper Settings;
+
+	UPROPERTY(VisibleAnywhere, Category=Data)
+	FChaosVDCharacterGroundConstraintDataDataWrapper Data;
+
+	bool bIsSelectedInEditor = false;
+};
+
+CVD_IMPLEMENT_SERIALIZER(FChaosVDCharacterGroundConstraint)
+
+
+
