[ChaosVD] Refactor of the existing CVD Settings so each area has its settings object that can be saved if desired.

Also, re implemented the viewport show flags and other options as proper editor menus that can be extended and such.

diff --git a/Engine/Plugins/ChaosClothAssetEditor/Source/ChaosClothAssetEditor/Private/ChaosClothAsset/SClothEditorViewportToolBarBase.cpp b/Engine/Plugins/ChaosClothAssetEditor/Source/ChaosClothAssetEditor/Private/ChaosClothAsset/SClothEditorViewportToolBarBase.cpp
--- a/Engine/Plugins/ChaosClothAssetEditor/Source/ChaosClothAssetEditor/Private/ChaosClothAsset/SClothEditorViewportToolBarBase.cpp
+++ b/Engine/Plugins/ChaosClothAssetEditor/Source/ChaosClothAssetEditor/Private/ChaosClothAsset/SClothEditorViewportToolBarBase.cpp
@@ -45,37 +45,6 @@ TSharedRef<SWidget> SChaosClothAssetEditorViewportToolBarBase::GenerateClothView
 	return OptionsMenuBuilder.MakeWidget();
 }
 
-TSharedRef<SWidget> SChaosClothAssetEditorViewportToolBarBase::GenerateFOVMenu() const
-{
-	// Re-implementation of inaccessible SCommonEditorViewportToolbarBase::GenerateFOVMenu()
-
-	const float FOVMin = 5.f;
-	const float FOVMax = 170.f;
-
-	return 
-		SNew(SBox)
-		.HAlign(HAlign_Right)
-		[
-			SNew(SBox)
-			.Padding(FMargin(4.0f, 0.0f, 0.0f, 0.0f))
-			.WidthOverride(100.0f)
-			[
-				SNew(SBorder)
-				.BorderImage(FAppStyle::Get().GetBrush("Menu.WidgetBorder"))
-				.Padding(FMargin(1.0f))
-				[
-					SNew(SSpinBox<float>)
-					.Style(&FAppStyle::Get(), "Menu.SpinBox")
-					.Font(FAppStyle::GetFontStyle(TEXT("MenuItem.Font")))
-					.MinValue(FOVMin)
-					.MaxValue(FOVMax)
-					.Value(this, &SChaosClothAssetEditorViewportToolBarBase::OnGetFOVValue)
-					.OnValueChanged(this, &SChaosClothAssetEditorViewportToolBarBase::OnFOVValueChanged)
-				]
-			]
-		];
-}
-
 TSharedRef<SWidget> SChaosClothAssetEditorViewportToolBarBase::GenerateCameraSpeedSettingsMenu() const
 {
 	// This comes from STransformViewportToolBar::FillCameraSpeedMenu
diff --git a/Engine/Plugins/ChaosClothAssetEditor/Source/ChaosClothAssetEditor/Private/ChaosClothAsset/SClothEditorViewportToolBarBase.h b/Engine/Plugins/ChaosClothAssetEditor/Source/ChaosClothAssetEditor/Private/ChaosClothAsset/SClothEditorViewportToolBarBase.h
--- a/Engine/Plugins/ChaosClothAssetEditor/Source/ChaosClothAssetEditor/Private/ChaosClothAsset/SClothEditorViewportToolBarBase.h
+++ b/Engine/Plugins/ChaosClothAssetEditor/Source/ChaosClothAssetEditor/Private/ChaosClothAsset/SClothEditorViewportToolBarBase.h
@@ -9,7 +9,6 @@ class CHAOSCLOTHASSETEDITOR_API SChaosClothAssetEditorViewportToolBarBase : publ
 {
 protected:
 	TSharedRef<SWidget> GenerateClothViewportOptionsMenu() const;
-	TSharedRef<SWidget> GenerateFOVMenu() const;
 	TSharedRef<SWidget> GenerateCameraSpeedSettingsMenu() const;
 	float OnGetFOVValue() const;
 	FText GetCameraSpeedLabel() const;
diff --git a/Engine/Plugins/ChaosVD/Config/DefaultChaosVD.ini b/Engine/Plugins/ChaosVD/Config/DefaultChaosVD.ini
new file mode 100644
--- /dev/null
+++ b/Engine/Plugins/ChaosVD/Config/DefaultChaosVD.ini
@@ -0,0 +1,2 @@
+[SectionsToSave]
+bCanSaveAllSections=true
\ No newline at end of file
diff --git a/Engine/Plugins/ChaosVD/Config/Engine.ini b/Engine/Plugins/ChaosVD/Config/Engine.ini
--- a/Engine/Plugins/ChaosVD/Config/Engine.ini
+++ b/Engine/Plugins/ChaosVD/Config/Engine.ini
@@ -1,4 +1,4 @@
-﻿[/Script/ChaosVD.ChaosVDEditorSettings]
+﻿[/Script/ChaosVD.ChaosVDCoreSettings]
 QueryOnlyMeshesMaterial=/ChaosVD/Materials/ChaosVDQueryOnlyMesh.ChaosVDQueryOnlyMesh
 SimOnlyMeshesMaterial=/ChaosVD/Materials/ChasoVDSimMesh.ChasoVDSimMesh
 InstancedMeshesMaterial=/ChaosVD/Materials/ChaosVDInstanced.ChaosVDInstanced
diff --git a/Engine/Plugins/ChaosVD/Source/ChaosVD/ChaosVD.Build.cs b/Engine/Plugins/ChaosVD/Source/ChaosVD/ChaosVD.Build.cs
--- a/Engine/Plugins/ChaosVD/Source/ChaosVD/ChaosVD.Build.cs
+++ b/Engine/Plugins/ChaosVD/Source/ChaosVD/ChaosVD.Build.cs
@@ -46,7 +46,8 @@ public class ChaosVD : ModuleRules
 				"StatusBar",
 				"SubobjectEditor",
 				"SubobjectDataInterface",
-				"PropertyEditor"
+				"PropertyEditor",
+				"CommonMenuExtensions"
 			}
 			);
 
diff --git a/Engine/Plugins/ChaosVD/Source/ChaosVD/Private/Actors/ChaosVDSolverInfoActor.cpp b/Engine/Plugins/ChaosVD/Source/ChaosVD/Private/Actors/ChaosVDSolverInfoActor.cpp
--- a/Engine/Plugins/ChaosVD/Source/ChaosVD/Private/Actors/ChaosVDSolverInfoActor.cpp
+++ b/Engine/Plugins/ChaosVD/Source/ChaosVD/Private/Actors/ChaosVDSolverInfoActor.cpp
@@ -13,6 +13,7 @@
 #include "Elements/Framework/TypedElementSelectionSet.h"
 #include "Engine/World.h"
 #include "Misc/ScopedSlowTask.h"
+#include "Settings/ChaosVDParticleVisualizationSettings.h"
 
 #define LOCTEXT_NAMESPACE "ChaosVisualDebugger"
 
@@ -23,6 +24,17 @@ AChaosVDSolverInfoActor::AChaosVDSolverInfoActor(const FObjectInitializer& Objec
 	JointsDataComponent = CreateDefaultSubobject<UChaosVDSolverJointConstraintDataComponent>(TEXT("JointDataComponent"));
 	CharacterGroundConstraintDataComponent = CreateDefaultSubobject<UChaosVDSolverCharacterGroundConstraintDataComponent>(TEXT("CharacterGroundConstraintDataComponent"));
 	bIsServer = false;
+
+	if (UChaosVDParticleVisualizationSettings* ParticleVisualizationSettings = GetMutableDefault<UChaosVDParticleVisualizationSettings>())
+	{
+		ParticleVisualizationSettings->OnSettingsChanged().AddUObject(this, &AChaosVDSolverInfoActor::HandleVisibilitySettingsUpdated);
+	}
+	
+	if (UChaosVDParticleVisualizationColorSettings* ColorVisualizationSettings = GetMutableDefault<UChaosVDParticleVisualizationColorSettings>())
+	{
+		ColorVisualizationSettings->OnSettingsChanged().AddUObject(this, &AChaosVDSolverInfoActor::HandleColorsSettingsUpdated);
+	}
+
 }
 
 void AChaosVDSolverInfoActor::SetSolverName(const FString& InSolverName)
@@ -97,7 +109,7 @@ bool AChaosVDSolverInfoActor::SelectParticleByID(int32 ParticleIDToSelect)
 	return true;
 }
 
-void AChaosVDSolverInfoActor::HandleVisibilitySettingsUpdated()
+void AChaosVDSolverInfoActor::HandleVisibilitySettingsUpdated(UObject* SettingsObject)
 {
 	for (const TPair<int32, AChaosVDParticleActor*>& ParticleWithIDPair : SolverParticlesByID)
 	{
@@ -108,7 +120,7 @@ void AChaosVDSolverInfoActor::HandleVisibilitySettingsUpdated()
 	}
 }
 
-void AChaosVDSolverInfoActor::HandleColorsSettingsUpdated()
+void AChaosVDSolverInfoActor::HandleColorsSettingsUpdated(UObject* SettingsObject)
 {
 	for (const TPair<int32, AChaosVDParticleActor*>& ParticleWithIDPair : SolverParticlesByID)
 	{
@@ -206,6 +218,16 @@ void AChaosVDSolverInfoActor::Destroyed()
 		return;
 	}
 
+	if (UChaosVDParticleVisualizationSettings* ParticleVisualizationSettings = GetMutableDefault<UChaosVDParticleVisualizationSettings>())
+	{
+		ParticleVisualizationSettings->OnSettingsChanged().RemoveAll(this);
+	}
+	
+	if (UChaosVDParticleVisualizationColorSettings* ColorVisualizationSettings = GetMutableDefault<UChaosVDParticleVisualizationColorSettings>())
+	{
+		ColorVisualizationSettings->OnSettingsChanged().RemoveAll(this);
+	}
+
 	constexpr float AmountOfWork = 1.0f;
 	const float PercentagePerElement = 1.0f / SolverParticlesByID.Num();
 
diff --git a/Engine/Plugins/ChaosVD/Source/ChaosVD/Private/Actors/ChaosVDSolverInfoActor.h b/Engine/Plugins/ChaosVD/Source/ChaosVD/Private/Actors/ChaosVDSolverInfoActor.h
--- a/Engine/Plugins/ChaosVD/Source/ChaosVD/Private/Actors/ChaosVDSolverInfoActor.h
+++ b/Engine/Plugins/ChaosVD/Source/ChaosVD/Private/Actors/ChaosVDSolverInfoActor.h
@@ -62,8 +62,6 @@ public:
 	template <typename TCallback>
 	void VisitAllParticleData(TCallback VisitCallback);
 
-	void HandleVisibilitySettingsUpdated();
-	void HandleColorsSettingsUpdated();
 	void RemoveSolverFolders(UWorld* World);
 
 	bool IsVisible() const;
@@ -75,6 +73,9 @@ public:
 
 protected:
 
+	void HandleVisibilitySettingsUpdated(UObject* SettingsObject);
+	void HandleColorsSettingsUpdated(UObject* SettingsObject);
+
 	void ApplySolverVisibilityToParticle(AChaosVDParticleActor* ParticleActor, bool bIsHidden);
 
 	virtual void HandlePostSelectionChange(const UTypedElementSelectionSet* ChangesSelectionSet) override;
diff --git a/Engine/Plugins/ChaosVD/Source/ChaosVD/Private/ChaosVDCommands.cpp b/Engine/Plugins/ChaosVD/Source/ChaosVD/Private/ChaosVDCommands.cpp
--- a/Engine/Plugins/ChaosVD/Source/ChaosVD/Private/ChaosVDCommands.cpp
+++ b/Engine/Plugins/ChaosVD/Source/ChaosVD/Private/ChaosVDCommands.cpp
@@ -6,13 +6,14 @@
 
 #define LOCTEXT_NAMESPACE "ChaosVisualDebugger"
 
-FChaosVDCommands::FChaosVDCommands() : TCommands<FChaosVDCommands>(TEXT("ActorPalette"), NSLOCTEXT("ChaosVisualDebugger", "ChaosVisualDebuggerEditor", "Chaos Visual Debugger Editor"), NAME_None, FChaosVDStyle::GetStyleSetName())
+FChaosVDCommands::FChaosVDCommands() : TCommands<FChaosVDCommands>(TEXT("ChaosVDEditor"), LOCTEXT("ChaosVisualDebuggerEditor", "Chaos Visual Debugger Editor"), NAME_None, FChaosVDStyle::GetStyleSetName())
 {
 }
 
 void FChaosVDCommands::RegisterCommands()
 {
-	UI_COMMAND(TrackUntrackSelectedObject, "Start or Stop tracking the selected object", "Start or Stop tracking the selected object", EUserInterfaceActionType::Button, FInputChord(EKeys::F8));
+	UI_COMMAND(ToggleFollowSelectedObject, "Follow Selected Object", "Start or Stop following the selected object", EUserInterfaceActionType::ToggleButton, FInputChord(EKeys::F8));
+	UI_COMMAND(OverridePlaybackFrameRate, "Override Recorded Framerate", "When enabled, allows to playback the recording at a fixed framerate", EUserInterfaceActionType::ToggleButton, FInputChord(EModifierKey::Control, EKeys::R));
 }
 
 #undef LOCTEXT_NAMESPACE
diff --git a/Engine/Plugins/ChaosVD/Source/ChaosVD/Private/ChaosVDCommands.h b/Engine/Plugins/ChaosVD/Source/ChaosVD/Private/ChaosVDCommands.h
--- a/Engine/Plugins/ChaosVD/Source/ChaosVD/Private/ChaosVDCommands.h
+++ b/Engine/Plugins/ChaosVD/Source/ChaosVD/Private/ChaosVDCommands.h
@@ -15,5 +15,6 @@ public:
 	virtual void RegisterCommands() override;
 
 public:
-	TSharedPtr<FUICommandInfo> TrackUntrackSelectedObject;
+	TSharedPtr<FUICommandInfo> ToggleFollowSelectedObject;
+	TSharedPtr<FUICommandInfo> OverridePlaybackFrameRate;
 };
\ No newline at end of file
diff --git a/Engine/Plugins/ChaosVD/Source/ChaosVD/Private/ChaosVDEditorSettings.cpp b/Engine/Plugins/ChaosVD/Source/ChaosVD/Private/ChaosVDEditorSettings.cpp
deleted file mode 100644
diff --git a/Engine/Plugins/ChaosVD/Source/ChaosVD/Private/ChaosVDEditorSettings.h b/Engine/Plugins/ChaosVD/Source/ChaosVD/Private/ChaosVDEditorSettings.h
deleted file mode 100644
diff --git a/Engine/Plugins/ChaosVD/Source/ChaosVD/Private/ChaosVDEditorVisualizationSettingsTab.cpp b/Engine/Plugins/ChaosVD/Source/ChaosVD/Private/ChaosVDEditorVisualizationSettingsTab.cpp
deleted file mode 100644
diff --git a/Engine/Plugins/ChaosVD/Source/ChaosVD/Private/ChaosVDEditorVisualizationSettingsTab.h b/Engine/Plugins/ChaosVD/Source/ChaosVD/Private/ChaosVDEditorVisualizationSettingsTab.h
deleted file mode 100644
diff --git a/Engine/Plugins/ChaosVD/Source/ChaosVD/Private/ChaosVDGeometryDataComponent.cpp b/Engine/Plugins/ChaosVD/Source/ChaosVD/Private/ChaosVDGeometryDataComponent.cpp
--- a/Engine/Plugins/ChaosVD/Source/ChaosVD/Private/ChaosVDGeometryDataComponent.cpp
+++ b/Engine/Plugins/ChaosVD/Source/ChaosVD/Private/ChaosVDGeometryDataComponent.cpp
@@ -2,7 +2,7 @@
 
 #include "ChaosVDGeometryDataComponent.h"
 
-#include "ChaosVDEditorSettings.h"
+#include "Settings/ChaosVDCoreSettings.h"
 #include "ChaosVDGeometryBuilder.h"
 #include "ChaosVDModule.h"
 #include "ChaosVDParticleActor.h"
@@ -10,6 +10,8 @@
 #include "Materials/Material.h"
 #include "Materials/MaterialInterface.h"
 #include "Materials/MaterialInstanceDynamic.h"
+#include "Visualizers/ChaosVDParticleDataComponentVisualizer.h"
+
 
 FChaosVDMeshDataInstanceHandle::FChaosVDMeshDataInstanceHandle(int32 InInstanceIndex, UMeshComponent* InMeshComponent, int32 InParticleID, int32 InSolverID)
 {
@@ -269,9 +271,9 @@ void FChaosVDGeometryComponentUtils::UpdateMeshVisibility(const TSharedPtr<FChao
 		return;
 	}
 
-	if (const UChaosVDEditorSettings* EditorSettings = GetDefault<UChaosVDEditorSettings>())
+	if (const UChaosVDParticleVisualizationSettings* ParticleVisualizationSettings = GetDefault<UChaosVDParticleVisualizationSettings>())
 	{
-		const EChaosVDGeometryVisibilityFlags CurrentVisibilityFlags = static_cast<EChaosVDGeometryVisibilityFlags>(EditorSettings->GeometryVisibilityFlags);
+		const EChaosVDGeometryVisibilityFlags CurrentVisibilityFlags = ParticleVisualizationSettings->GetGeometryVisualizationFlags();
 		
 		bool bShouldGeometryBeVisible = false;
 
@@ -325,28 +327,28 @@ FLinearColor FChaosVDGeometryComponentUtils::GetGeometryParticleColor(const TSha
 		return ColorToApply;
 	}
 
-	const UChaosVDEditorSettings* EditorSettings = GetDefault<UChaosVDEditorSettings>();
-	if (!EditorSettings)
+	const UChaosVDParticleVisualizationColorSettings* VisualizationSettings = GetDefault<UChaosVDParticleVisualizationColorSettings>();
+	if (!VisualizationSettings)
 	{
 		return ColorToApply;
 	}
 
-	switch (EditorSettings->ParticleColorMode)
+	switch (VisualizationSettings->ParticleColorMode)
 	{
 	case EChaosVDParticleDebugColorMode::ShapeType:
 		{
-			ColorToApply = InGeometryHandle->GetImplicitObject() ? EditorSettings->ColorsByShapeType.GetColorFromShapeType(Chaos::GetInnerType(InGeometryHandle->GetImplicitObject()->GetType())) : DefaultColor;
+			ColorToApply = InGeometryHandle->GetImplicitObject() ? VisualizationSettings->ColorsByShapeType.GetColorFromShapeType(Chaos::GetInnerType(InGeometryHandle->GetImplicitObject()->GetType())) : DefaultColor;
 			break;
 		}
 	case EChaosVDParticleDebugColorMode::State:
 		{
 			if (InParticleData.Type == EChaosVDParticleType::Static)
 			{
-				ColorToApply = EditorSettings->ColorsByParticleState.GetColorFromState(EChaosVDObjectStateType::Static);
+				ColorToApply = VisualizationSettings->ColorsByParticleState.GetColorFromState(EChaosVDObjectStateType::Static);
 			}
 			else
 			{
-				ColorToApply = EditorSettings->ColorsByParticleState.GetColorFromState(InParticleData.ParticleDynamicsMisc.MObjectState);
+				ColorToApply = VisualizationSettings->ColorsByParticleState.GetColorFromState(InParticleData.ParticleDynamicsMisc.MObjectState);
 			}
 			break;
 		}
@@ -354,11 +356,11 @@ FLinearColor FChaosVDGeometryComponentUtils::GetGeometryParticleColor(const TSha
 		{
 			if (InParticleData.Type == EChaosVDParticleType::Static)
 			{
-				ColorToApply = EditorSettings->ColorsByClientServer.GetColorFromState(bIsServer, EChaosVDObjectStateType::Static);
+				ColorToApply = VisualizationSettings->ColorsByClientServer.GetColorFromState(bIsServer, EChaosVDObjectStateType::Static);
 			}
 			else
 			{
-				ColorToApply = EditorSettings->ColorsByClientServer.GetColorFromState(bIsServer, InParticleData.ParticleDynamicsMisc.MObjectState);
+				ColorToApply = VisualizationSettings->ColorsByClientServer.GetColorFromState(bIsServer, InParticleData.ParticleDynamicsMisc.MObjectState);
 			}
 			break;
 		}
@@ -374,7 +376,7 @@ FLinearColor FChaosVDGeometryComponentUtils::GetGeometryParticleColor(const TSha
 
 UMaterialInterface* FChaosVDGeometryComponentUtils::GetBaseMaterialForType(EChaosVDMaterialType Type)
 {
-	const UChaosVDEditorSettings* EditorSettings = GetDefault<UChaosVDEditorSettings>();
+	const UChaosVDCoreSettings* EditorSettings = GetDefault<UChaosVDCoreSettings>();
 	if (!EditorSettings)
 	{
 		return nullptr;
diff --git a/Engine/Plugins/ChaosVD/Source/ChaosVD/Private/ChaosVDGeometryDataComponent.h b/Engine/Plugins/ChaosVD/Source/ChaosVD/Private/ChaosVDGeometryDataComponent.h
--- a/Engine/Plugins/ChaosVD/Source/ChaosVD/Private/ChaosVDGeometryDataComponent.h
+++ b/Engine/Plugins/ChaosVD/Source/ChaosVD/Private/ChaosVDGeometryDataComponent.h
@@ -1,8 +1,6 @@
 // Copyright Epic Games, Inc. All Rights Reserved.
 
 #pragma once
-
-#include "ChaosVDEditorSettings.h"
 #include "ChaosVDExtractedGeometryDataHandle.h"
 #include "InstancedStaticMeshDelegates.h"
 #include "Chaos/ImplicitFwd.h"
@@ -16,17 +14,18 @@
 #include "Templates/SharedPointer.h"
 #include "Components/MeshComponent.h"
 
-class FChaosVDScene;
-
 #include "ChaosVDGeometryDataComponent.generated.h"
 
-class UChaosVDStaticMeshComponent;
 class FChaosVDGeometryBuilder;
-class UChaosVDInstancedStaticMeshComponent;
-struct FChaosVDExtractedGeometryDataHandle;
+class FChaosVDScene;
 class IChaosVDGeometryComponent;
+class UChaosVDInstancedStaticMeshComponent;
+class UChaosVDStaticMeshComponent;
 class UMaterialInstanceDynamic;
 
+struct FChaosVDExtractedGeometryDataHandle;
+struct FChaosVDParticleDataWrapper;
+
 namespace Chaos
 {
 	class FImplicitObject;
diff --git a/Engine/Plugins/ChaosVD/Source/ChaosVD/Private/ChaosVDParticleActor.cpp b/Engine/Plugins/ChaosVD/Source/ChaosVD/Private/ChaosVDParticleActor.cpp
--- a/Engine/Plugins/ChaosVD/Source/ChaosVD/Private/ChaosVDParticleActor.cpp
+++ b/Engine/Plugins/ChaosVD/Source/ChaosVD/Private/ChaosVDParticleActor.cpp
@@ -7,6 +7,7 @@
 #include "ChaosVDScene.h"
 #include "Actors/ChaosVDSolverInfoActor.h"
 #include "Components/ChaosVDInstancedStaticMeshComponent.h"
+#include "Components/ChaosVDSolverCharacterGroundConstraintDataComponent.h"
 #include "Components/ChaosVDSolverCollisionDataComponent.h"
 #include "Components/ChaosVDStaticMeshComponent.h"
 #include "Components/MeshComponent.h"
@@ -14,6 +15,7 @@
 #include "Components/StaticMeshComponent.h"
 #include "DataWrappers/ChaosVDParticleDataWrapper.h"
 #include "Engine/StaticMesh.h"
+#include "Visualizers/ChaosVDSolverCollisionDataComponentVisualizer.h"
 
 namespace Chaos::VisualDebugger::Cvars
 {
diff --git a/Engine/Plugins/ChaosVD/Source/ChaosVD/Private/ChaosVDPlaybackController.cpp b/Engine/Plugins/ChaosVD/Source/ChaosVD/Private/ChaosVDPlaybackController.cpp
--- a/Engine/Plugins/ChaosVD/Source/ChaosVD/Private/ChaosVDPlaybackController.cpp
+++ b/Engine/Plugins/ChaosVD/Source/ChaosVD/Private/ChaosVDPlaybackController.cpp
@@ -2,7 +2,6 @@
 
 #include "ChaosVDPlaybackController.h"
 
-#include "ChaosVDEditorSettings.h"
 #include "ChaosVDModule.h"
 #include "ChaosVDPlaybackControllerInstigator.h"
 #include "ChaosVDRecording.h"
@@ -31,12 +30,6 @@ FChaosVDPlaybackController::FChaosVDPlaybackController(const TWeakPtr<FChaosVDSc
 
 	RecordingStoppedHandle = FChaosVDRuntimeModule::Get().RegisterRecordingStopCallback(FChaosVDRecordingStateChangedDelegate::FDelegate::CreateRaw(this, &FChaosVDPlaybackController::HandleDisconnectedFromSession));
 
-	if (UChaosVDEditorSettings* Settings = GetMutableDefault<UChaosVDEditorSettings>())
-	{
-		Settings->OnPlaybackSettingsChanged().AddRaw(this, &FChaosVDPlaybackController::HandleFrameRateOverrideSettingsChanged);
-		HandleFrameRateOverrideSettingsChanged(Settings);
-	}
-
 	CurrentPlaybackInstigator = IChaosVDPlaybackControllerInstigator::InvalidGuid;
 }
 
@@ -47,11 +40,6 @@ FChaosVDPlaybackController::~FChaosVDPlaybackController()
 	{
 		FChaosVDRuntimeModule::Get().RemoveRecordingStopCallback(RecordingStoppedHandle);
 	}
-	
-	if (UChaosVDEditorSettings* Settings = GetMutableDefault<UChaosVDEditorSettings>())
-	{
-		Settings->OnPlaybackSettingsChanged().RemoveAll(this);
-	}
 
 	UnloadCurrentRecording(EChaosVDUnloadRecordingFlags::Silent);
 }
@@ -178,16 +166,6 @@ void FChaosVDPlaybackController::EnqueueGeometryDataUpdate(const Chaos::FConstIm
 	GeometryDataUpdateGTQueue.Enqueue({NewGeometry, GeometryID });
 }
 
-void FChaosVDPlaybackController::HandleFrameRateOverrideSettingsChanged(UChaosVDEditorSettings* CVDSettings)
-{
-	if (!CVDSettings)
-	{
-		return;
-	}
-
-	CurrentFrameRateOverride = CVDSettings->bPlaybackAtRecordedFrameRate ? InvalidFrameRateOverride : CVDSettings->TargetFrameRateOverride;
-}
-
 void FChaosVDPlaybackController::PlaySolverStepData(int32 TrackID, const TSharedRef<FChaosVDScene>& InSceneToControlSharedPtr, const FChaosVDSolverFrameData& InSolverFrameData, int32 StepIndex)
 {
 	if (InSolverFrameData.SolverSteps.IsValidIndex(StepIndex))
@@ -755,11 +733,23 @@ float FChaosVDPlaybackController::GetFrameTimeOverride() const
 	return CurrentFrameRateOverride >= MinimumFrameRateOverride ? 1.0f / static_cast<float>(CurrentFrameRateOverride) : InvalidFrameRateOverride;
 }
 
+int32 FChaosVDPlaybackController::GetFrameRateOverride() const
+{
+	constexpr int32 MinimumFrameRateOverride = 1;
+	return CurrentFrameRateOverride >= MinimumFrameRateOverride ? CurrentFrameRateOverride : InvalidFrameRateOverride;
+}
+
+void FChaosVDPlaybackController::SetFrameRateOverride(float NewFrameRateOverride)
+{
+	constexpr int32 MinimumFrameRateOverride = 1;
+	CurrentFrameRateOverride = NewFrameRateOverride >= MinimumFrameRateOverride ? NewFrameRateOverride : InvalidFrameRateOverride;
+}
+
 float FChaosVDPlaybackController::GetFrameTimeForTrack(EChaosVDTrackType TrackType, int32 TrackID, const FChaosVDTrackInfo& TrackInfo) const
 {
 	const float TargetFrameTimeOverride = GetFrameTimeOverride();
-	const bool bHastFrameRateOverride = !FMath::IsNearlyEqual(TargetFrameTimeOverride, FChaosVDPlaybackController::InvalidFrameRateOverride);
-	if (bHastFrameRateOverride)
+	const bool bHasValidFrameRateOverride = bUseFrameRateOverride && !FMath::IsNearlyEqual(TargetFrameTimeOverride, FChaosVDPlaybackController::InvalidFrameRateOverride);
+	if (bHasValidFrameRateOverride)
 	{
 		return TargetFrameTimeOverride;
 	}
diff --git a/Engine/Plugins/ChaosVD/Source/ChaosVD/Private/ChaosVDPlaybackController.h b/Engine/Plugins/ChaosVD/Source/ChaosVD/Private/ChaosVDPlaybackController.h
--- a/Engine/Plugins/ChaosVD/Source/ChaosVD/Private/ChaosVDPlaybackController.h
+++ b/Engine/Plugins/ChaosVD/Source/ChaosVD/Private/ChaosVDPlaybackController.h
@@ -11,7 +11,7 @@
 #include "HAL/ThreadSafeBool.h"
 
 class IChaosVDPlaybackControllerInstigator;
-class UChaosVDEditorSettings;
+class UChaosVDCoreSettings;
 struct FChaosVDTraceSessionDescriptor;
 struct FChaosVDTrackInfo;
 class FChaosVDScene;
@@ -235,7 +235,14 @@ public:
 	bool AcquireExclusivePlaybackControls(const IChaosVDPlaybackControllerInstigator& InPlaybackInstigator);
 	bool ReleaseExclusivePlaybackControls(const IChaosVDPlaybackControllerInstigator& InPlaybackInstigator);
 
+	bool IsUsingFrameRateOverride() const { return bUseFrameRateOverride; }
+
+	bool ToggleUseFrameRateOverride() { return bUseFrameRateOverride =  !bUseFrameRateOverride; }
+
 	float GetFrameTimeOverride() const;
+	int32 GetFrameRateOverride() const;
+	void SetFrameRateOverride(float NewFrameRateOverride);
+
 	float GetFrameTimeForTrack(EChaosVDTrackType TrackType, int32 TrackID, const FChaosVDTrackInfo& TrackInfo) const;
 
 	void UpdateTrackVisibility(EChaosVDTrackType Type, int32 TrackID, bool bNewVisibility);
@@ -263,8 +270,6 @@ protected:
 	/** Add the provided Geometry info data to the queue. The update will be broadcast in the game thread */
 	void EnqueueGeometryDataUpdate(const Chaos::FConstImplicitObjectPtr& NewGeometry, const uint32 GeometryID);
 
-	void HandleFrameRateOverrideSettingsChanged(UChaosVDEditorSettings* CVDSettings);
-
 	void PlaySolverStepData(int32 TrackID, const TSharedRef<FChaosVDScene>& InSceneToControlSharedPtr, const FChaosVDSolverFrameData& InSolverFrameData, int32 StepIndex);
 
 	/** Map containing all track info, by track type*/
@@ -299,7 +304,9 @@ protected:
 	int32 MaxFramesLaggingBehindDuringLiveSession = 50;
 	int32 MinFramesLaggingBehindDuringLiveSession = 5;
 
-	int32 CurrentFrameRateOverride = InvalidFrameRateOverride;
+	int32 CurrentFrameRateOverride = 60;
+
+	bool bUseFrameRateOverride = false;
 
 	bool bPauseRequested = false;
 
diff --git a/Engine/Plugins/ChaosVD/Source/ChaosVD/Private/ChaosVDPlaybackViewportClient.cpp b/Engine/Plugins/ChaosVD/Source/ChaosVD/Private/ChaosVDPlaybackViewportClient.cpp
--- a/Engine/Plugins/ChaosVD/Source/ChaosVD/Private/ChaosVDPlaybackViewportClient.cpp
+++ b/Engine/Plugins/ChaosVD/Source/ChaosVD/Private/ChaosVDPlaybackViewportClient.cpp
@@ -2,19 +2,12 @@
 
 #include "ChaosVDPlaybackViewportClient.h"
 
-#include "ChaosVDEditorSettings.h"
-#include "ChaosVDEngine.h"
-#include "ChaosVDModule.h"
-#include "ChaosVDObjectDetailsTab.h"
 #include "ChaosVDParticleActor.h"
-#include "ChaosVDPlaybackController.h"
 #include "ChaosVDScene.h"
 #include "ChaosVDSkySphereInterface.h"
-#include "ChaosVDTabsIDs.h"
 #include "Components/ChaosVDSceneQueryDataComponent.h"
 #include "ComponentVisualizer.h"
 #include "EditorModeManager.h"
-#include "Elements/Framework/TypedElementSelectionSet.h"
 #include "Engine/DirectionalLight.h"
 #include "EngineUtils.h"
 #include "SceneView.h"
@@ -22,12 +15,8 @@
 #include "Selection.h"
 #include "UnrealWidget.h"
 #include "Actors/ChaosVDSolverInfoActor.h"
-#include "Components/ChaosVDParticleDataComponent.h"
-#include "Components/ChaosVDSolverCollisionDataComponent.h"
-#include "Components/ChaosVDSolverJointConstraintDataComponent.h"
 #include "Components/InstancedStaticMeshComponent.h"
 #include "Visualizers/ChaosVDDebugDrawUtils.h"
-#include "Widgets/SChaosVDDetailsView.h"
 #include "Widgets/SChaosVDMainTab.h"
 
 FChaosVDPlaybackViewportClient::FChaosVDPlaybackViewportClient(const TSharedPtr<FEditorModeTools>& InModeTools, const TSharedPtr<SEditorViewport>& InEditorViewportWidget) : FEditorViewportClient(InModeTools.Get(), nullptr, InEditorViewportWidget), CVDWorld(nullptr)
@@ -39,14 +28,8 @@ FChaosVDPlaybackViewportClient::FChaosVDPlaybackViewportClient(const TSharedPtr<
 		GEngine->OnActorMoving().AddRaw(this, &FChaosVDPlaybackViewportClient::HandleActorMoving);
 	}
 
-	if (UChaosVDEditorSettings* Settings = GetMutableDefault<UChaosVDEditorSettings>())
-	{
-		Settings->OnFarClippingOverrideChanged().AddRaw(this, &FChaosVDPlaybackViewportClient::HandleViewportSettingsChanged);
-		Settings->OnVisibilitySettingsChanged().AddRaw(this, &FChaosVDPlaybackViewportClient::HandleViewportSettingsChanged);
-		Settings->OnColorSettingsChanged().AddRaw(this, &FChaosVDPlaybackViewportClient::HandleViewportSettingsChanged);
-
-		HandleViewportSettingsChanged(Settings);
-	}
+	constexpr float DefaultFarClipPlaneOverride = 20000.0f;
+	OverrideFarClipPlane(DefaultFarClipPlaneOverride);
 }
 
 FChaosVDPlaybackViewportClient::~FChaosVDPlaybackViewportClient()
@@ -63,13 +46,6 @@ FChaosVDPlaybackViewportClient::~FChaosVDPlaybackViewportClient()
 	{
 		GEngine->OnActorMoving().RemoveAll(this);
 	}
-	
-	if (UChaosVDEditorSettings* Settings = GetMutableDefault<UChaosVDEditorSettings>())
-	{
-		Settings->OnFarClippingOverrideChanged().RemoveAll(this);
-		Settings->OnVisibilitySettingsChanged().RemoveAll(this);
-		Settings->OnColorSettingsChanged().RemoveAll(this);
-	}
 }
 
 void FChaosVDPlaybackViewportClient::ProcessClick(FSceneView& View, HHitProxy* HitProxy, FKey Key, EInputEvent Event, uint32 HitX, uint32 HitY)
@@ -190,33 +166,22 @@ void FChaosVDPlaybackViewportClient::HandleActorMoving(AActor* MovedActor) const
 	}
 }
 
-void FChaosVDPlaybackViewportClient::HandleViewportSettingsChanged(UChaosVDEditorSettings* SettingsObject)
+void FChaosVDPlaybackViewportClient::TrackSelectedObject()
 {
-	if (SettingsObject)
+	if (!bAutoTrackSelectedObject || !ModeTools.IsValid())
 	{
-		OverrideFarClipPlane(SettingsObject->FarClippingOverride);
-		EngineShowFlags.SetMeshEdges(EnumHasAnyFlags(static_cast<EChaosVDGeometryVisibilityFlags>(SettingsObject->GeometryVisibilityFlags), EChaosVDGeometryVisibilityFlags::ShowTriangleEdges));
-		Invalidate();
+		return;
 	}
-}
 
-void FChaosVDPlaybackViewportClient::TrackSelectedObject()
-{
 	if (const TSharedPtr<FChaosVDScene> CVDSceneSharedPtr = CVDScene.Pin())
 	{
-		if (const UChaosVDEditorSettings* CVDEditorSettings = GetDefault<UChaosVDEditorSettings>())
-		{
-			if (ModeTools.IsValid() && CVDEditorSettings->TrackingTarget == EChaosVDActorTrackingTarget::SelectedObject)
-			{
-				USelection* CurrentSelection = ModeTools->GetSelectedActors();
+		USelection* CurrentSelection = ModeTools->GetSelectedActors();
 
-				//TODO: Update this if we add multi selection support
-				if (const AActor* SelectedActor = CurrentSelection ? CurrentSelection->GetTop<AActor>() : nullptr)
-				{
-					const FBox ActorBounds = SelectedActor->GetComponentsBoundingBox(false);
-					FocusViewportOnBox(ActorBounds.ExpandBy(CVDEditorSettings->ExpandViewTrackingBy), true);		
-				}
-			}
+		//TODO: Update this if we add multi selection support
+		if (const AActor* SelectedActor = CurrentSelection ? CurrentSelection->GetTop<AActor>() : nullptr)
+		{
+			const FBox ActorBounds = SelectedActor->GetComponentsBoundingBox(false);
+			FocusViewportOnBox(ActorBounds.ExpandBy(TrackingViewDistance), true);		
 		}
 	}
 }
@@ -234,11 +199,12 @@ bool FChaosVDPlaybackViewportClient::InputKey(const FInputKeyEventArgs& EventArg
 
 void FChaosVDPlaybackViewportClient::ToggleObjectTrackingIfSelected()
 {
-	// Currently we only have two options, so toggle between them
-	if (UChaosVDEditorSettings* CVDEditorSettings = GetMutableDefault<UChaosVDEditorSettings>())
-	{
-		CVDEditorSettings->TrackingTarget = CVDEditorSettings->TrackingTarget == EChaosVDActorTrackingTarget::Disabled ? EChaosVDActorTrackingTarget::SelectedObject : EChaosVDActorTrackingTarget::Disabled;
-	}
+	bAutoTrackSelectedObject = !bAutoTrackSelectedObject;
+}
+
+void FChaosVDPlaybackViewportClient::SetAutoTrackingViewDistance(float NewDistance)
+{
+	TrackingViewDistance = NewDistance;
 }
 
 void FChaosVDPlaybackViewportClient::GoToLocation(const FVector& InLocation)
diff --git a/Engine/Plugins/ChaosVD/Source/ChaosVD/Private/ChaosVDPlaybackViewportClient.h b/Engine/Plugins/ChaosVD/Source/ChaosVD/Private/ChaosVDPlaybackViewportClient.h
--- a/Engine/Plugins/ChaosVD/Source/ChaosVD/Private/ChaosVDPlaybackViewportClient.h
+++ b/Engine/Plugins/ChaosVD/Source/ChaosVD/Private/ChaosVDPlaybackViewportClient.h
@@ -1,9 +1,9 @@
 ﻿// Copyright Epic Games, Inc. All Rights Reserved.
 
 #pragma once
-#include "LevelEditorViewport.h"
+#include "EditorViewportClient.h"
 
-class UChaosVDEditorSettings;
+class UChaosVDCoreSettings;
 struct FChaosVDGameFrameData;
 class FChaosVDScene;
 enum class EChaosVDActorTrackingMode;
@@ -30,6 +30,9 @@ public:
 	virtual bool InputKey(const FInputKeyEventArgs& EventArgs) override;
 
 	void ToggleObjectTrackingIfSelected();
+	bool IsAutoTrackingSelectedObject() const { return bAutoTrackSelectedObject; };
+	void SetAutoTrackingViewDistance(float NewDistance);
+	float GetAutoTrackingViewDistance() const { return TrackingViewDistance; }
 
 	void GoToLocation(const FVector& InLocation);
 
@@ -39,9 +42,12 @@ private:
 
 	void HandleObjectFocused(UObject* FocusedObject);
 	void HandleActorMoving(AActor* MovedActor) const;
-	void HandleViewportSettingsChanged(UChaosVDEditorSettings* SettingsObject);
 
 	FDelegateHandle ObjectFocusedDelegateHandle;
 	UWorld* CVDWorld;
 	TWeakPtr<FChaosVDScene> CVDScene;
+
+	bool bAutoTrackSelectedObject = false;
+	float TrackingViewDistance = 120.0f;
+
 };
diff --git a/Engine/Plugins/ChaosVD/Source/ChaosVD/Private/ChaosVDScene.cpp b/Engine/Plugins/ChaosVD/Source/ChaosVD/Private/ChaosVDScene.cpp
--- a/Engine/Plugins/ChaosVD/Source/ChaosVD/Private/ChaosVDScene.cpp
+++ b/Engine/Plugins/ChaosVD/Source/ChaosVD/Private/ChaosVDScene.cpp
@@ -4,7 +4,6 @@
 
 #include "Actors/ChaosVDSceneQueryDataContainer.h"
 #include "Actors/ChaosVDSolverInfoActor.h"
-#include "ChaosVDEditorSettings.h"
 #include "ChaosVDGeometryBuilder.h"
 #include "ChaosVDModule.h"
 #include "ChaosVDParticleActor.h"
@@ -13,27 +12,27 @@
 #include "ChaosVDSelectionCustomization.h"
 #include "ChaosVDSkySphereInterface.h"
 #include "Components/ChaosVDSceneQueryDataComponent.h"
+#include "Components/ChaosVDSolverCharacterGroundConstraintDataComponent.h"
 #include "Components/ChaosVDSolverCollisionDataComponent.h"
+#include "Components/ChaosVDSolverJointConstraintDataComponent.h"
 #include "DataWrappers/ChaosVDParticleDataWrapper.h"
-#include "EditorActorFolders.h"
-#include "EditorLevelUtils.h"
 #include "Engine/DirectionalLight.h"
 #include "Engine/Engine.h"
+#include "Engine/Level.h"
 #include "Engine/StreamableManager.h"
 #include "Engine/World.h"
+#include "Elements/Actor/ActorElementData.h"
+#include "Elements/Component/ComponentElementData.h"
 #include "Elements/Framework/EngineElementsLibrary.h"
 #include "Elements/Framework/TypedElementSelectionSet.h"
-#include "Misc/ScopedSlowTask.h"
+#include "Elements/Object/ObjectElementData.h"
 #include "Materials/Material.h"
+#include "Misc/ScopedSlowTask.h"
 #include "Selection.h"
+#include "Settings/ChaosVDParticleVisualizationSettings.h"
+#include "Settings/ChaosVDCoreSettings.h"
 #include "UObject/Package.h"
 #include "WorldPersistentFolders.h"
-#include "Components/ChaosVDSolverCharacterGroundConstraintDataComponent.h"
-#include "Components/ChaosVDSolverJointConstraintDataComponent.h"
-#include "Elements/Actor/ActorElementData.h"
-#include "Elements/Component/ComponentElementData.h"
-#include "Elements/Object/ObjectElementData.h"
-#include "Engine/Level.h"
 
 #define LOCTEXT_NAMESPACE "ChaosVisualDebugger"
 
@@ -65,7 +64,7 @@ void FChaosVDScene::Initialize()
 	
 	StreamableManager = MakeShared<FStreamableManager>();
 
-	if (UChaosVDEditorSettings* Settings = GetMutableDefault<UChaosVDEditorSettings>())
+	if (UChaosVDCoreSettings* Settings = GetMutableDefault<UChaosVDCoreSettings>())
 	{
 		// TODO: Do an async load instead, and prepare a loading screen or notification popup
 		// Jira for tracking UE-191639
@@ -73,9 +72,6 @@ void FChaosVDScene::Initialize()
 		StreamableManager->RequestSyncLoad(Settings->SimOnlyMeshesMaterial.ToSoftObjectPath());
 		StreamableManager->RequestSyncLoad(Settings->InstancedMeshesMaterial.ToSoftObjectPath());
 		StreamableManager->RequestSyncLoad(Settings->InstancedMeshesQueryOnlyMaterial.ToSoftObjectPath());
-		
-		Settings->OnVisibilitySettingsChanged().AddRaw(this, &FChaosVDScene::HandleVisibilitySettingsChanged);
-		Settings->OnColorSettingsChanged().AddRaw(this, &FChaosVDScene::HandleColorSettingsChanged);
 	}
 
 	bIsInitialized = true;
@@ -92,12 +88,6 @@ void FChaosVDScene::DeInitialize()
 		return;
 	}
 
-	if (UChaosVDEditorSettings* Settings = GetMutableDefault<UChaosVDEditorSettings>())
-	{
-		Settings->OnVisibilitySettingsChanged().RemoveAll(this);
-		Settings->OnColorSettingsChanged().RemoveAll(this);
-	}
-
 	CleanUpScene();
 
 	DeInitializeSelectionSets();
@@ -470,7 +460,7 @@ void FChaosVDScene::CreateBaseLights(UWorld* TargetWorld) const
 
 	const FVector SpawnPosition(0.0, 0.0, 2000.0);
 	
-	if (const UChaosVDEditorSettings* Settings = GetDefault<UChaosVDEditorSettings>())
+	if (const UChaosVDCoreSettings* Settings = GetDefault<UChaosVDCoreSettings>())
 	{
 		if (ADirectionalLight* DirectionalLightActor = TargetWorld->SpawnActor<ADirectionalLight>())
 		{
@@ -616,28 +606,6 @@ void FChaosVDScene::ClearSelectionAndNotify()
 	SelectionSet->NotifyPendingChanges();
 }
 
-void FChaosVDScene::HandleVisibilitySettingsChanged(UChaosVDEditorSettings* SettingsObject)
-{
-	for (const TPair<int32, AChaosVDSolverInfoActor*>& SolverDataInfoWithID : SolverDataContainerBySolverID)
-	{
-		if (AChaosVDSolverInfoActor* SolverDataInfo = SolverDataInfoWithID.Value)
-		{
-			SolverDataInfo->HandleVisibilitySettingsUpdated();
-		}
-	}
-}
-
-void FChaosVDScene::HandleColorSettingsChanged(UChaosVDEditorSettings* SettingsObject)
-{
-	for (const TPair<int32, AChaosVDSolverInfoActor*>& SolverDataInfoWithID : SolverDataContainerBySolverID)
-	{
-		if (AChaosVDSolverInfoActor* SolverDataInfo = SolverDataInfoWithID.Value)
-		{
-			SolverDataInfo->HandleColorsSettingsUpdated();
-		}
-	}
-}
-
 void FChaosVDScene::InitializeSelectionSets()
 {
 	SelectionSet = NewObject<UTypedElementSelectionSet>(GetTransientPackage(), NAME_None, RF_Transactional);
diff --git a/Engine/Plugins/ChaosVD/Source/ChaosVD/Private/ChaosVDScene.h b/Engine/Plugins/ChaosVD/Source/ChaosVD/Private/ChaosVDScene.h
--- a/Engine/Plugins/ChaosVD/Source/ChaosVD/Private/ChaosVDScene.h
+++ b/Engine/Plugins/ChaosVD/Source/ChaosVD/Private/ChaosVDScene.h
@@ -17,7 +17,7 @@ struct FTypedElementSelectionOptions;
 class AChaosVDSceneQueryDataContainer;
 class AChaosVDSolverInfoActor;
 class AChaosVDSceneCollisionContainer;
-class UChaosVDEditorSettings;
+class UChaosVDCoreSettings;
 class FChaosVDGeometryBuilder;
 class AChaosVDParticleActor;
 class FReferenceCollector;
@@ -160,9 +160,6 @@ private:
 
 	void ClearSelectionAndNotify();
 
-	void HandleVisibilitySettingsChanged(UChaosVDEditorSettings* SettingsObject);
-	void HandleColorSettingsChanged(UChaosVDEditorSettings* SettingsObject);
-
 	void InitializeSelectionSets();
 	void DeInitializeSelectionSets();
 
diff --git a/Engine/Plugins/ChaosVD/Source/ChaosVD/Private/ChaosVDTabsIDs.cpp b/Engine/Plugins/ChaosVD/Source/ChaosVD/Private/ChaosVDTabsIDs.cpp
--- a/Engine/Plugins/ChaosVD/Source/ChaosVD/Private/ChaosVDTabsIDs.cpp
+++ b/Engine/Plugins/ChaosVD/Source/ChaosVD/Private/ChaosVDTabsIDs.cpp
@@ -9,7 +9,6 @@ const FName FChaosVDTabID::WorldOutliner = TEXT("WorldOutliner");
 const FName FChaosVDTabID::DetailsPanel = TEXT("DetailsPanel");
 const FName FChaosVDTabID::OutputLog = TEXT("ChaosVDOutputLog");
 const FName FChaosVDTabID::SolversTrack = TEXT("SolversTrack");
-const FName FChaosVDTabID::CVDEditorSettings = TEXT("CVDEditorSettings");
 const FName FChaosVDTabID::StatusBar = TEXT("StatusBar");
 const FName FChaosVDTabID::CollisionDataDetails = TEXT("CollisionDataDetails");
 const FName FChaosVDTabID::SceneQueryDataDetails = TEXT("SceneQueryDataDetails");
diff --git a/Engine/Plugins/ChaosVD/Source/ChaosVD/Private/ChaosVDTabsIDs.h b/Engine/Plugins/ChaosVD/Source/ChaosVD/Private/ChaosVDTabsIDs.h
--- a/Engine/Plugins/ChaosVD/Source/ChaosVD/Private/ChaosVDTabsIDs.h
+++ b/Engine/Plugins/ChaosVD/Source/ChaosVD/Private/ChaosVDTabsIDs.h
@@ -14,7 +14,6 @@ public:
 	static const FName DetailsPanel;
 	static const FName OutputLog;
 	static const FName SolversTrack;
-	static const FName CVDEditorSettings;
 	static const FName StatusBar;
 	static const FName CollisionDataDetails;
 	static const FName SceneQueryDataDetails;
diff --git a/Engine/Plugins/ChaosVD/Source/ChaosVD/Private/Settings/ChaosVDCharacterConstraintsVisualizationSettings.cpp b/Engine/Plugins/ChaosVD/Source/ChaosVD/Private/Settings/ChaosVDCharacterConstraintsVisualizationSettings.cpp
new file mode 100644
--- /dev/null
+++ b/Engine/Plugins/ChaosVD/Source/ChaosVD/Private/Settings/ChaosVDCharacterConstraintsVisualizationSettings.cpp
@@ -0,0 +1,22 @@
+// Copyright Epic Games, Inc. All Rights Reserved.
+
+#include "ChaosVDCharacterConstraintsVisualizationSettings.h"
+
+void UChaosVDCharacterConstraintsVisualizationSettings::SetCharacterGroundConstraintDataVisualizationFlags(EChaosVDCharacterGroundConstraintDataVisualizationFlags NewFlags)
+{
+	if (UChaosVDCharacterConstraintsVisualizationSettings* Settings = GetMutableDefault<UChaosVDCharacterConstraintsVisualizationSettings>())
+	{
+		Settings->GlobalCharacterGroundConstraintDataVisualizationFlags = NewFlags;
+		Settings->BroadcastSettingsChanged(Settings);
+	}
+}
+
+EChaosVDCharacterGroundConstraintDataVisualizationFlags UChaosVDCharacterConstraintsVisualizationSettings::GetCharacterGroundConstraintDataVisualizationFlags()
+{
+	if (UChaosVDCharacterConstraintsVisualizationSettings* Settings = GetMutableDefault<UChaosVDCharacterConstraintsVisualizationSettings>())
+	{
+		return Settings->GlobalCharacterGroundConstraintDataVisualizationFlags;
+	}
+
+	return EChaosVDCharacterGroundConstraintDataVisualizationFlags::None;
+}
diff --git a/Engine/Plugins/ChaosVD/Source/ChaosVD/Private/Settings/ChaosVDCharacterConstraintsVisualizationSettings.h b/Engine/Plugins/ChaosVD/Source/ChaosVD/Private/Settings/ChaosVDCharacterConstraintsVisualizationSettings.h
new file mode 100644
--- /dev/null
+++ b/Engine/Plugins/ChaosVD/Source/ChaosVD/Private/Settings/ChaosVDCharacterConstraintsVisualizationSettings.h
@@ -0,0 +1,79 @@
+// Copyright Epic Games, Inc. All Rights Reserved.
+
+#pragma once
+
+#include "ChaosVDCoreSettings.h"
+#include "UObject/Object.h"
+#include "ChaosVDCharacterConstraintsVisualizationSettings.generated.h"
+
+UENUM(meta = (Bitflags, UseEnumValuesAsMaskValuesInEditor = "true"))
+enum class EChaosVDCharacterGroundConstraintDataVisualizationFlags : uint32
+{
+	None					= 0 UMETA(Hidden),
+	/** Draw the target movement vector */
+	TargetDeltaPosition		= 1 << 0,
+	/** Draw the target orientation facing vector */
+	TargetDeltaFacing		= 1 << 1,
+	/** Draw the ground query distance based on the constraint's data */
+	GroundQueryDistance		= 1 << 2,
+	/** Draw the ground query normal based on the constraint's data */
+	GroundQueryNormal		= 1 << 3,
+	/** Draw the applied force vector */
+	AppliedRadialForce		= 1 << 4,
+	/** Draw the applied force vector */
+	AppliedNormalForce		= 1 << 5,
+	/** Draw the applied force vector */
+	AppliedTorque			= 1 << 6,
+	/** Draw the constraint even if it is disabled */
+	DrawDisabled			= 1 << 7,
+	/** Only debugs draw data for a selected constraint */
+	OnlyDrawSelected		= 1 << 8,
+	/** Enables debug draw for constraint data from any solver */
+	EnableDraw				= 1 << 9,
+};
+ENUM_CLASS_FLAGS(EChaosVDCharacterGroundConstraintDataVisualizationFlags);
+
+UCLASS(config=ChaosVD)
+class UChaosVDCharacterConstraintsVisualizationSettings : public UChaosVDVisualizationSettingsObjectBase
+{
+	GENERATED_BODY()
+
+public:
+
+	/** If true, any debug draw text available will be drawn */
+	UPROPERTY(EditAnywhere, Category=DebugDraw)
+	bool bShowDebugText = false;
+
+	/** The depth priority used for while drawing data. Can be World or Foreground (with this one the shapes will be drawn on top of the geometry and be always visible) */
+	UPROPERTY(config, EditAnywhere, Category=DebugDraw)
+	TEnumAsByte<ESceneDepthPriorityGroup> DepthPriority = ESceneDepthPriorityGroup::SDPG_Foreground;
+	/** Scale to apply to the force vector before draw it. */
+	UPROPERTY(config, EditAnywhere, Category=DebugDraw)
+	float ForceScale = 0.01f;
+	/** Scale to apply to the torque vector before draw it. */
+	UPROPERTY(config, EditAnywhere, Category=DebugDraw)
+	float TorqueScale = 0.01f;
+	/** Scale to apply to anything that does not have a dedicated scale setting before draw it. */
+	UPROPERTY(config, EditAnywhere, Category=DebugDraw)
+	float GeneralScale = 1.0f;
+	/** Line thickness to use as a base to calculate the different line thickness values used to debug draw the data. */
+	UPROPERTY(config, EditAnywhere, Category=DebugDraw)
+	float BaseLineThickness = 2.0f;
+	/** Color used for normal force vector */
+	UPROPERTY(config, EditAnywhere, Category = DebugDraw)
+	FColor NormalForceColor = FColor(255, 0, 0);
+	/** Color used for radial force vector */
+	UPROPERTY(config, EditAnywhere, Category = DebugDraw)
+	FColor RadialForceColor = FColor(255, 255, 0);
+	/** Color used for torque vector */
+	UPROPERTY(config, EditAnywhere, Category = DebugDraw)
+	FColor TorqueColor = FColor(255, 0, 255);
+
+	static void SetCharacterGroundConstraintDataVisualizationFlags(EChaosVDCharacterGroundConstraintDataVisualizationFlags NewFlags);
+	static EChaosVDCharacterGroundConstraintDataVisualizationFlags GetCharacterGroundConstraintDataVisualizationFlags();
+
+private:
+	/** Set of flags to enable/disable visualization of specific character ground constraint data as debug draw */
+	EChaosVDCharacterGroundConstraintDataVisualizationFlags GlobalCharacterGroundConstraintDataVisualizationFlags = EChaosVDCharacterGroundConstraintDataVisualizationFlags::GroundQueryDistance | EChaosVDCharacterGroundConstraintDataVisualizationFlags::GroundQueryNormal | EChaosVDCharacterGroundConstraintDataVisualizationFlags::TargetDeltaPosition;
+
+};
diff --git a/Engine/Plugins/ChaosVD/Source/ChaosVD/Private/Settings/ChaosVDCollisionVisualizationSettings.cpp b/Engine/Plugins/ChaosVD/Source/ChaosVD/Private/Settings/ChaosVDCollisionVisualizationSettings.cpp
new file mode 100644
--- /dev/null
+++ b/Engine/Plugins/ChaosVD/Source/ChaosVD/Private/Settings/ChaosVDCollisionVisualizationSettings.cpp
@@ -0,0 +1,23 @@
+// Copyright Epic Games, Inc. All Rights Reserved.
+
+#include "ChaosVDCollisionVisualizationSettings.h"
+
+void UChaosVDCollisionDataVisualizationSettings::SetCollisionDataVisualizationFlags(EChaosVDCollisionVisualizationFlags NewFlags)
+{
+	if (UChaosVDCollisionDataVisualizationSettings* Settings = GetMutableDefault<UChaosVDCollisionDataVisualizationSettings>())
+	{
+		Settings->CollisionDataVisualizationFlags = NewFlags;
+		Settings->BroadcastSettingsChanged(Settings);
+	}
+}
+
+EChaosVDCollisionVisualizationFlags UChaosVDCollisionDataVisualizationSettings::GetCollisionDataVisualizationFlags()
+{
+	if (UChaosVDCollisionDataVisualizationSettings* Settings = GetMutableDefault<UChaosVDCollisionDataVisualizationSettings>())
+	{
+		return Settings->CollisionDataVisualizationFlags ;
+	}
+
+	return EChaosVDCollisionVisualizationFlags::None;
+}
+
diff --git a/Engine/Plugins/ChaosVD/Source/ChaosVD/Private/Settings/ChaosVDCollisionVisualizationSettings.h b/Engine/Plugins/ChaosVD/Source/ChaosVD/Private/Settings/ChaosVDCollisionVisualizationSettings.h
new file mode 100644
--- /dev/null
+++ b/Engine/Plugins/ChaosVD/Source/ChaosVD/Private/Settings/ChaosVDCollisionVisualizationSettings.h
@@ -0,0 +1,53 @@
+// Copyright Epic Games, Inc. All Rights Reserved.
+
+#pragma once
+#include "ChaosVDCoreSettings.h"
+
+#include "ChaosVDCollisionVisualizationSettings.generated.h"
+
+UENUM(meta = (Bitflags, UseEnumValuesAsMaskValuesInEditor = "true"))
+enum class EChaosVDCollisionVisualizationFlags: uint32
+{
+	None							= 0 UMETA(Hidden),
+	ContactPoints					= 1 << 0,
+	ContactInfo						= 1 << 1,
+	NetPushOut						= 1 << 2,
+	NetImpulse						= 1 << 3,
+	ContactNormal					= 1 << 4,
+	AccumulatedImpulse				= 1 << 5,
+	DrawInactiveContacts			= 1 << 6,
+	DrawDataOnlyForSelectedParticle	= 1 << 7,
+	EnableDraw						= 1 << 8,
+};
+ENUM_CLASS_FLAGS(EChaosVDCollisionVisualizationFlags);
+
+UCLASS(config=ChaosVD)
+class UChaosVDCollisionDataVisualizationSettings : public UChaosVDVisualizationSettingsObjectBase
+{
+	GENERATED_BODY()
+public:
+
+	/** If true, any available debug text will be drawn */
+	UPROPERTY(EditAnywhere, Category=DebugDraw)
+	bool bShowDebugText = false;
+
+	/** The depth priority used for while drawing contact data. Can be World or Foreground (with this one the shapes will be drawn on top of the geometry and be always visible) */
+	UPROPERTY(EditAnywhere, config, Category=DebugDraw)
+	TEnumAsByte<ESceneDepthPriorityGroup> DepthPriority = ESceneDepthPriorityGroup::SDPG_World;
+
+	/** The radius of the debug draw circle used to represent a contact point */
+	UPROPERTY(EditAnywhere, config, Category=DebugDraw)
+	float ContactCircleRadius = 6.0f;
+
+	/** The scale value to be applied to the normal vector of a contact used to change its size to make it easier to see */
+	UPROPERTY(EditAnywhere, config, Category=DebugDraw)
+	float ContactNormalScale = 30.0f;
+
+	static void SetCollisionDataVisualizationFlags(EChaosVDCollisionVisualizationFlags NewFlags);
+	static EChaosVDCollisionVisualizationFlags GetCollisionDataVisualizationFlags();
+
+private:
+	/** Set of flags to enable/disable visualization of specific collision data as debug draw */
+	EChaosVDCollisionVisualizationFlags CollisionDataVisualizationFlags = EChaosVDCollisionVisualizationFlags::ContactInfo | EChaosVDCollisionVisualizationFlags::ContactPoints | EChaosVDCollisionVisualizationFlags::ContactNormal;
+};
+
diff --git a/Engine/Plugins/ChaosVD/Source/ChaosVD/Private/Settings/ChaosVDCoreSettings.cpp b/Engine/Plugins/ChaosVD/Source/ChaosVD/Private/Settings/ChaosVDCoreSettings.cpp
new file mode 100644
--- /dev/null
+++ b/Engine/Plugins/ChaosVD/Source/ChaosVD/Private/Settings/ChaosVDCoreSettings.cpp
@@ -0,0 +1,31 @@
+// Copyright Epic Games, Inc. All Rights Reserved.
+
+#include "ChaosVDCoreSettings.h"
+
+#include "Widgets/SChaosVDPlaybackViewport.h"
+
+
+void UChaosVDSettingsObjectBase::PostEditChangeProperty(FPropertyChangedEvent& PropertyChangedEvent)
+{
+	UObject::PostEditChangeProperty(PropertyChangedEvent);
+
+	BroadcastSettingsChanged(this);
+}
+
+void UChaosVDSettingsObjectBase::PostEditUndo()
+{
+	UObject::PostEditUndo();
+	BroadcastSettingsChanged(this);
+}
+
+void UChaosVDSettingsObjectBase::BroadcastSettingsChanged(UObject* SettingsObject)
+{
+	SettingsChangedDelegate.Broadcast(this);
+}
+
+void UChaosVDVisualizationSettingsObjectBase::BroadcastSettingsChanged(UObject* SettingsObject)
+{
+	Super::BroadcastSettingsChanged(SettingsObject);
+	SChaosVDPlaybackViewport::ExecuteExternalViewportInvalidateRequest();
+}
+
diff --git a/Engine/Plugins/ChaosVD/Source/ChaosVD/Private/Settings/ChaosVDCoreSettings.h b/Engine/Plugins/ChaosVD/Source/ChaosVD/Private/Settings/ChaosVDCoreSettings.h
new file mode 100644
--- /dev/null
+++ b/Engine/Plugins/ChaosVD/Source/ChaosVD/Private/Settings/ChaosVDCoreSettings.h
@@ -0,0 +1,64 @@
+// Copyright Epic Games, Inc. All Rights Reserved.
+
+#pragma once
+
+#include "Engine/World.h"
+#include "Engine/StaticMesh.h"
+#include "UObject/Object.h"
+#include "UObject/SoftObjectPtr.h"
+
+#include "ChaosVDCoreSettings.generated.h"
+
+class UChaosVDCoreSettings;
+class UMaterial;
+
+DECLARE_MULTICAST_DELEGATE_OneParam(FChaosVDSettingChanged, UObject* SettingsObject)
+
+UCLASS(config = ChaosVD)
+class UChaosVDSettingsObjectBase : public UObject
+{
+	GENERATED_BODY()
+public:
+
+	virtual void PostEditChangeProperty(FPropertyChangedEvent& PropertyChangedEvent) override;
+
+	FChaosVDSettingChanged& OnSettingsChanged() { return SettingsChangedDelegate; }
+
+	virtual void PostEditUndo() override;
+
+protected:
+	virtual void BroadcastSettingsChanged(UObject* SettingsObject);
+	
+private:
+	FChaosVDSettingChanged SettingsChangedDelegate;
+};
+
+UCLASS(config = ChaosVD)
+class UChaosVDVisualizationSettingsObjectBase : public UChaosVDSettingsObjectBase
+{
+	GENERATED_BODY()
+protected:
+	virtual void BroadcastSettingsChanged(UObject* SettingsObject) override;
+};
+
+UCLASS(config = Engine)
+class UChaosVDCoreSettings : public UChaosVDSettingsObjectBase
+{
+	GENERATED_BODY()
+public:
+
+	UPROPERTY(Config, Transient)
+	TSoftObjectPtr<UMaterial> QueryOnlyMeshesMaterial;
+
+	UPROPERTY(Config, Transient)
+	TSoftObjectPtr<UMaterial> SimOnlyMeshesMaterial;
+
+	UPROPERTY(Config, Transient)
+	TSoftObjectPtr<UMaterial> InstancedMeshesMaterial;
+
+	UPROPERTY(Config, Transient)
+	TSoftObjectPtr<UMaterial> InstancedMeshesQueryOnlyMaterial;
+
+	UPROPERTY(Config)
+	FSoftClassPath SkySphereActorClass;
+};
diff --git a/Engine/Plugins/ChaosVD/Source/ChaosVD/Private/Settings/ChaosVDJointConstraintVisualizationSettings.cpp b/Engine/Plugins/ChaosVD/Source/ChaosVD/Private/Settings/ChaosVDJointConstraintVisualizationSettings.cpp
new file mode 100644
--- /dev/null
+++ b/Engine/Plugins/ChaosVD/Source/ChaosVD/Private/Settings/ChaosVDJointConstraintVisualizationSettings.cpp
@@ -0,0 +1,23 @@
+// Copyright Epic Games, Inc. All Rights Reserved.
+
+#include "ChaosVDJointConstraintVisualizationSettings.h"
+
+void UChaosVDJointConstraintsVisualizationSettings::SetJointsDataVisualizationFlags(EChaosVDJointsDataVisualizationFlags NewFlags)
+{
+	if (UChaosVDJointConstraintsVisualizationSettings* Settings = GetMutableDefault<UChaosVDJointConstraintsVisualizationSettings>())
+	{
+		Settings->GlobalJointsDataVisualizationFlags = NewFlags;
+		Settings->BroadcastSettingsChanged(Settings);
+	}
+}
+
+EChaosVDJointsDataVisualizationFlags UChaosVDJointConstraintsVisualizationSettings::GetJointsDataVisualizationFlags()
+{
+	if (UChaosVDJointConstraintsVisualizationSettings* Settings = GetMutableDefault<UChaosVDJointConstraintsVisualizationSettings>())
+	{
+		return Settings->GlobalJointsDataVisualizationFlags;
+	}
+
+	return EChaosVDJointsDataVisualizationFlags::None;
+}
+
diff --git a/Engine/Plugins/ChaosVD/Source/ChaosVD/Private/Settings/ChaosVDJointConstraintVisualizationSettings.h b/Engine/Plugins/ChaosVD/Source/ChaosVD/Private/Settings/ChaosVDJointConstraintVisualizationSettings.h
new file mode 100644
--- /dev/null
+++ b/Engine/Plugins/ChaosVD/Source/ChaosVD/Private/Settings/ChaosVDJointConstraintVisualizationSettings.h
@@ -0,0 +1,76 @@
+// Copyright Epic Games, Inc. All Rights Reserved.
+
+#pragma once
+#include "ChaosVDCoreSettings.h"
+
+#include "ChaosVDJointConstraintVisualizationSettings.generated.h"
+
+UENUM(meta = (Bitflags, UseEnumValuesAsMaskValuesInEditor = "true"))
+enum class EChaosVDJointsDataVisualizationFlags : uint32
+{
+	None					= 0 UMETA(Hidden),
+	/** Draw the PushOut vector based on the constraint's data */
+	PushOut					= 1 << 0,
+	/** Draw the Angular Impulse vector based on the constraint's data */
+	AngularImpulse			= 1 << 1 UMETA(Hidden),
+	ActorConnector			= 1 << 2,
+	CenterOfMassConnector	= 1 << 3,
+	Stretch					= 1 << 4,
+	Axes					= 1 << 5,
+	/** Draw the joint even if one of the particles or both are kinematic */
+	DrawKinematic			= 1 << 6,
+	/** Draw the joint even if it is disabled */
+	DrawDisabled			= 1 << 7,
+	/** Only debugs draw data for a selected joint constraint */
+	OnlyDrawSelected		= 1 << 8,
+	/** Enables Debug draw for Joint Constraint data from any solver that is visible */
+	EnableDraw				= 1 << 9,
+};
+ENUM_CLASS_FLAGS(EChaosVDJointsDataVisualizationFlags);
+
+UCLASS(config=ChaosVD)
+class UChaosVDJointConstraintsVisualizationSettings : public UChaosVDVisualizationSettingsObjectBase
+{
+	GENERATED_BODY()
+
+public:
+
+	/** If true, any debug draw text available will be drawn */
+	UPROPERTY(EditAnywhere, Category=DebugDraw)
+	bool bShowDebugText = false;
+
+	/** The depth priority used for while drawing data. Can be World or Foreground (with this one the shapes will be drawn on top of the geometry and be always visible) */
+	UPROPERTY(EditAnywhere, Category=DebugDraw)
+	TEnumAsByte<ESceneDepthPriorityGroup> DepthPriority = ESceneDepthPriorityGroup::SDPG_Foreground;
+
+	/** Scale to apply to the Linear Impulse vector before draw it. */
+	UPROPERTY(EditAnywhere, config, Category=DebugDraw)
+	float LinearImpulseScale = 0.001;
+
+	/** Scale to apply to the Angular Impulse vector before draw it. */
+	UPROPERTY(EditAnywhere, config, Category=DebugDraw)
+	float AngularImpulseScale = 0.1f;
+
+	/** Scale to apply to anything that does not have a dedicated scale setting before draw it. */
+	UPROPERTY(EditAnywhere, config, Category=DebugDraw)
+	float GeneralScale = 1.0f;
+
+	/** Line thickness to use as a base to calculate the different line thickness values used to debug draw the data. */
+	UPROPERTY(EditAnywhere, config, Category=DebugDraw)
+	float BaseLineThickness = 2.0f;
+
+	/** Size of the debug drawn Center Of Mass. */
+	UPROPERTY(EditAnywhere, config, Category=DebugDraw)
+	float CenterOfMassSize = 1.0f;
+
+	/** Size of the debug drawn if the Constraint Axis */
+	UPROPERTY(EditAnywhere, config, Category=DebugDraw)
+	float ConstraintAxisLength = 10.0f;
+
+	static void SetJointsDataVisualizationFlags(EChaosVDJointsDataVisualizationFlags NewFlags);
+	static EChaosVDJointsDataVisualizationFlags GetJointsDataVisualizationFlags();
+
+private:
+	/** Set of flags to enable/disable visualization of specific joint constraints data as debug draw */
+	EChaosVDJointsDataVisualizationFlags GlobalJointsDataVisualizationFlags = EChaosVDJointsDataVisualizationFlags::ActorConnector | EChaosVDJointsDataVisualizationFlags::DrawKinematic;
+};
\ No newline at end of file
diff --git a/Engine/Plugins/ChaosVD/Source/ChaosVD/Private/Settings/ChaosVDParticleVisualizationSettings.cpp b/Engine/Plugins/ChaosVD/Source/ChaosVD/Private/Settings/ChaosVDParticleVisualizationSettings.cpp
new file mode 100644
--- /dev/null
+++ b/Engine/Plugins/ChaosVD/Source/ChaosVD/Private/Settings/ChaosVDParticleVisualizationSettings.cpp
@@ -0,0 +1,169 @@
+// Copyright Epic Games, Inc. All Rights Reserved.
+
+#include "ChaosVDParticleVisualizationSettings.h"
+
+#include "DataWrappers/ChaosVDParticleDataWrapper.h"
+
+
+FColor FChaosDebugDrawColorsByState::GetColorFromState(EChaosVDObjectStateType State) const
+{
+	switch (State)
+	{
+		case EChaosVDObjectStateType::Sleeping:
+			return SleepingColor;
+		case EChaosVDObjectStateType::Kinematic:
+			return KinematicColor;
+		case EChaosVDObjectStateType::Static:
+			return StaticColor;
+		case EChaosVDObjectStateType::Dynamic:
+			return DynamicColor;
+		default:
+			return FColor::Purple;
+	}
+}
+
+FColor FChaosParticleDataDebugDrawColors::GetColorForDataID(EChaosVDParticleDataVisualizationFlags DataID, bool bIsSelected) const
+{
+	constexpr float DefaultIntensityFactor = 0.6f;
+	constexpr float SelectedIntensityFactor = 1.0f;
+	const float IntensityFactor = bIsSelected ? SelectedIntensityFactor : DefaultIntensityFactor;
+
+	return (GetLinearColorForDataID(DataID) * IntensityFactor).ToFColorSRGB();
+}
+
+const FLinearColor& FChaosParticleDataDebugDrawColors::GetLinearColorForDataID(EChaosVDParticleDataVisualizationFlags DataID) const
+{
+	static FLinearColor InvalidColor = FColor::Purple;
+	switch (DataID)
+	{
+	case EChaosVDParticleDataVisualizationFlags::Acceleration:
+		return AccelerationColor;
+	case EChaosVDParticleDataVisualizationFlags::Velocity:
+		return VelocityColor;
+	case EChaosVDParticleDataVisualizationFlags::AngularVelocity:
+		return AngularVelocityColor;
+	case EChaosVDParticleDataVisualizationFlags::AngularAcceleration:
+		return AngularAccelerationColor;
+	case EChaosVDParticleDataVisualizationFlags::LinearImpulse:
+		return LinearImpulseColor;
+	case EChaosVDParticleDataVisualizationFlags::AngularImpulse:
+		return AngularImpulseColor;
+	case EChaosVDParticleDataVisualizationFlags::ClusterConnectivityEdge:
+		return ConnectivityDataColor;
+	case EChaosVDParticleDataVisualizationFlags::CenterOfMass:
+		return CenterOfMassColor;
+	case EChaosVDParticleDataVisualizationFlags::None:
+	case EChaosVDParticleDataVisualizationFlags::DrawDataOnlyForSelectedParticle:
+	default:
+		return InvalidColor;
+	}
+}
+
+float UChaosVDParticleVisualizationDebugDrawSettings::GetScaleFortDataID(EChaosVDParticleDataVisualizationFlags DataID) const
+{
+	switch (DataID)
+    {
+    	case EChaosVDParticleDataVisualizationFlags::Acceleration:
+    		return AccelerationScale;
+    	case EChaosVDParticleDataVisualizationFlags::Velocity:
+    		return VelocityScale;
+    	case EChaosVDParticleDataVisualizationFlags::AngularVelocity:
+    		return AngularVelocityScale;
+    	case EChaosVDParticleDataVisualizationFlags::AngularAcceleration:
+    		return AngularAccelerationScale;
+    	case EChaosVDParticleDataVisualizationFlags::LinearImpulse:
+    		return LinearImpulseScale;
+    	case EChaosVDParticleDataVisualizationFlags::AngularImpulse:
+    		return AngularImpulseScale;
+    	case EChaosVDParticleDataVisualizationFlags::ClusterConnectivityEdge:
+    	case EChaosVDParticleDataVisualizationFlags::CenterOfMass:
+    	case EChaosVDParticleDataVisualizationFlags::None:
+    	case EChaosVDParticleDataVisualizationFlags::DrawDataOnlyForSelectedParticle:
+    	default:
+    		return 1.0f;
+    }
+}
+
+void UChaosVDParticleVisualizationDebugDrawSettings::SetDataDebugDrawVisualizationFlags(EChaosVDParticleDataVisualizationFlags Flags)
+{
+	if (UChaosVDParticleVisualizationDebugDrawSettings* Settings = GetMutableDefault<UChaosVDParticleVisualizationDebugDrawSettings>())
+	{
+		Settings->ParticleDataVisualizationFlags = Flags;
+		Settings->BroadcastSettingsChanged(Settings);
+	}
+}
+
+EChaosVDParticleDataVisualizationFlags UChaosVDParticleVisualizationDebugDrawSettings::GetDataDebugDrawVisualizationFlags()
+{
+	if (UChaosVDParticleVisualizationDebugDrawSettings* Settings = GetMutableDefault<UChaosVDParticleVisualizationDebugDrawSettings>())
+	{
+		return Settings->ParticleDataVisualizationFlags;
+	}
+
+	return EChaosVDParticleDataVisualizationFlags::None;
+}
+
+void UChaosVDParticleVisualizationSettings::SetGeometryVisualizationFlags(EChaosVDGeometryVisibilityFlags Flags)
+{
+	if (UChaosVDParticleVisualizationSettings* Settings = GetMutableDefault<UChaosVDParticleVisualizationSettings>())
+	{
+		Settings->GeometryVisibilityFlags = Flags;
+		Settings->BroadcastSettingsChanged(Settings);
+	}
+}
+
+EChaosVDGeometryVisibilityFlags UChaosVDParticleVisualizationSettings::GetGeometryVisualizationFlags()
+{
+	if (UChaosVDParticleVisualizationSettings* Settings = GetMutableDefault<UChaosVDParticleVisualizationSettings>())
+	{
+		return Settings->GeometryVisibilityFlags;
+	}
+
+	return EChaosVDGeometryVisibilityFlags::None;
+}
+
+FColor FChaosDebugDrawColorsByShapeType::GetColorFromShapeType(Chaos::EImplicitObjectType ShapeType) const
+{
+	switch(ShapeType)
+	{
+		case Chaos::ImplicitObjectType::Sphere:
+			return SimpleTypeColor;
+		case Chaos::ImplicitObjectType::Box:
+			return SimpleTypeColor;
+		case Chaos::ImplicitObjectType::Plane:
+			return SimpleTypeColor;
+		case Chaos::ImplicitObjectType::Capsule:
+			return SimpleTypeColor;
+		case Chaos::ImplicitObjectType::TaperedCylinder:
+			return SimpleTypeColor;
+		case Chaos::ImplicitObjectType::Cylinder:
+			return SimpleTypeColor;
+		case Chaos::ImplicitObjectType::Convex:
+			return ConvexColor;
+		case Chaos::ImplicitObjectType::HeightField:
+			return HeightFieldColor;
+		case Chaos::ImplicitObjectType::TriangleMesh:
+			return TriangleMeshColor;
+		case Chaos::ImplicitObjectType::LevelSet:
+			return LevelSetColor;			
+		default:
+			return FColor::Purple; 
+	}
+}
+
+FColor FChaosDebugDrawColorsByClientServer::GetColorFromState(bool bIsServer, EChaosVDObjectStateType State) const
+{
+	switch (State)
+	{
+	case EChaosVDObjectStateType::Sleeping:
+		return bIsServer ? ServerSleepingColor : ClientSleepingColor;
+	case EChaosVDObjectStateType::Kinematic:
+		return bIsServer ? ServerColor : ClientColor;
+	case EChaosVDObjectStateType::Static:
+		return bIsServer ? ServerColor : ClientColor;
+	case EChaosVDObjectStateType::Dynamic:
+		return bIsServer ? ServerDynamicColor : ClientDynamicColor;
+	default:
+		return FColor::Purple;
+	}
+}
diff --git a/Engine/Plugins/ChaosVD/Source/ChaosVD/Private/Settings/ChaosVDParticleVisualizationSettings.h b/Engine/Plugins/ChaosVD/Source/ChaosVD/Private/Settings/ChaosVDParticleVisualizationSettings.h
new file mode 100644
--- /dev/null
+++ b/Engine/Plugins/ChaosVD/Source/ChaosVD/Private/Settings/ChaosVDParticleVisualizationSettings.h
@@ -0,0 +1,314 @@
+// Copyright Epic Games, Inc. All Rights Reserved.
+
+#pragma once
+
+#include "ChaosVDCoreSettings.h"
+#include "Chaos/ImplicitObject.h"
+#include "UObject/Object.h"
+#include "ChaosVDParticleVisualizationSettings.generated.h"
+
+enum class EChaosVDObjectStateType: int8;
+
+UENUM(meta = (Bitflags, UseEnumValuesAsMaskValuesInEditor = "true"))
+enum class EChaosVDGeometryVisibilityFlags : uint8
+{
+	None				= 0 UMETA(Hidden),
+	/** Draws all geometry that is for Query Only */
+	Query = 1 << 1,
+	/** Draws all geometry that is for [Physics Collision] or [Physics Collision and Query only] */
+	Simulated = 1 << 2,
+	/** Draws all simple geometry */
+	Simple = 1 << 3,
+	/** Draws all complex geometry */
+	Complex = 1 << 4,
+	/** Draws heightfields even if complex is not selected */
+	ShowHeightfields = 1 << 5,
+	/** Draws all particles that are in a disabled state */
+	ShowDisabledParticles = 1 << 6,
+};
+ENUM_CLASS_FLAGS(EChaosVDGeometryVisibilityFlags)
+
+UENUM(meta = (Bitflags, UseEnumValuesAsMaskValuesInEditor = "true"))
+enum class EChaosVDParticleDataVisualizationFlags : uint32
+{
+	None				= 0 UMETA(Hidden),
+	Velocity			= 1 << 0,
+	AngularVelocity		= 1 << 1,
+	Acceleration		= 1 << 2,
+	AngularAcceleration = 1 << 3,
+	LinearImpulse		= 1 << 4,
+	AngularImpulse		= 1 << 5,
+	ClusterConnectivityEdge	= 1 << 6,
+	CenterOfMass	= 1 << 7,
+	DrawDataOnlyForSelectedParticle	= 1 << 8,
+	EnableDraw	= 1 << 9,
+};
+ENUM_CLASS_FLAGS(EChaosVDParticleDataVisualizationFlags);
+
+
+/** Structure holding the settings using to debug draw Particles shape based on their shape type on the Chaos Visual Debugger */
+USTRUCT()
+struct FChaosDebugDrawColorsByShapeType
+{
+	GENERATED_BODY()
+
+	/** Color used for Sphere, Plane, Cube, Capsule, Cylinder, tapered shapes */
+	UPROPERTY(EditAnywhere, Category=DebugDraw)
+	FColor SimpleTypeColor = FColor(0, 255, 0); 
+
+	/** Color used for convex shapes */
+	UPROPERTY(EditAnywhere, Category=DebugDraw)
+	FColor ConvexColor = FColor(0, 255, 255);
+
+	/** Color used for heightfield */
+	UPROPERTY(EditAnywhere, Category=DebugDraw)
+	FColor HeightFieldColor = FColor(0, 0, 255);
+	
+	/** Color used for triangle meshes */
+	UPROPERTY(EditAnywhere, Category=DebugDraw)
+	FColor TriangleMeshColor = FColor(255, 0, 0);
+
+	/** Color used for triangle LevelSets */
+	UPROPERTY(EditAnywhere, Category=DebugDraw)
+	FColor LevelSetColor = FColor(255, 0, 128);
+
+	FColor GetColorFromShapeType(Chaos::EImplicitObjectType ShapeType) const;
+};
+
+/** Structure holding the settings using to debug draw Particles shape based on whether they are client or server objects (in PIE) Chaos Visual Debugger */
+USTRUCT()
+struct FChaosDebugDrawColorsByClientServer
+{
+	GENERATED_BODY()
+
+	/** Color used for server shapes that are not awake or sleeping dynamic */
+	UPROPERTY(config, EditAnywhere, Category=DebugDraw)
+	FColor ServerColor = FColor(50, 0, 0); 
+
+	/** Color used for server shapes that are awake dynamic */
+	UPROPERTY(config, EditAnywhere, Category = DebugDraw)
+	FColor ServerDynamicColor = FColor(150, 0, 0);
+
+	/** Color used for server shapes that are sleeping dynamics */
+	UPROPERTY(config, EditAnywhere, Category = DebugDraw)
+	FColor ServerSleepingColor = FColor(10, 0, 0);
+
+	/** Color used for client shapes that are not awake or sleeping dynamic */
+	UPROPERTY(config, EditAnywhere, Category = DebugDraw)
+	FColor ClientColor = FColor(0, 0, 50);
+
+	/** Color used for server shapes that are awake dynamic */
+	UPROPERTY(config, EditAnywhere, Category = DebugDraw)
+	FColor ClientDynamicColor = FColor(0, 0, 150);
+
+	/** Color used for client shapes that are sleeping dynamics */
+	UPROPERTY(config, EditAnywhere, Category = DebugDraw)
+	FColor ClientSleepingColor = FColor(0, 0, 100);
+
+	FColor GetColorFromState(bool bIsServer, EChaosVDObjectStateType State) const;
+};
+
+UENUM()
+enum class EChaosVDParticleDebugColorMode
+{
+	/** Draw particles with the default gray color */
+	None,
+	/** Draw particles with a specific color based on the recorded particle state */
+	State,
+	/** Draw particles with a specific color based on their shape type */
+	ShapeType,
+	/** Draw particles with a specific color based on if they are a Server Particle or Client particle */
+	ClientServer,
+};
+
+/** Structure holding the settings using to debug draw Particles shape based on their state on the Chaos Visual Debugger */
+USTRUCT()
+struct FChaosDebugDrawColorsByState
+{
+	GENERATED_BODY()
+
+	/** Color used for dynamic particles */
+	UPROPERTY(config, EditAnywhere, Category=DebugDraw)
+	FColor DynamicColor = FColor(255, 255, 0);
+	
+	/** Color used for sleeping particles */
+	UPROPERTY(config, EditAnywhere, Category=DebugDraw)
+	FColor SleepingColor = FColor(128, 128, 128);
+
+	/** Color used for kinematic particles */
+	UPROPERTY(config, EditAnywhere, Category=DebugDraw)
+	FColor KinematicColor = FColor(0, 128, 255);
+
+	/** Color used for static particles */
+	UPROPERTY(config, EditAnywhere, Category=DebugDraw)
+	FColor StaticColor = FColor(255, 0, 0);
+
+	FColor GetColorFromState(EChaosVDObjectStateType State) const;
+};
+
+USTRUCT()
+struct FChaosParticleDataDebugDrawColors
+{
+	GENERATED_BODY()
+
+	/** Color to apply to the Velocity vector when drawing it */
+	UPROPERTY(config, EditAnywhere, Category=DebugDraw)
+	FLinearColor VelocityColor = FColor::Green;
+
+	/** Color to apply to the Angular Velocity vector when drawing it */
+	UPROPERTY(config, EditAnywhere, Category=DebugDraw)
+	FLinearColor AngularVelocityColor = FColor::Blue;
+
+	/** Color to apply to the Acceleration vector when drawing it */
+	UPROPERTY(config, EditAnywhere, Category=DebugDraw)
+	FLinearColor AccelerationColor = FColor::Orange;
+
+	/** Color to apply to the Angular Acceleration vector when drawing it */
+	UPROPERTY(config, EditAnywhere, Category=DebugDraw)
+	FLinearColor AngularAccelerationColor = FColor::Silver;
+
+	/** Color to apply to the Linear Impulse when drawing it */
+	UPROPERTY(config, EditAnywhere, Category=DebugDraw)
+	FLinearColor LinearImpulseColor = FColor::Turquoise;
+
+	/** Color to apply to the Angular Impulse vector when drawing it */
+	UPROPERTY(config, EditAnywhere, Category=DebugDraw)
+	FLinearColor AngularImpulseColor = FColor::Emerald;
+
+	/** Color to apply the debug drawn sphere representing the center of mass location */
+	UPROPERTY(config, EditAnywhere, Category=DebugDraw)
+	FLinearColor CenterOfMassColor = FColor::Red;
+
+	/** Color to apply to when drawing the connectivity data */
+	UPROPERTY(config, EditAnywhere, Category=DebugDraw)
+	FLinearColor ConnectivityDataColor = FColor::Yellow;
+	
+	FColor GetColorForDataID(EChaosVDParticleDataVisualizationFlags DataID, bool bIsSelected = false) const;
+	const FLinearColor& GetLinearColorForDataID(EChaosVDParticleDataVisualizationFlags DataID) const;
+};
+
+namespace Chaos::VisualDebugger::ParticleDataUnitsStrings
+{
+	static FString Velocity = TEXT("cm/s");
+	static FString AngularVelocity = TEXT("rad/s");
+	static FString Acceleration = TEXT("cm/s2");
+	static FString AngularAcceleration = TEXT("rad/s2");
+	static FString LinearImpulse = TEXT("g.m/s");
+	static FString AngularImpulse = TEXT("g.m2/s");
+
+	static FString GetUnitByID(EChaosVDParticleDataVisualizationFlags DataID)
+	{
+		switch (DataID)
+		{
+		case EChaosVDParticleDataVisualizationFlags::Velocity:
+			return Velocity;
+		case EChaosVDParticleDataVisualizationFlags::AngularVelocity:
+			return AngularVelocity;
+		case EChaosVDParticleDataVisualizationFlags::Acceleration:
+			return Acceleration;
+		case EChaosVDParticleDataVisualizationFlags::AngularAcceleration:
+			return AngularAcceleration;
+		case EChaosVDParticleDataVisualizationFlags::LinearImpulse:
+			return LinearImpulse;
+		case EChaosVDParticleDataVisualizationFlags::AngularImpulse:
+			return AngularImpulse;
+
+		case EChaosVDParticleDataVisualizationFlags::None:
+		case EChaosVDParticleDataVisualizationFlags::ClusterConnectivityEdge:
+		case EChaosVDParticleDataVisualizationFlags::CenterOfMass:
+		case EChaosVDParticleDataVisualizationFlags::DrawDataOnlyForSelectedParticle:
+		default:
+			return TEXT("");
+		}
+	}
+}
+
+UCLASS(config=ChaosVD)
+class UChaosVDParticleVisualizationColorSettings : public UChaosVDVisualizationSettingsObjectBase
+{
+	GENERATED_BODY()
+public:
+
+	UPROPERTY(config, EditAnywhere, Category = "Colors Mode")
+	EChaosVDParticleDebugColorMode ParticleColorMode = EChaosVDParticleDebugColorMode::None;
+	
+	UPROPERTY(config, EditAnywhere, Category = "Colors By Shape", meta=(EditCondition = "ParticleColorMode == EChaosVDParticleDebugColorMode::ShapeType", EditConditionHides))
+	FChaosDebugDrawColorsByShapeType ColorsByShapeType;
+	
+	UPROPERTY(config, EditAnywhere, Category = "Colors By State", meta=(EditCondition = "ParticleColorMode == EChaosVDParticleDebugColorMode::State", EditConditionHides))
+	FChaosDebugDrawColorsByState ColorsByParticleState;
+	
+	UPROPERTY(config, EditAnywhere, Category = "Colors By Solver", meta=(EditCondition = "ParticleColorMode == EChaosVDParticleDebugColorMode::ClientServer", EditConditionHides))
+	FChaosDebugDrawColorsByClientServer ColorsByClientServer;
+};
+
+UCLASS(config=ChaosVD)
+class UChaosVDParticleVisualizationDebugDrawSettings : public UChaosVDVisualizationSettingsObjectBase
+{
+	GENERATED_BODY()
+public:
+
+	UPROPERTY(EditAnywhere, Category=DebugDraw)
+	bool bShowDebugText = false;
+
+	/** The depth priority used for while drawing contact data. Can be World or Foreground (with this one the shapes will be drawn on top of the geometry and be always visible) */
+	UPROPERTY(config, EditAnywhere, Category=DebugDraw)
+	TEnumAsByte<ESceneDepthPriorityGroup> DepthPriority = ESceneDepthPriorityGroup::SDPG_World;
+
+	/** Scale to apply to the Velocity vector before draw it. Unit is cm/s */
+	UPROPERTY(config, EditAnywhere, Category=DebugDraw)
+	float VelocityScale = 0.5f;
+
+	/** Scale to apply to the Angular Velocity vector before draw it. Unit is rad/s */
+	UPROPERTY(config, EditAnywhere, Category=DebugDraw)
+	float AngularVelocityScale = 50.0f;
+
+	/** Scale to apply to the Acceleration vector before draw it. Unit is cm/s2 */
+	UPROPERTY(config, EditAnywhere, Category=DebugDraw)
+	float AccelerationScale = 0.005f;
+
+	/** Scale to apply to the Angular Acceleration vector before draw it. Unit is rad/s2 */
+	UPROPERTY(config, EditAnywhere, Category=DebugDraw)
+	float AngularAccelerationScale = 0.5f;
+
+	/** Scale to apply to the Linear Impulse vector before draw it. Unit is g.m/s */
+	UPROPERTY(config, EditAnywhere, Category=DebugDraw)
+	float LinearImpulseScale = 0.001;
+
+	/** Scale to apply to the Angular Impulse vector before draw it. Unit is g.m2/s */
+	UPROPERTY(config, EditAnywhere, Category=DebugDraw)
+	float AngularImpulseScale = 0.1f;
+
+	/** Radius to use when creating the sphere that will represent the center of mass location */
+	UPROPERTY(config, EditAnywhere, Category=DebugDraw)
+	float CenterOfMassRadius = 10.0f;
+
+	UPROPERTY(config, EditAnywhere, Category=DebugDraw)
+	FChaosParticleDataDebugDrawColors ColorSettings;
+
+	float GetScaleFortDataID(EChaosVDParticleDataVisualizationFlags DataID) const;
+	
+	static void SetDataDebugDrawVisualizationFlags(EChaosVDParticleDataVisualizationFlags Flags);
+
+	static EChaosVDParticleDataVisualizationFlags GetDataDebugDrawVisualizationFlags();
+
+private:
+	/** Set of flags to enable/disable visualization of specific particle data as debug draw */
+	UPROPERTY()
+	EChaosVDParticleDataVisualizationFlags ParticleDataVisualizationFlags = EChaosVDParticleDataVisualizationFlags::Velocity | EChaosVDParticleDataVisualizationFlags::AngularVelocity;
+};
+
+UCLASS()
+class UChaosVDParticleVisualizationSettings : public UChaosVDVisualizationSettingsObjectBase
+{
+	GENERATED_BODY()
+public:
+
+	static void SetGeometryVisualizationFlags(EChaosVDGeometryVisibilityFlags Flags);
+	static EChaosVDGeometryVisibilityFlags GetGeometryVisualizationFlags();
+
+private:
+
+	/** Set of flags to enable/disable visibility of specific types of geometry/particles */
+	EChaosVDGeometryVisibilityFlags GeometryVisibilityFlags = EChaosVDGeometryVisibilityFlags::Simulated | EChaosVDGeometryVisibilityFlags::Simple |  EChaosVDGeometryVisibilityFlags::ShowHeightfields;
+};
\ No newline at end of file
diff --git a/Engine/Plugins/ChaosVD/Source/ChaosVD/Private/Settings/ChaosVDSceneQueryVisualizationSettings.cpp b/Engine/Plugins/ChaosVD/Source/ChaosVD/Private/Settings/ChaosVDSceneQueryVisualizationSettings.cpp
new file mode 100644
--- /dev/null
+++ b/Engine/Plugins/ChaosVD/Source/ChaosVD/Private/Settings/ChaosVDSceneQueryVisualizationSettings.cpp
@@ -0,0 +1,22 @@
+// Copyright Epic Games, Inc. All Rights Reserved.
+
+#include "ChaosVDSceneQueryVisualizationSettings.h"
+
+void UChaosVDSceneQueriesVisualizationSettings::SetSceneQueryDataVisualizationFlags(EChaosVDSceneQueryVisualizationFlags NewFlags)
+{
+	if (UChaosVDSceneQueriesVisualizationSettings* Settings = GetMutableDefault<UChaosVDSceneQueriesVisualizationSettings>())
+	{
+		Settings->GlobalSceneQueriesVisualizationFlags = NewFlags;
+		Settings->BroadcastSettingsChanged(Settings);
+	}
+}
+
+EChaosVDSceneQueryVisualizationFlags UChaosVDSceneQueriesVisualizationSettings::GetSceneQueryDataVisualizationFlags()
+{
+	if (UChaosVDSceneQueriesVisualizationSettings* Settings = GetMutableDefault<UChaosVDSceneQueriesVisualizationSettings>())
+	{
+		return Settings->GlobalSceneQueriesVisualizationFlags;
+	}
+
+	return EChaosVDSceneQueryVisualizationFlags::None;
+}
diff --git a/Engine/Plugins/ChaosVD/Source/ChaosVD/Private/Settings/ChaosVDSceneQueryVisualizationSettings.h b/Engine/Plugins/ChaosVD/Source/ChaosVD/Private/Settings/ChaosVDSceneQueryVisualizationSettings.h
new file mode 100644
--- /dev/null
+++ b/Engine/Plugins/ChaosVD/Source/ChaosVD/Private/Settings/ChaosVDSceneQueryVisualizationSettings.h
@@ -0,0 +1,41 @@
+// Copyright Epic Games, Inc. All Rights Reserved.
+
+#pragma once
+#include "ChaosVDCoreSettings.h"
+
+#include "ChaosVDSceneQueryVisualizationSettings.generated.h"
+
+/** Set of visualization flags options for Scene Queries */
+UENUM(meta = (Bitflags, UseEnumValuesAsMaskValuesInEditor = "true"))
+enum class EChaosVDSceneQueryVisualizationFlags: uint32
+{
+	None					= 0 UMETA(Hidden),
+	EnableDraw				= 1 << 0,
+	DrawLineTraceQueries	= 1 << 1,
+	DrawSweepQueries		= 1 << 2,
+	DrawOverlapQueries		= 1 << 3,
+	DrawHits				= 1 << 4,
+	OnlyDrawSelectedQuery	= 1 << 5,
+	HideEmptyQueries		= 1 << 6,
+	HideSubQueries			= 1 << 7,
+};
+ENUM_CLASS_FLAGS(EChaosVDSceneQueryVisualizationFlags);
+
+UCLASS(config=ChaosVD)
+class UChaosVDSceneQueriesVisualizationSettings : public UChaosVDVisualizationSettingsObjectBase
+{
+	GENERATED_BODY()
+
+public:
+
+	static void SetSceneQueryDataVisualizationFlags(EChaosVDSceneQueryVisualizationFlags NewFlags);
+	static EChaosVDSceneQueryVisualizationFlags GetSceneQueryDataVisualizationFlags();
+
+	/** If true, any debug draw text available will be drawn */
+	UPROPERTY(EditAnywhere, Category=DebugDraw)
+	bool bShowText = false;
+
+private:
+	/** Set of flags to enable/disable visualization of specific scene queries data as debug draw */
+	EChaosVDSceneQueryVisualizationFlags GlobalSceneQueriesVisualizationFlags = EChaosVDSceneQueryVisualizationFlags::DrawHits | EChaosVDSceneQueryVisualizationFlags::DrawLineTraceQueries;
+};
diff --git a/Engine/Plugins/ChaosVD/Source/ChaosVD/Private/Utils/ChaosVDUserInterfaceUtils.cpp b/Engine/Plugins/ChaosVD/Source/ChaosVD/Private/Utils/ChaosVDUserInterfaceUtils.cpp
new file mode 100644
--- /dev/null
+++ b/Engine/Plugins/ChaosVD/Source/ChaosVD/Private/Utils/ChaosVDUserInterfaceUtils.cpp
@@ -0,0 +1,91 @@
+// Copyright Epic Games, Inc. All Rights Reserved.
+
+#include "ChaosVDUserInterfaceUtils.h"
+#include "PropertyEditorModule.h"
+#include "Modules/ModuleManager.h"
+#include "ToolMenu.h"
+#include "ToolMenuEntry.h"
+#include "ToolMenuSection.h"
+#include "Widgets/SBoxPanel.h"
+#include "Widgets/Text/STextBlock.h"
+
+TSharedRef<IStructureDetailsView> Chaos::VisualDebugger::Utils::MakeStructDetailsViewForMenu()
+{
+	FPropertyEditorModule& PropertyEditorModule = FModuleManager::GetModuleChecked<FPropertyEditorModule>("PropertyEditor");
+
+	const FStructureDetailsViewArgs StructDetailsViewArgs;
+	FDetailsViewArgs DetailsViewArgs;
+	DetailsViewArgs.bShowOptions = false;
+	DetailsViewArgs.bAllowFavoriteSystem = false;
+	DetailsViewArgs.bAllowSearch = false;
+	DetailsViewArgs.bShowScrollBar = false;
+	DetailsViewArgs.bUpdatesFromSelection = false;
+	DetailsViewArgs.ColumnWidth = 1.0f;
+
+	return PropertyEditorModule.CreateStructureDetailView(DetailsViewArgs,StructDetailsViewArgs, nullptr);
+}
+
+TSharedRef<IDetailsView> Chaos::VisualDebugger::Utils::MakeObjectDetailsViewForMenu()
+{
+	FPropertyEditorModule& PropertyEditorModule = FModuleManager::GetModuleChecked<FPropertyEditorModule>("PropertyEditor");
+
+	FDetailsViewArgs DetailsViewArgs;
+	DetailsViewArgs.bShowOptions = true;
+	DetailsViewArgs.bAllowFavoriteSystem = false;
+	DetailsViewArgs.bAllowSearch = false;
+	DetailsViewArgs.bShowScrollBar = false;
+	DetailsViewArgs.bShowObjectLabel = false;
+	DetailsViewArgs.bCustomNameAreaLocation = true;
+	DetailsViewArgs.ColumnWidth = 0.45f;
+
+	return PropertyEditorModule.CreateDetailView(DetailsViewArgs);
+}
+
+void Chaos::VisualDebugger::Utils::CreateMenuEntryForObject(UToolMenu* Menu, UObject* Object, EChaosVDSaveSettingsOptions MenuEntryOptions)
+{
+	if (!Menu)
+	{
+		return;
+	}
+
+	if (!Object)
+	{
+		TSharedRef<SWidget> ErrorMessageWidget = SNew(SVerticalBox)
+		+SVerticalBox::Slot()
+		[
+			SNew(STextBlock)
+			.Text(NSLOCTEXT("ChaosVisualDebugger", "CreateMenuEntryForObjectErrorMessage", "Failed to create menu for object. The provied object is null"))
+		];
+		
+		FToolMenuEntry ErrorMenuEntry = FToolMenuEntry::InitWidget(Object->GetFName(), ErrorMessageWidget, FText::GetEmpty());
+		Menu->AddMenuEntry(NAME_None, ErrorMenuEntry);
+		return;
+	}
+	
+	TSharedRef<IDetailsView> DetailsView = MakeObjectDetailsViewForMenu();
+	DetailsView->SetObject(Object);
+	FToolMenuEntry MenuEntry = FToolMenuEntry::InitWidget(Object->GetFName(), DetailsView, FText::GetEmpty());
+	Menu->AddMenuEntry(NAME_None, MenuEntry);
+
+	if (EnumHasAnyFlags(MenuEntryOptions, EChaosVDSaveSettingsOptions::ShowSaveButton))
+	{
+		Menu->AddMenuEntry(NAME_None,FToolMenuEntry::InitSeparator(NAME_None));
+		
+		const FString SaveObjectButtonName = Object->GetName() + TEXT("SaveButton");
+		
+		FToolMenuEntry SaveMenuEntry = FToolMenuEntry::InitMenuEntry(FName(SaveObjectButtonName),
+			NSLOCTEXT("ChaosVisualDebugger","CreateMenuEntryForObjectSaveButtonLabel", "Save Settings"),
+			NSLOCTEXT("ChaosVisualDebugger","CreateMenuEntryForObjectSaveButtonToolTip", "Saves the current settings into the Editor's confiuration file"),
+			FSlateIcon(FAppStyle::Get().GetStyleSetName(), TEXT("LevelEditor.Save")),
+			FUIAction(
+				FExecuteAction::CreateLambda([ObjectWeakPtr = TWeakObjectPtr(Object)]()
+				{
+					if (UObject* Object = ObjectWeakPtr.Get())
+					{
+						Object->SaveConfig();
+					}
+				})));
+			
+		Menu->AddMenuEntry(NAME_None, SaveMenuEntry);
+	}
+}
diff --git a/Engine/Plugins/ChaosVD/Source/ChaosVD/Private/Utils/ChaosVDUserInterfaceUtils.h b/Engine/Plugins/ChaosVD/Source/ChaosVD/Private/Utils/ChaosVDUserInterfaceUtils.h
new file mode 100644
--- /dev/null
+++ b/Engine/Plugins/ChaosVD/Source/ChaosVD/Private/Utils/ChaosVDUserInterfaceUtils.h
@@ -0,0 +1,42 @@
+// Copyright Epic Games, Inc. All Rights Reserved.
+
+#pragma once
+#include "IStructureDetailsView.h"
+
+class UToolMenu;
+
+namespace Chaos::VisualDebugger::Utils
+{
+	TSharedRef<IStructureDetailsView> MakeStructDetailsViewForMenu();
+
+	TSharedRef<IDetailsView> MakeObjectDetailsViewForMenu();
+
+	UENUM()
+	enum class EChaosVDSaveSettingsOptions
+	{
+		None = 0,
+		ShowSaveButton = 1 << 0
+	};
+	ENUM_CLASS_FLAGS(EChaosVDSaveSettingsOptions)
+
+	void CreateMenuEntryForObject(UToolMenu* Menu, UObject* Object, EChaosVDSaveSettingsOptions MenuEntryOptions = EChaosVDSaveSettingsOptions::None);
+
+	template <typename TObject>
+	void CreateMenuEntryForDefaultObject(UToolMenu* Menu, EChaosVDSaveSettingsOptions MenuEntryOptions = EChaosVDSaveSettingsOptions::None)
+	{
+		CreateMenuEntryForObject(Menu, GetMutableDefault<TObject>(), MenuEntryOptions);
+	}
+
+	template <typename TStruct>
+	void SetStructToDetailsView(TStruct* NewStruct, TSharedRef<IStructureDetailsView>& InDetailsView)
+	{
+		TSharedPtr<FStructOnScope> StructDataView = nullptr;
+
+		if (NewStruct)
+		{
+			StructDataView = MakeShared<FStructOnScope>(TStruct::StaticStruct(), reinterpret_cast<uint8*>(NewStruct));
+		}
+
+		InDetailsView->SetStructureData(StructDataView);
+	}
+}
diff --git a/Engine/Plugins/ChaosVD/Source/ChaosVD/Private/Visualizers/ChaosVDCharacterGroundConstraintsDataComponentVisualizer.cpp b/Engine/Plugins/ChaosVD/Source/ChaosVD/Private/Visualizers/ChaosVDCharacterGroundConstraintsDataComponentVisualizer.cpp
--- a/Engine/Plugins/ChaosVD/Source/ChaosVD/Private/Visualizers/ChaosVDCharacterGroundConstraintsDataComponentVisualizer.cpp
+++ b/Engine/Plugins/ChaosVD/Source/ChaosVD/Private/Visualizers/ChaosVDCharacterGroundConstraintsDataComponentVisualizer.cpp
@@ -10,12 +10,63 @@
 #include "SceneView.h"
 #include "Actors/ChaosVDSolverInfoActor.h"
 #include "Components/ChaosVDSolverCharacterGroundConstraintDataComponent.h"
+#include "Settings/ChaosVDCharacterConstraintsVisualizationSettings.h"
+#include "ToolMenu.h"
+#include "ToolMenus.h"
+#include "ToolMenuEntry.h"
+#include "ToolMenuSection.h"
+#include "Utils/ChaosVDUserInterfaceUtils.h"
 #include "Visualizers/ChaosVDDebugDrawUtils.h"
+#include "Widgets/SChaosVDEnumFlagsMenu.h"
 #include "Widgets/SChaosVDMainTab.h"
+#include "Widgets/SChaosVDViewportToolbar.h"
+
+#define LOCTEXT_NAMESPACE "ChaosVisualDebugger"
+
+bool FChaosVDCharacterGroundConstraintVisualizationDataContext::IsVisualizationFlagEnabled(EChaosVDCharacterGroundConstraintDataVisualizationFlags Flag) const
+{
+	const EChaosVDCharacterGroundConstraintDataVisualizationFlags FlagsAsParticleFlags = static_cast<EChaosVDCharacterGroundConstraintDataVisualizationFlags>(VisualizationFlags);
+	return EnumHasAnyFlags(FlagsAsParticleFlags, Flag);
+}
 
 IMPLEMENT_HIT_PROXY(HChaosVDCharacterGroundConstraintProxy, HComponentVisProxy)
 
 
+FChaosVDCharacterGroundConstraintDataComponentVisualizer::FChaosVDCharacterGroundConstraintDataComponentVisualizer()
+{
+	RegisterVisualizerMenus();
+}
+
+void FChaosVDCharacterGroundConstraintDataComponentVisualizer::RegisterVisualizerMenus()
+{
+		UToolMenus* ToolMenus = UToolMenus::Get();
+	
+	if (!ensure(ToolMenus))
+	{
+		return;
+	}
+
+	if (UToolMenu* Menu = ToolMenus->ExtendMenu(SChaosVDViewportToolbar::ShowMenuName))
+	{
+		FToolMenuSection& Section = Menu->AddSection("CharacterConstraintDataVisualization.Show", LOCTEXT("CharacterConstraintDataVisualizationShowMenuLabel", "Character Constraints Data Visualization"));
+		
+		Section.AddSubMenu(TEXT("CharacterConstraintDataVisualizationFlags"), LOCTEXT("CharacterConstraintDataVisualizationFlagsMenuLabel", "Character Constraints Data Flags"), LOCTEXT("CharacterConstraintDataVisualizationFlagsMenuToolTip", "Set of flags to enable/disable visibility of specific types of Character Constraints data"), FNewToolMenuDelegate::CreateLambda([](UToolMenu* Menu)
+		                   {
+			                   TSharedRef<SWidget> VisualizationFlagsWidget = SNew(SChaosVDEnumFlagsMenu<EChaosVDCharacterGroundConstraintDataVisualizationFlags>)
+				                   .CurrentValue_Static(&UChaosVDCharacterConstraintsVisualizationSettings::GetCharacterGroundConstraintDataVisualizationFlags)
+				                   .OnEnumSelectionChanged_Static(&UChaosVDCharacterConstraintsVisualizationSettings::SetCharacterGroundConstraintDataVisualizationFlags);
+			
+			                   FToolMenuEntry FlagsMenuEntry = FToolMenuEntry::InitWidget("CharacterConstraintDataVisualizationFlags", VisualizationFlagsWidget,FText::GetEmpty());
+			                   Menu->AddMenuEntry(NAME_None, FlagsMenuEntry);
+		                   }),
+		                   false, FSlateIcon(FAppStyle::Get().GetStyleSetName(), TEXT("ClassIcon.Character")));
+
+		using namespace Chaos::VisualDebugger::Utils;
+		Section.AddSubMenu(TEXT("CharacterConstraintDataVisualizationSettings"), LOCTEXT("CharacterConstraintDataVisualizationMenuLabel", "Character Constraints Visualization Settings"), LOCTEXT("CharacterConstraintDataVisualizationMenuToolTip", "Options to change how the recorded Character Constraints data is debug drawn"), FNewToolMenuDelegate::CreateStatic(&CreateMenuEntryForDefaultObject<UChaosVDCharacterConstraintsVisualizationSettings>, EChaosVDSaveSettingsOptions::ShowSaveButton),
+		                   false, FSlateIcon(FAppStyle::Get().GetStyleSetName(), TEXT("Icons.Toolbar.Settings")));
+	}
+}
+
 void FChaosVDCharacterGroundConstraintDataComponentVisualizer::DrawVisualization(const UActorComponent* Component, const FSceneView* View, FPrimitiveDrawInterface* PDI)
 {
 	const UChaosVDSolverCharacterGroundConstraintDataComponent* ConstraintDataComponent = Cast<UChaosVDSolverCharacterGroundConstraintDataComponent>(Component);
@@ -52,12 +103,8 @@ void FChaosVDCharacterGroundConstraintDataComponentVisualizer::DrawVisualization
 	VisualizationContext.SpaceTransform = SolverInfoActor->GetSimulationTransform();
 	VisualizationContext.SolverInfoActor = SolverInfoActor;
 
-	if (const UChaosVDEditorSettings* EditorSettings = GetDefault<UChaosVDEditorSettings>())
-	{
-		VisualizationContext.VisualizationFlags = EditorSettings->GlobalCharacterGroundConstraintDataVisualizationFlags;
-		VisualizationContext.bShowDebugText = EditorSettings->bShowDebugText;
-		VisualizationContext.DebugDrawSettings = &EditorSettings->CharacterGroundConstraintDataDebugDrawSettings;
-	}
+	VisualizationContext.VisualizationFlags = static_cast<uint32>(UChaosVDCharacterConstraintsVisualizationSettings::GetCharacterGroundConstraintDataVisualizationFlags());
+	VisualizationContext.DebugDrawSettings = GetDefault<UChaosVDCharacterConstraintsVisualizationSettings>();
 
 	if (!VisualizationContext.IsVisualizationFlagEnabled(EChaosVDCharacterGroundConstraintDataVisualizationFlags::EnableDraw))
 	{
@@ -128,8 +175,9 @@ void FChaosVDCharacterGroundConstraintDataComponentVisualizer::DrawConstraint(co
 	{
 		return;
 	}
-	
-	if (!VisualizationContext.DebugDrawSettings)
+
+	const UChaosVDCharacterConstraintsVisualizationSettings* DebugDrawSettings =  Cast<UChaosVDCharacterConstraintsVisualizationSettings>(VisualizationContext.DebugDrawSettings);
+	if (!DebugDrawSettings)
 	{
 		return;
 	}
@@ -145,24 +193,17 @@ void FChaosVDCharacterGroundConstraintDataComponentVisualizer::DrawConstraint(co
 	}
 
 	const FChaosVDParticleDataWrapper* CharacterParticleData = nullptr;
-	const FChaosVDParticleDataWrapper* GroundParticleData = nullptr;
 
 	if (AChaosVDParticleActor* CharacterParticle = VisualizationContext.SolverInfoActor->GetParticleActor(InConstraintData.CharacterParticleIndex))
 	{
 		CharacterParticleData = CharacterParticle->GetParticleData();
 	}
 
-	if (AChaosVDParticleActor* GroundParticle = VisualizationContext.SolverInfoActor->GetParticleActor(InConstraintData.GroundParticleIndex))
-	{
-		GroundParticleData = GroundParticle->GetParticleData();
-	}
-
 	if (!CharacterParticleData)
 	{
 		return;
 	}
 
-
 	if (!CharacterParticleData->ParticleMassProps.HasValidData())
 	{
 		// If we don't have mass data, all the following calculations will be off
@@ -172,7 +213,7 @@ void FChaosVDCharacterGroundConstraintDataComponentVisualizer::DrawConstraint(co
 
 	PDI->SetHitProxy(new HChaosVDCharacterGroundConstraintProxy(Component, VisualizationContext.DataSelectionHandle));
 
-	const float LineThickness = InConstraintData.bIsSelectedInEditor ?  VisualizationContext.DebugDrawSettings->BaseLineThickness * 1.5f :  VisualizationContext.DebugDrawSettings->BaseLineThickness;
+	const float LineThickness = InConstraintData.bIsSelectedInEditor ?  DebugDrawSettings->BaseLineThickness * 1.5f :  DebugDrawSettings->BaseLineThickness;
 
 	const FVector CharacterPos = CharacterParticleData->ParticlePositionRotation.MX;
 	const FVector UpDir = InConstraintData.Settings.VerticalAxis;
@@ -186,15 +227,15 @@ void FChaosVDCharacterGroundConstraintDataComponentVisualizer::DrawConstraint(co
 
 	if (VisualizationContext.IsVisualizationFlagEnabled(EChaosVDCharacterGroundConstraintDataVisualizationFlags::TargetDeltaPosition))
 	{
-		FChaosVDDebugDrawUtils::DrawArrowVector(PDI, CharacterPos, CharacterPos + InConstraintData.Data.TargetDeltaPosition, FText::GetEmpty(), FColor::Blue, VisualizationContext.DebugDrawSettings->DepthPriority, 0.5f * LineThickness);
+		FChaosVDDebugDrawUtils::DrawArrowVector(PDI, CharacterPos, CharacterPos + InConstraintData.Data.TargetDeltaPosition, FText::GetEmpty(), FColor::Blue, DebugDrawSettings->DepthPriority, 0.5f * LineThickness);
 	}
 
 	if (VisualizationContext.IsVisualizationFlagEnabled(EChaosVDCharacterGroundConstraintDataVisualizationFlags::TargetDeltaFacing))
 	{
-		const FVector Forward = CharacterParticleData->ParticlePositionRotation.MR * FVector::XAxisVector * VisualizationContext.DebugDrawSettings->GeneralScale * 10.0f;
+		const FVector Forward = CharacterParticleData->ParticlePositionRotation.MR * FVector::XAxisVector * DebugDrawSettings->GeneralScale * 10.0f;
 		const FVector TargetForward = FQuat(UpDir, InConstraintData.Data.TargetDeltaFacing) * Forward;
-		FChaosVDDebugDrawUtils::DrawArrowVector(PDI, CharacterPos, CharacterPos + Forward, FText::GetEmpty(), FColor::Silver, VisualizationContext.DebugDrawSettings->DepthPriority, 0.25f * LineThickness);
-		FChaosVDDebugDrawUtils::DrawArrowVector(PDI, CharacterPos, CharacterPos + TargetForward, FText::GetEmpty(), FColor::White, VisualizationContext.DebugDrawSettings->DepthPriority, 0.25f * LineThickness);
+		FChaosVDDebugDrawUtils::DrawArrowVector(PDI, CharacterPos, CharacterPos + Forward, FText::GetEmpty(), FColor::Silver, DebugDrawSettings->DepthPriority, 0.25f * LineThickness);
+		FChaosVDDebugDrawUtils::DrawArrowVector(PDI, CharacterPos, CharacterPos + TargetForward, FText::GetEmpty(), FColor::White, DebugDrawSettings->DepthPriority, 0.25f * LineThickness);
 	}
 
 	if (VisualizationContext.IsVisualizationFlagEnabled(EChaosVDCharacterGroundConstraintDataVisualizationFlags::GroundQueryDistance))
@@ -203,14 +244,14 @@ void FChaosVDCharacterGroundConstraintDataComponentVisualizer::DrawConstraint(co
 		{
 			if (GroundDistance <= 4.0f * TargetHeight)
 			{
-				FChaosVDDebugDrawUtils::DrawLine(PDI, CharacterPos, CharacterPos - UpDir * TargetHeight, FColor::Green, FText::GetEmpty(), VisualizationContext.DebugDrawSettings->DepthPriority, LineThickness);
-				FChaosVDDebugDrawUtils::DrawLine(PDI, CharacterPos - UpDir * TargetHeight, CharacterPos - UpDir * GroundDistance, FColor::Silver, FText::GetEmpty(), VisualizationContext.DebugDrawSettings->DepthPriority, LineThickness);
+				FChaosVDDebugDrawUtils::DrawLine(PDI, CharacterPos, CharacterPos - UpDir * TargetHeight, FColor::Green, FText::GetEmpty(), DebugDrawSettings->DepthPriority, LineThickness);
+				FChaosVDDebugDrawUtils::DrawLine(PDI, CharacterPos - UpDir * TargetHeight, CharacterPos - UpDir * GroundDistance, FColor::Silver, FText::GetEmpty(), DebugDrawSettings->DepthPriority, LineThickness);
 			}
 		}
 		else
 		{
-			FChaosVDDebugDrawUtils::DrawLine(PDI, CharacterPos, CharacterPos - UpDir * GroundDistance, FColor::Green, FText::GetEmpty(), VisualizationContext.DebugDrawSettings->DepthPriority, LineThickness);
-			FChaosVDDebugDrawUtils::DrawLine(PDI, CharacterPos - UpDir * GroundDistance, CharacterPos - UpDir * TargetHeight, FColor::Red, FText::GetEmpty(), VisualizationContext.DebugDrawSettings->DepthPriority, LineThickness);
+			FChaosVDDebugDrawUtils::DrawLine(PDI, CharacterPos, CharacterPos - UpDir * GroundDistance, FColor::Green, FText::GetEmpty(), DebugDrawSettings->DepthPriority, LineThickness);
+			FChaosVDDebugDrawUtils::DrawLine(PDI, CharacterPos - UpDir * GroundDistance, CharacterPos - UpDir * TargetHeight, FColor::Red, FText::GetEmpty(), DebugDrawSettings->DepthPriority, LineThickness);
 		}
 	}
 
@@ -218,29 +259,31 @@ void FChaosVDCharacterGroundConstraintDataComponentVisualizer::DrawConstraint(co
 	{
 		if (GroundDistance < 4.0f * TargetHeight)
 		{
-			const FVector ScaledGroundNormal = 10.0f * InConstraintData.Data.GroundNormal * VisualizationContext.DebugDrawSettings->GeneralScale;
+			const FVector ScaledGroundNormal = 10.0f * InConstraintData.Data.GroundNormal * DebugDrawSettings->GeneralScale;
 			const FVector GroundPos = CharacterPos - UpDir * GroundDistance;
-			FChaosVDDebugDrawUtils::DrawArrowVector(PDI, GroundPos, GroundPos + ScaledGroundNormal, FText::GetEmpty(), FColor::Cyan, VisualizationContext.DebugDrawSettings->DepthPriority, 0.25f * LineThickness);
+			FChaosVDDebugDrawUtils::DrawArrowVector(PDI, GroundPos, GroundPos + ScaledGroundNormal, FText::GetEmpty(), FColor::Cyan, DebugDrawSettings->DepthPriority, 0.25f * LineThickness);
 		}
 	}
 	
 	if (VisualizationContext.IsVisualizationFlagEnabled(EChaosVDCharacterGroundConstraintDataVisualizationFlags::AppliedNormalForce))
 	{
-		const FVector NormalForce = VisualizationContext.DebugDrawSettings->ForceScale * InConstraintData.Data.GroundNormal.Dot(InConstraintData.State.SolverAppliedForce) * InConstraintData.Data.GroundNormal;
-		FChaosVDDebugDrawUtils::DrawArrowVector(PDI, CharacterPos, CharacterPos + NormalForce, FText::GetEmpty(), VisualizationContext.DebugDrawSettings->NormalForceColor, VisualizationContext.DebugDrawSettings->DepthPriority, LineThickness);
+		const FVector NormalForce = DebugDrawSettings->ForceScale * InConstraintData.Data.GroundNormal.Dot(InConstraintData.State.SolverAppliedForce) * InConstraintData.Data.GroundNormal;
+		FChaosVDDebugDrawUtils::DrawArrowVector(PDI, CharacterPos, CharacterPos + NormalForce, FText::GetEmpty(), DebugDrawSettings->NormalForceColor, DebugDrawSettings->DepthPriority, LineThickness);
 	}
 
 	if (VisualizationContext.IsVisualizationFlagEnabled(EChaosVDCharacterGroundConstraintDataVisualizationFlags::AppliedRadialForce))
 	{
-		const FVector RadialForce = VisualizationContext.DebugDrawSettings->ForceScale * (InConstraintData.State.SolverAppliedForce - InConstraintData.Data.GroundNormal.Dot(InConstraintData.State.SolverAppliedForce) * InConstraintData.Data.GroundNormal);
-		FChaosVDDebugDrawUtils::DrawArrowVector(PDI, CharacterPos, CharacterPos + RadialForce, FText::GetEmpty(), VisualizationContext.DebugDrawSettings->NormalForceColor, VisualizationContext.DebugDrawSettings->DepthPriority, LineThickness);
+		const FVector RadialForce = DebugDrawSettings->ForceScale * (InConstraintData.State.SolverAppliedForce - InConstraintData.Data.GroundNormal.Dot(InConstraintData.State.SolverAppliedForce) * InConstraintData.Data.GroundNormal);
+		FChaosVDDebugDrawUtils::DrawArrowVector(PDI, CharacterPos, CharacterPos + RadialForce, FText::GetEmpty(), DebugDrawSettings->NormalForceColor, DebugDrawSettings->DepthPriority, LineThickness);
 	}
 
 	if (VisualizationContext.IsVisualizationFlagEnabled(EChaosVDCharacterGroundConstraintDataVisualizationFlags::AppliedTorque))
 	{
-		const FVector Torque = VisualizationContext.DebugDrawSettings->TorqueScale * InConstraintData.State.SolverAppliedTorque;
-		FChaosVDDebugDrawUtils::DrawArrowVector(PDI, CharacterPos, CharacterPos + Torque, FText::GetEmpty(), VisualizationContext.DebugDrawSettings->TorqueColor, VisualizationContext.DebugDrawSettings->DepthPriority, LineThickness);
+		const FVector Torque = DebugDrawSettings->TorqueScale * InConstraintData.State.SolverAppliedTorque;
+		FChaosVDDebugDrawUtils::DrawArrowVector(PDI, CharacterPos, CharacterPos + Torque, FText::GetEmpty(), DebugDrawSettings->TorqueColor, DebugDrawSettings->DepthPriority, LineThickness);
 	}
 
 	PDI->SetHitProxy(nullptr);
 }
+
+#undef LOCTEXT_NAMESPACE
diff --git a/Engine/Plugins/ChaosVD/Source/ChaosVD/Private/Visualizers/ChaosVDCharacterGroundConstraintsDataComponentVisualizer.h b/Engine/Plugins/ChaosVD/Source/ChaosVD/Private/Visualizers/ChaosVDCharacterGroundConstraintsDataComponentVisualizer.h
--- a/Engine/Plugins/ChaosVD/Source/ChaosVD/Private/Visualizers/ChaosVDCharacterGroundConstraintsDataComponentVisualizer.h
+++ b/Engine/Plugins/ChaosVD/Source/ChaosVD/Private/Visualizers/ChaosVDCharacterGroundConstraintsDataComponentVisualizer.h
@@ -6,54 +6,19 @@
 #include "Chaos/Core.h"
 #include "Components/ChaosVDSolverCharacterGroundConstraintDataComponent.h"
 
-struct FChaosVDCharacterGroundConstraintDebugDrawSettings;
 class AChaosVDSolverInfoActor;
+enum class EChaosVDCharacterGroundConstraintDataVisualizationFlags : uint32;
+struct FChaosVDCharacterGroundConstraintDebugDrawSettings;
 struct FChaosVDCharacterGroundConstraint;
 
-UENUM(meta = (Bitflags, UseEnumValuesAsMaskValuesInEditor = "true"))
-enum class EChaosVDCharacterGroundConstraintDataVisualizationFlags : uint32
-{
-	None					= 0 UMETA(Hidden),
-	/** Draw the target movement vector */
-	TargetDeltaPosition		= 1 << 0,
-	/** Draw the target orientation facing vector */
-	TargetDeltaFacing		= 1 << 1,
-	/** Draw the ground query distance based on the constraint's data */
-	GroundQueryDistance		= 1 << 2,
-	/** Draw the ground query normal based on the constraint's data */
-	GroundQueryNormal		= 1 << 3,
-	/** Draw the applied force vector */
-	AppliedRadialForce		= 1 << 4,
-	/** Draw the applied force vector */
-	AppliedNormalForce		= 1 << 5,
-	/** Draw the applied force vector */
-	AppliedTorque			= 1 << 6,
-	/** Draw the constraint even if it is disabled */
-	DrawDisabled			= 1 << 7,
-	/** Only debugs draw data for a selected constraint */
-	OnlyDrawSelected		= 1 << 8,
-	/** Enables debug draw for constraint data from any solver */
-	EnableDraw				= 1 << 9,
-};
-ENUM_CLASS_FLAGS(EChaosVDCharacterGroundConstraintDataVisualizationFlags);
-
 /** Visualization context structure specific for character ground constraint visualizations */
 struct FChaosVDCharacterGroundConstraintVisualizationDataContext : public FChaosVDVisualizationContext
 {
 	FChaosVDCharacterGroundConstraintSelectionHandle DataSelectionHandle = FChaosVDCharacterGroundConstraintSelectionHandle(nullptr);
-	bool bIsServerVisualizationEnabled = false;
 	
 	AChaosVDSolverInfoActor* SolverInfoActor = nullptr;
 
-	const FChaosVDCharacterGroundConstraintDebugDrawSettings* DebugDrawSettings = nullptr;
-
-	bool bShowDebugText = false;
-
-	bool IsVisualizationFlagEnabled(EChaosVDCharacterGroundConstraintDataVisualizationFlags Flag) const
-	{
-		const EChaosVDCharacterGroundConstraintDataVisualizationFlags FlagsAsParticleFlags = static_cast<EChaosVDCharacterGroundConstraintDataVisualizationFlags>(VisualizationFlags);
-		return EnumHasAnyFlags(FlagsAsParticleFlags, Flag);
-	}
+	bool IsVisualizationFlagEnabled(EChaosVDCharacterGroundConstraintDataVisualizationFlags Flag) const;
 };
 
 /** Custom Hit Proxy for debug drawn scene queries */
@@ -79,9 +44,9 @@ struct HChaosVDCharacterGroundConstraintProxy : public HComponentVisProxy
 class FChaosVDCharacterGroundConstraintDataComponentVisualizer final : public FComponentVisualizer
 {
 public:
-	FChaosVDCharacterGroundConstraintDataComponentVisualizer()
-	{
-	}
+	FChaosVDCharacterGroundConstraintDataComponentVisualizer();
+	
+	void RegisterVisualizerMenus();
 
 	virtual void DrawVisualization(const UActorComponent* Component, const FSceneView* View, FPrimitiveDrawInterface* PDI) override;
 	virtual bool VisProxyHandleClick(FEditorViewportClient* InViewportClient, HComponentVisProxy* VisProxy, const FViewportClick& Click) override;
diff --git a/Engine/Plugins/ChaosVD/Source/ChaosVD/Private/Visualizers/ChaosVDDebugDrawUtils.cpp b/Engine/Plugins/ChaosVD/Source/ChaosVD/Private/Visualizers/ChaosVDDebugDrawUtils.cpp
--- a/Engine/Plugins/ChaosVD/Source/ChaosVD/Private/Visualizers/ChaosVDDebugDrawUtils.cpp
+++ b/Engine/Plugins/ChaosVD/Source/ChaosVD/Private/Visualizers/ChaosVDDebugDrawUtils.cpp
@@ -3,7 +3,6 @@
 #include "Visualizers/ChaosVDDebugDrawUtils.h"
 
 #include "CanvasItem.h"
-#include "ChaosVDEditorSettings.h"
 #include "ChaosVDGeometryBuilder.h"
 #include "Generators/MeshShapeGenerator.h"
 #include "Engine/Engine.h"
@@ -93,23 +92,17 @@ void FChaosVDDebugDrawUtils::DrawPoint(FPrimitiveDrawInterface* PDI, const FVect
 
 void FChaosVDDebugDrawUtils::DrawString(FStringView StringToDraw, const FVector& Location, const FColor& Color, EChaosVDDebugDrawTextLocationMode LocationMode)
 {
-	if (const UChaosVDEditorSettings* CVDEditorSettings = GetDefault<UChaosVDEditorSettings>())
+	if (!StringToDraw.IsEmpty())
 	{
-		if (CVDEditorSettings->bShowDebugText)
-		{
-			TextToDrawQueue.Enqueue({ FText::AsCultureInvariant(StringToDraw.GetData()), LocationMode, Location, Color });
-		}
+		TextToDrawQueue.Enqueue({ FText::AsCultureInvariant(StringToDraw.GetData()), LocationMode, Location, Color });
 	}
 }
 
 void FChaosVDDebugDrawUtils::DrawText(const FText& InText, const FVector& Location, const FColor& Color, EChaosVDDebugDrawTextLocationMode LocationMode)
 {
-	if (const UChaosVDEditorSettings* CVDEditorSettings = GetDefault<UChaosVDEditorSettings>())
+	if (!InText.IsEmptyOrWhitespace())
 	{
-		if (CVDEditorSettings->bShowDebugText)
-		{
-			TextToDrawQueue.Enqueue({ InText, LocationMode, Location, Color });
-		}
+		TextToDrawQueue.Enqueue({ InText, LocationMode, Location, Color });
 	}
 }
 
diff --git a/Engine/Plugins/ChaosVD/Source/ChaosVD/Private/Visualizers/ChaosVDJointConstraintsDataComponentVisualizer.cpp b/Engine/Plugins/ChaosVD/Source/ChaosVD/Private/Visualizers/ChaosVDJointConstraintsDataComponentVisualizer.cpp
--- a/Engine/Plugins/ChaosVD/Source/ChaosVD/Private/Visualizers/ChaosVDJointConstraintsDataComponentVisualizer.cpp
+++ b/Engine/Plugins/ChaosVD/Source/ChaosVD/Private/Visualizers/ChaosVDJointConstraintsDataComponentVisualizer.cpp
@@ -3,6 +3,7 @@
 #include "Visualizers/ChaosVDJointConstraintsDataComponentVisualizer.h"
 
 #include "ChaosVDScene.h"
+#include "ChaosVDStyle.h"
 #include "ChaosVDTabsIDs.h"
 #include "EditorModeManager.h"
 #include "EditorViewportClient.h"
@@ -10,11 +11,21 @@
 #include "SceneView.h"
 #include "Actors/ChaosVDSolverInfoActor.h"
 #include "Components/ChaosVDSolverJointConstraintDataComponent.h"
+#include "Settings/ChaosVDJointConstraintVisualizationSettings.h"
+#include "ToolMenu.h"
+#include "ToolMenus.h"
+#include "ToolMenuEntry.h"
+#include "ToolMenuSection.h"
+#include "Utils/ChaosVDUserInterfaceUtils.h"
 #include "Visualizers/ChaosVDDebugDrawUtils.h"
 #include "Widgets/SChaosVDMainTab.h"
+#include "Widgets/SChaosVDViewportToolbar.h"
+#include "Widgets/SChaosVDEnumFlagsMenu.h"
 
 IMPLEMENT_HIT_PROXY(HChaosVDJointConstraintProxy, HComponentVisProxy)
 
+#define LOCTEXT_NAMESPACE "ChaosVisualDebugger"
+
 namespace Chaos::VisualDebugger::Utils
 {
 	FVector GetCoMWorldPosition(const FChaosVDParticleDataWrapper& ParticleData)
@@ -97,6 +108,42 @@ namespace Chaos::VisualDebugger::Utils
 	}
 }
 
+FChaosVDJointConstraintsDataComponentVisualizer::FChaosVDJointConstraintsDataComponentVisualizer()
+{
+	RegisterVisualizerMenus();
+}
+
+void FChaosVDJointConstraintsDataComponentVisualizer::RegisterVisualizerMenus()
+{
+	UToolMenus* ToolMenus = UToolMenus::Get();
+	
+	if (!ensure(ToolMenus))
+	{
+		return;
+	}
+
+	if (UToolMenu* Menu = ToolMenus->ExtendMenu(SChaosVDViewportToolbar::ShowMenuName))
+	{
+		FToolMenuSection& Section = Menu->AddSection("JointConstraintDataVisualization.Show", LOCTEXT("JointConstraintDataVisualizationShowMenuLabel", "Joint Constraint Data Visualization"));
+		
+		Section.AddSubMenu(TEXT("JointConstraintDataVisualizationFlags"), LOCTEXT("JointConstraintDataVisualizationFlagsMenuLabel", "Joint Constraints Data Flags"), LOCTEXT("JointConstraintDataVisualizationFlagsMenuToolTip", "Set of flags to enable/disable visibility of specific types of joint constraint data"), FNewToolMenuDelegate::CreateLambda([](UToolMenu* Menu)
+		                   {
+			                   TSharedRef<SWidget> VisualizationFlagsWidget = SNew(SChaosVDEnumFlagsMenu<EChaosVDJointsDataVisualizationFlags>)
+				                   .CurrentValue_Static(&UChaosVDJointConstraintsVisualizationSettings::GetJointsDataVisualizationFlags)
+				                   .OnEnumSelectionChanged_Static(&UChaosVDJointConstraintsVisualizationSettings::SetJointsDataVisualizationFlags);
+			
+			                   FToolMenuEntry FlagsMenuEntry = FToolMenuEntry::InitWidget("JointConstraintDataVisualizationFlags", VisualizationFlagsWidget,FText::GetEmpty());
+			                   Menu->AddMenuEntry(NAME_None, FlagsMenuEntry);
+		                   }),
+		                   false, FSlateIcon(FChaosVDStyle::Get().GetStyleSetName(), TEXT("ConnectionIcon")));
+
+		using namespace Chaos::VisualDebugger::Utils;
+		Section.AddSubMenu(TEXT("JointConstraintDataVisualizationSettings"), LOCTEXT("JointConstraintDataVisualizationMenuLabel", "Joint Constraint Visualization Settings"), LOCTEXT("JointConstraintDataVisualizationMenuToolTip", "Options to change how the recorded joint constraint data is debug drawn"), FNewToolMenuDelegate::CreateStatic(&CreateMenuEntryForDefaultObject<UChaosVDJointConstraintsVisualizationSettings>, EChaosVDSaveSettingsOptions::ShowSaveButton),
+		                   false, FSlateIcon(FAppStyle::Get().GetStyleSetName(), TEXT("Icons.Toolbar.Settings")));
+	}
+}
+
+
 void FChaosVDJointConstraintsDataComponentVisualizer::DrawVisualization(const UActorComponent* Component, const FSceneView* View, FPrimitiveDrawInterface* PDI)
 {
 	const UChaosVDSolverJointConstraintDataComponent* JointConstraintDataComponent = Cast<UChaosVDSolverJointConstraintDataComponent>(Component);
@@ -133,11 +180,11 @@ void FChaosVDJointConstraintsDataComponentVisualizer::DrawVisualization(const UA
 	VisualizationContext.SpaceTransform = SolverInfoActor->GetSimulationTransform();
 	VisualizationContext.SolverInfoActor = SolverInfoActor;
 
-	if (const UChaosVDEditorSettings* EditorSettings = GetDefault<UChaosVDEditorSettings>())
+	if (const UChaosVDJointConstraintsVisualizationSettings* EditorSettings = GetDefault<UChaosVDJointConstraintsVisualizationSettings>())
 	{
-		VisualizationContext.VisualizationFlags = EditorSettings->GlobalJointsDataVisualizationFlags;
+		VisualizationContext.VisualizationFlags = static_cast<uint32>(UChaosVDJointConstraintsVisualizationSettings::GetJointsDataVisualizationFlags());
 		VisualizationContext.bShowDebugText = EditorSettings->bShowDebugText;
-		VisualizationContext.DebugDrawSettings = &EditorSettings->JointsDataDebugDrawSettings;
+		VisualizationContext.DebugDrawSettings = EditorSettings;
 	}
 
 	if (!VisualizationContext.IsVisualizationFlagEnabled(EChaosVDJointsDataVisualizationFlags::EnableDraw))
@@ -197,7 +244,7 @@ void FChaosVDJointConstraintsDataComponentVisualizer::DebugDrawAllAxis(const FCh
 {
 	for(int32 AxisIndex = 0; AxisIndex < 3 ; AxisIndex++)
 	{
-		FChaosVDDebugDrawUtils::DrawArrowVector(PDI, InPosition, InPosition + VisualizationContext.DebugDrawSettings->GeneralScale * VisualizationContext.DebugDrawSettings->ConstraintAxisLength * VisualizationContext.SpaceTransform.TransformVector(InRotationMatrix.GetAxis(AxisIndex)),FText::GetEmpty(), Chaos::VisualDebugger::Utils::GenerateSelectionAwareDebugColor(AxisColors[AxisIndex], InJointConstraintData.bIsSelectedInEditor), VisualizationContext.DebugDrawSettings->DepthPriority, LineThickness * 0.2f);
+		FChaosVDDebugDrawUtils::DrawArrowVector(PDI, InPosition, InPosition + VisualizationContext.DebugDrawSettings->GeneralScale * VisualizationContext.DebugDrawSettings->ConstraintAxisLength * VisualizationContext.SpaceTransform.TransformVector(InRotationMatrix.GetAxis(AxisIndex)), FText::GetEmpty(), Chaos::VisualDebugger::Utils::GenerateSelectionAwareDebugColor(AxisColors[AxisIndex], InJointConstraintData.bIsSelectedInEditor), VisualizationContext.DebugDrawSettings->DepthPriority, LineThickness * 0.2f);
 	}
 }
 
@@ -378,3 +425,5 @@ void FChaosVDJointConstraintsDataComponentVisualizer::DrawJointConstraint(const
 	
 	PDI->SetHitProxy(nullptr);
 }
+
+#undef LOCTEXT_NAMESPACE
diff --git a/Engine/Plugins/ChaosVD/Source/ChaosVD/Private/Visualizers/ChaosVDJointConstraintsDataComponentVisualizer.h b/Engine/Plugins/ChaosVD/Source/ChaosVD/Private/Visualizers/ChaosVDJointConstraintsDataComponentVisualizer.h
--- a/Engine/Plugins/ChaosVD/Source/ChaosVD/Private/Visualizers/ChaosVDJointConstraintsDataComponentVisualizer.h
+++ b/Engine/Plugins/ChaosVD/Source/ChaosVD/Private/Visualizers/ChaosVDJointConstraintsDataComponentVisualizer.h
@@ -2,36 +2,15 @@
 
 #pragma once
 #include "ComponentVisualizer.h"
-#include "IChaosVDParticleVisualizationDataProvider.h"
 #include "Chaos/Core.h"
 #include "Components/ChaosVDSolverJointConstraintDataComponent.h"
+#include "IChaosVDParticleVisualizationDataProvider.h"
+#include "Settings/ChaosVDJointConstraintVisualizationSettings.h"
 
-struct FChaosVDJointsDebugDrawSettings;
-class AChaosVDSolverInfoActor;
 struct FChaosVDJointConstraint;
 
-UENUM(meta = (Bitflags, UseEnumValuesAsMaskValuesInEditor = "true"))
-enum class EChaosVDJointsDataVisualizationFlags : uint32
-{
-	None					= 0 UMETA(Hidden),
-	/** Draw the PushOut vector based on the constraint's data */
-	PushOut					= 1 << 0,
-	/** Draw the Angular Impulse vector based on the constraint's data */
-	AngularImpulse			= 1 << 1 UMETA(Hidden),
-	ActorConnector			= 1 << 2,
-	CenterOfMassConnector	= 1 << 3,
-	Stretch					= 1 << 4,
-	Axes					= 1 << 5,
-	/** Draw the joint even if one of the particles or both are kinematic */
-	DrawKinematic			= 1 << 6,
-	/** Draw the joint even if it is disabled */
-	DrawDisabled			= 1 << 7,
-	/** Only debugs draw data for a selected joint constraint */
-	OnlyDrawSelected		= 1 << 8,
-	/** Enables Debug draw for Joint Constraint data from any solver that is visible */
-	EnableDraw				= 1 << 9,
-};
-ENUM_CLASS_FLAGS(EChaosVDJointsDataVisualizationFlags);
+class AChaosVDSolverInfoActor;
+class UChaosVDJointConstraintsVisualizationSettings;
 
 /** Visualization context structure specific for Joints visualizations */
 struct FChaosVDJointVisualizationDataContext : public FChaosVDVisualizationContext
@@ -41,7 +20,7 @@ struct FChaosVDJointVisualizationDataContext : public FChaosVDVisualizationConte
 	
 	AChaosVDSolverInfoActor* SolverInfoActor = nullptr;
 
-	const FChaosVDJointsDebugDrawSettings* DebugDrawSettings = nullptr;
+	const UChaosVDJointConstraintsVisualizationSettings* DebugDrawSettings = nullptr;
 
 	bool bShowDebugText = false;
 
@@ -75,9 +54,9 @@ struct HChaosVDJointConstraintProxy : public HComponentVisProxy
 class FChaosVDJointConstraintsDataComponentVisualizer final : public FComponentVisualizer
 {
 public:
-	FChaosVDJointConstraintsDataComponentVisualizer()
-	{
-	}
+	FChaosVDJointConstraintsDataComponentVisualizer();
+	
+	void RegisterVisualizerMenus();
 
 	virtual void DrawVisualization(const UActorComponent* Component, const FSceneView* View, FPrimitiveDrawInterface* PDI) override;
 	virtual bool VisProxyHandleClick(FEditorViewportClient* InViewportClient, HComponentVisProxy* VisProxy, const FViewportClick& Click) override;
diff --git a/Engine/Plugins/ChaosVD/Source/ChaosVD/Private/Visualizers/ChaosVDParticleDataComponentVisualizer.cpp b/Engine/Plugins/ChaosVD/Source/ChaosVD/Private/Visualizers/ChaosVDParticleDataComponentVisualizer.cpp
--- a/Engine/Plugins/ChaosVD/Source/ChaosVD/Private/Visualizers/ChaosVDParticleDataComponentVisualizer.cpp
+++ b/Engine/Plugins/ChaosVD/Source/ChaosVD/Private/Visualizers/ChaosVDParticleDataComponentVisualizer.cpp
@@ -3,13 +3,20 @@
 #include "Visualizers/ChaosVDParticleDataComponentVisualizer.h"
 
 #include "Actors/ChaosVDSolverInfoActor.h"
-#include "ChaosVDEditorSettings.h"
 #include "ChaosVDParticleActor.h"
 #include "ChaosVDScene.h"
-#include "SceneView.h"
 #include "Components/ChaosVDParticleDataComponent.h"
 #include "DataWrappers/ChaosVDParticleDataWrapper.h"
+#include "SceneView.h"
+#include "Settings/ChaosVDParticleVisualizationSettings.h"
+#include "ToolMenu.h"
+#include "ToolMenus.h"
+#include "ToolMenuEntry.h"
+#include "ToolMenuSection.h"
+#include "Utils/ChaosVDUserInterfaceUtils.h"
 #include "Visualizers/ChaosVDDebugDrawUtils.h"
+#include "Widgets/SChaosVDEnumFlagsMenu.h"
+#include "Widgets/SChaosVDViewportToolbar.h"
 
 IMPLEMENT_HIT_PROXY(HChaosVDParticleDataProxy, HComponentVisProxy)
 
@@ -38,15 +45,78 @@ struct FChaosVDScopedParticleHitProxy
 	FPrimitiveDrawInterface* PDIPtr = nullptr;
 };
 
+FChaosVDParticleDataComponentVisualizer::FChaosVDParticleDataComponentVisualizer()
+{
+	RegisterVisualizerMenus();
+}
+
+void FChaosVDParticleDataComponentVisualizer::RegisterVisualizerMenus()
+{
+	UToolMenus* ToolMenus = UToolMenus::Get();
+	
+	if (!ensure(ToolMenus))
+	{
+		return;
+	}
+
+	if (UToolMenu* Menu = ToolMenus->ExtendMenu(SChaosVDViewportToolbar::ShowMenuName))
+	{
+		FToolMenuSection& Section = Menu->AddSection("ParticleVisualization.Show", LOCTEXT("ParticleVisualizationShowMenuLabel", "Particle Visualization"));
+
+		FNewToolMenuDelegate GeometryVisualizationFlagsMenuBuilder = FNewToolMenuDelegate::CreateLambda([](UToolMenu* Menu)
+		{
+			if (Menu)
+			{
+				TSharedRef<SWidget> VisualizationFlagsWidget = SNew(SChaosVDEnumFlagsMenu<EChaosVDGeometryVisibilityFlags>)
+					.CurrentValue_Static(&UChaosVDParticleVisualizationSettings::GetGeometryVisualizationFlags)
+					.OnEnumSelectionChanged_Lambda(&UChaosVDParticleVisualizationSettings::SetGeometryVisualizationFlags);
+
+				Menu->AddMenuEntry(NAME_None, FToolMenuEntry::InitWidget("GeometryVisualizationFlags",VisualizationFlagsWidget,FText::GetEmpty()));
+			}
+		});
+
+		FNewToolMenuDelegate ParticleDataVisualizationFlagsMenuBuilder = FNewToolMenuDelegate::CreateLambda([](UToolMenu* Menu)
+		{
+			if (Menu)
+			{
+				TSharedRef<SWidget> VisualizationFlagsWidget = SNew(SChaosVDEnumFlagsMenu<EChaosVDParticleDataVisualizationFlags>)
+					.CurrentValue_Static(&UChaosVDParticleVisualizationDebugDrawSettings::GetDataDebugDrawVisualizationFlags)
+					.OnEnumSelectionChanged_Static(&UChaosVDParticleVisualizationDebugDrawSettings::SetDataDebugDrawVisualizationFlags);
+
+				Menu->AddMenuEntry(NAME_None, FToolMenuEntry::InitWidget("ParticleDebugDrawDataVisualizationFlags", VisualizationFlagsWidget,FText::GetEmpty()));
+			}
+		});
+		
+		FNewToolMenuDelegate ParticleDataVisualizationSettingsMenuBuilder = FNewToolMenuDelegate::CreateLambda([](UToolMenu* Menu)
+		{
+			using namespace Chaos::VisualDebugger::Utils;
+			CreateMenuEntryForDefaultObject<UChaosVDParticleVisualizationDebugDrawSettings>(Menu, EChaosVDSaveSettingsOptions::ShowSaveButton);
+		});
+
+		FNewToolMenuDelegate ParticleColorizationMenuBuilder = FNewToolMenuDelegate::CreateLambda([](UToolMenu* Menu)
+		{
+			using namespace Chaos::VisualDebugger::Utils;
+			CreateMenuEntryForDefaultObject<UChaosVDParticleVisualizationColorSettings>(Menu, EChaosVDSaveSettingsOptions::ShowSaveButton);
+		});
+
+		constexpr bool bOpenSubMenuOnClick = false;
+		
+		Section.AddSubMenu(TEXT("GeometryVisualizationFlags"), LOCTEXT("GeometryVisualizationFlagsMenuLabel", "Geometry Flags"), LOCTEXT("GeometryVisualizationFlagsMenuToolTip", "Set of flags to enable/disable visibility of specific types of geometry/particles"), GeometryVisualizationFlagsMenuBuilder, bOpenSubMenuOnClick, FSlateIcon(FAppStyle::Get().GetStyleSetName(), TEXT("ShowFlagsMenu.StaticMeshes")));
+		Section.AddSubMenu(TEXT("ParticleDataVisualizationFlags"), LOCTEXT("ParticleDataVisualizationFlagsMenuLabel", "Particle Data Flags"), LOCTEXT("ParticleDataVisualizationFlagsMenuToolTip", "Set of flags to enable/disable visualization of specific particle data as debug draw"), ParticleDataVisualizationFlagsMenuBuilder, bOpenSubMenuOnClick, FSlateIcon(FAppStyle::Get().GetStyleSetName(), TEXT("StaticMeshEditor.SetDrawAdditionalData")));
+		Section.AddSubMenu(TEXT("ParticleDataVisualizationSettings"), LOCTEXT("ParticleDataVisualizationSettingsMenuLabel", "Particle Data Visualization Settings"), LOCTEXT("ParticleDataVisualizationSettingsMenuToolTip", "Options to control how particle data is debug drawn"), ParticleDataVisualizationSettingsMenuBuilder, bOpenSubMenuOnClick, FSlateIcon(FAppStyle::Get().GetStyleSetName(), TEXT("Icons.Toolbar.Settings")));
+		Section.AddSubMenu(TEXT("ParticleColorizationFlags"), LOCTEXT("ParticleColorizationOptionsMenuLabel", "Particle Colorization"), LOCTEXT("Particle ColorizationMenuToolTip", "Changes what colors are used to draw the particles and its data"), ParticleColorizationMenuBuilder, bOpenSubMenuOnClick, FSlateIcon(FAppStyle::Get().GetStyleSetName(), TEXT("ColorPicker.ColorThemes")));
+	}
+}
+
 void FChaosVDParticleDataComponentVisualizer::DrawVisualization(const UActorComponent* Component, const FSceneView* View, FPrimitiveDrawInterface* PDI)
 {
-	const UChaosVDEditorSettings* EditorSettings = GetDefault<UChaosVDEditorSettings>();
-	if (!EditorSettings)
+	const UChaosVDParticleVisualizationDebugDrawSettings* VisualizationSettings = GetDefault<UChaosVDParticleVisualizationDebugDrawSettings>();
+	if (!VisualizationSettings)
 	{
 		return;
 	}
 
-	if (EditorSettings->GlobalParticleDataVisualizationFlags == 0)
+	if (VisualizationSettings->GetDataDebugDrawVisualizationFlags() == EChaosVDParticleDataVisualizationFlags::None)
 	{
 		// Nothing to visualize
 		return;
@@ -76,12 +146,12 @@ void FChaosVDParticleDataComponentVisualizer::DrawVisualization(const UActorComp
 	}
 
 	FChaosVDParticleDataVisualizationContext VisualizationContext;
-	VisualizationContext.VisualizationFlags = EditorSettings->GlobalParticleDataVisualizationFlags;
+	VisualizationContext.VisualizationFlags = static_cast<uint32>(VisualizationSettings->GetDataDebugDrawVisualizationFlags());
 	VisualizationContext.SpaceTransform = SolverDataActor->GetSimulationTransform();
 	VisualizationContext.CVDScene = CVDScene;
 	VisualizationContext.GeometryGenerator = CVDScene->GetGeometryGenerator();
-	VisualizationContext.bShowDebugText = EditorSettings->bShowDebugText;
-	VisualizationContext.DebugDrawSettings = &EditorSettings->ParticleDataDebugDrawSettings;
+	VisualizationContext.bShowDebugText = VisualizationSettings->bShowDebugText;
+	VisualizationContext.DebugDrawSettings = VisualizationSettings;
 
 	if (!VisualizationContext.IsVisualizationFlagEnabled(EChaosVDParticleDataVisualizationFlags::EnableDraw))
 	{
diff --git a/Engine/Plugins/ChaosVD/Source/ChaosVD/Private/Visualizers/ChaosVDParticleDataComponentVisualizer.h b/Engine/Plugins/ChaosVD/Source/ChaosVD/Private/Visualizers/ChaosVDParticleDataComponentVisualizer.h
--- a/Engine/Plugins/ChaosVD/Source/ChaosVD/Private/Visualizers/ChaosVDParticleDataComponentVisualizer.h
+++ b/Engine/Plugins/ChaosVD/Source/ChaosVD/Private/Visualizers/ChaosVDParticleDataComponentVisualizer.h
@@ -3,38 +3,21 @@
 #pragma once
 
 #include "ComponentVisualizer.h"
-#include "IChaosVDParticleVisualizationDataProvider.h"
-#include "DataWrappers/ChaosVDParticleDataWrapper.h"
 #include "HAL/Platform.h"
+#include "IChaosVDParticleVisualizationDataProvider.h"
+#include "Settings/ChaosVDParticleVisualizationSettings.h"
 #include "Templates/SharedPointer.h"
 
 class FChaosVDGeometryBuilder;
 struct FChaosParticleDataDebugDrawSettings;
 
-UENUM(meta = (Bitflags, UseEnumValuesAsMaskValuesInEditor = "true"))
-enum class EChaosVDParticleDataVisualizationFlags : uint32
-{
-	None				= 0 UMETA(Hidden),
-	Velocity			= 1 << 0,
-	AngularVelocity		= 1 << 1,
-	Acceleration		= 1 << 2,
-	AngularAcceleration = 1 << 3,
-	LinearImpulse		= 1 << 4,
-	AngularImpulse		= 1 << 5,
-	ClusterConnectivityEdge	= 1 << 6,
-	CenterOfMass	= 1 << 7,
-	DrawDataOnlyForSelectedParticle	= 1 << 8,
-	EnableDraw	= 1 << 9,
-};
-ENUM_CLASS_FLAGS(EChaosVDParticleDataVisualizationFlags);
-
 struct FChaosVDVisualizedParticleDataSelectionHandle
 {
 	int32 ParticleIndex = INDEX_NONE;
 	int32 SolverID = INDEX_NONE;
 };
 
-struct FChaosParticleDataDebugDrawSettings;
+class UChaosVDParticleVisualizationDebugDrawSettings;
 
 struct FChaosVDParticleDataVisualizationContext : public FChaosVDVisualizationContext
 {
@@ -42,7 +25,7 @@ struct FChaosVDParticleDataVisualizationContext : public FChaosVDVisualizationCo
 	bool bIsSelectedData = false;
 	bool bShowDebugText = false;
 
-	const FChaosParticleDataDebugDrawSettings* DebugDrawSettings = nullptr;
+	const UChaosVDParticleVisualizationDebugDrawSettings* DebugDrawSettings = nullptr;
 
 	bool IsVisualizationFlagEnabled(EChaosVDParticleDataVisualizationFlags Flag) const
 	{
@@ -51,42 +34,6 @@ struct FChaosVDParticleDataVisualizationContext : public FChaosVDVisualizationCo
 	}
 };
 
-namespace Chaos::VisualDebugger::ParticleDataUnitsStrings
-{
-	static FString Velocity = TEXT("cm/s");
-	static FString AngularVelocity = TEXT("rad/s");
-	static FString Acceleration = TEXT("cm/s2");
-	static FString AngularAcceleration = TEXT("rad/s2");
-	static FString LinearImpulse = TEXT("g.m/s");
-	static FString AngularImpulse = TEXT("g.m2/s");
-
-	static FString GetUnitByID(EChaosVDParticleDataVisualizationFlags DataID)
-	{
-		switch (DataID)
-		{
-		case EChaosVDParticleDataVisualizationFlags::Velocity:
-			return Velocity;
-		case EChaosVDParticleDataVisualizationFlags::AngularVelocity:
-			return AngularVelocity;
-		case EChaosVDParticleDataVisualizationFlags::Acceleration:
-			return Acceleration;
-		case EChaosVDParticleDataVisualizationFlags::AngularAcceleration:
-			return AngularAcceleration;
-		case EChaosVDParticleDataVisualizationFlags::LinearImpulse:
-			return LinearImpulse;
-		case EChaosVDParticleDataVisualizationFlags::AngularImpulse:
-			return AngularImpulse;
-
-		case EChaosVDParticleDataVisualizationFlags::None:
-		case EChaosVDParticleDataVisualizationFlags::ClusterConnectivityEdge:
-		case EChaosVDParticleDataVisualizationFlags::CenterOfMass:
-		case EChaosVDParticleDataVisualizationFlags::DrawDataOnlyForSelectedParticle:
-		default:
-			return TEXT("");
-		}
-	}
-}
-
 /** Custom Hit Proxy for debug drawn particle data */
 struct HChaosVDParticleDataProxy : public HComponentVisProxy
 {
@@ -110,10 +57,15 @@ struct HChaosVDParticleDataProxy : public HComponentVisProxy
 class FChaosVDParticleDataComponentVisualizer : public FComponentVisualizer
 {
 public:
+	FChaosVDParticleDataComponentVisualizer();
+
+	void RegisterVisualizerMenus();
+
 	virtual void DrawVisualization(const UActorComponent* Component, const FSceneView* View, FPrimitiveDrawInterface* PDI) override;
 	virtual bool VisProxyHandleClick(FEditorViewportClient* InViewportClient, HComponentVisProxy* VisProxy, const FViewportClick& Click) override;
 
 protected:
+	
 	void DrawParticleVector(FPrimitiveDrawInterface* PDI, const FVector& StartLocation, const FVector& InVector, EChaosVDParticleDataVisualizationFlags VectorID, const FChaosVDParticleDataVisualizationContext& InVisualizationContext, float LineThickness);
 	void DrawVisualizationForParticleData(const UActorComponent* Component, FPrimitiveDrawInterface* PDI, const FSceneView* View, const FChaosVDParticleDataVisualizationContext& InVisualizationContext, const FChaosVDParticleDataWrapper& InParticleDataViewer);
 };
diff --git a/Engine/Plugins/ChaosVD/Source/ChaosVD/Private/Visualizers/ChaosVDSceneQueryDataComponentVisualizer.cpp b/Engine/Plugins/ChaosVD/Source/ChaosVD/Private/Visualizers/ChaosVDSceneQueryDataComponentVisualizer.cpp
--- a/Engine/Plugins/ChaosVD/Source/ChaosVD/Private/Visualizers/ChaosVDSceneQueryDataComponentVisualizer.cpp
+++ b/Engine/Plugins/ChaosVD/Source/ChaosVD/Private/Visualizers/ChaosVDSceneQueryDataComponentVisualizer.cpp
@@ -5,25 +5,68 @@
 #include "Actors/ChaosVDSceneQueryDataContainer.h"
 #include "Components/ChaosVDSceneQueryDataComponent.h"
 
-#include "ChaosVDEditorSettings.h"
 #include "ChaosVDGeometryBuilder.h"
 #include "ChaosVDScene.h"
-#include "ChaosVDSceneQueryDataInspectorTab.h"
+#include "ChaosVDStyle.h"
 #include "ChaosVDTabsIDs.h"
 #include "EditorModeManager.h"
 #include "EditorViewportClient.h"
 #include "SceneView.h"
 #include "Actors/ChaosVDSolverInfoActor.h"
+#include "Settings/ChaosVDSceneQueryVisualizationSettings.h"
+#include "ToolMenu.h"
+#include "ToolMenus.h"
+#include "ToolMenuEntry.h"
+#include "ToolMenuSection.h"
+#include "Utils/ChaosVDUserInterfaceUtils.h"
 #include "Visualizers/ChaosVDDebugDrawUtils.h"
 #include "Visualizers/IChaosVDParticleVisualizationDataProvider.h"
 #include "Widgets/SChaosVDMainTab.h"
+#include "Widgets/SChaosVDViewportToolbar.h"
+#include "Widgets/SChaosVDEnumFlagsMenu.h"
 
-class UChaosVDEditorSettings;
+class UChaosVDCoreSettings;
 
 IMPLEMENT_HIT_PROXY(HChaosVDSceneQueryProxy, HComponentVisProxy)
 
 #define LOCTEXT_NAMESPACE "ChaosVisualDebugger"
 
+FChaosVDSceneQueryDataComponentVisualizer::FChaosVDSceneQueryDataComponentVisualizer()
+{
+	RegisterVisualizerMenus();
+}
+
+void FChaosVDSceneQueryDataComponentVisualizer::RegisterVisualizerMenus()
+{
+	UToolMenus* ToolMenus = UToolMenus::Get();
+	
+	if (!ensure(ToolMenus))
+	{
+		return;
+	}
+
+	if (UToolMenu* Menu = ToolMenus->ExtendMenu(SChaosVDViewportToolbar::ShowMenuName))
+	{
+		FToolMenuSection& Section = Menu->AddSection("SceneQueryDataVisualization.Show", LOCTEXT("SceneQueryDataShowMenuLabel", "Scene Query Data Visualization"));
+		
+		Section.AddSubMenu(TEXT("SceneQueryDataFlags"), LOCTEXT("SceneQueryDataFlagsMenuLabel", "Scene Query Data Flags"), LOCTEXT("SceneQueryDataFlagsMenuToolTip", "Set of flags to enable/disable visibility of specific types of scene query data"), FNewToolMenuDelegate::CreateLambda([](UToolMenu* Menu)
+						   {
+							   TSharedRef<SWidget> VisualizationFlagsWidget = SNew(SChaosVDEnumFlagsMenu<EChaosVDSceneQueryVisualizationFlags>)
+								   .CurrentValue_Static(&UChaosVDSceneQueriesVisualizationSettings::GetSceneQueryDataVisualizationFlags)
+								   .OnEnumSelectionChanged_Static(&UChaosVDSceneQueriesVisualizationSettings::SetSceneQueryDataVisualizationFlags);
+			
+
+							   FToolMenuEntry FlagsMenuEntry = FToolMenuEntry::InitWidget("CollisionVisualizationFlagsFlags", VisualizationFlagsWidget,FText::GetEmpty());
+							   Menu->AddMenuEntry(NAME_None, FlagsMenuEntry);
+						   }),
+						   false, FSlateIcon(FChaosVDStyle::Get().GetStyleSetName(), TEXT("SceneQueriesInspectorIcon")));
+		
+		using namespace Chaos::VisualDebugger::Utils;
+        Section.AddSubMenu(TEXT("SceneQueryDataSettings"), LOCTEXT("SceneQuerySettingsMenuLabel", "Scene Query Visualization Settings"), LOCTEXT("SceneQuerySettingsMenuToolTip", "Options to change how the recorded scene query data is debug drawn"), FNewToolMenuDelegate::CreateStatic(&CreateMenuEntryForDefaultObject<UChaosVDSceneQueriesVisualizationSettings>, EChaosVDSaveSettingsOptions::ShowSaveButton),
+                           false, FSlateIcon(FAppStyle::Get().GetStyleSetName(), TEXT("Icons.Toolbar.Settings")));
+	}
+}
+
 void FChaosVDSceneQueryDataComponentVisualizer::DrawVisualization(const UActorComponent* Component, const FSceneView* View, FPrimitiveDrawInterface* PDI)
 {
 	const UChaosVDSceneQueryDataComponent* SceneQueryDataComponent = Cast<UChaosVDSceneQueryDataComponent>(Component);
@@ -61,10 +104,8 @@ void FChaosVDSceneQueryDataComponentVisualizer::DrawVisualization(const UActorCo
 	VisualizationContext.SpaceTransform = FTransform::Identity;
 	VisualizationContext.GeometryGenerator = GeometryGenerator;
 
-	if (const UChaosVDEditorSettings* EditorSettings = GetDefault<UChaosVDEditorSettings>())
-	{
-		VisualizationContext.VisualizationFlags = EditorSettings->GlobalSceneQueriesVisualizationFlags;
-	}
+	VisualizationContext.VisualizationFlags = static_cast<uint32>(UChaosVDSceneQueriesVisualizationSettings::GetSceneQueryDataVisualizationFlags());
+	VisualizationContext.DebugDrawSettings = GetDefault<UChaosVDSceneQueriesVisualizationSettings>();
 
 	if (EnumHasAnyFlags(EChaosVDSceneQueryVisualizationFlags::EnableDraw, static_cast<EChaosVDSceneQueryVisualizationFlags>(VisualizationContext.VisualizationFlags)))
 	{
@@ -113,9 +154,16 @@ bool FChaosVDSceneQueryDataComponentVisualizer::VisProxyHandleClick(FEditorViewp
 
 void FChaosVDSceneQueryDataComponentVisualizer::DrawLineTraceQuery(const UActorComponent* Component, const FChaosVDQueryDataWrapper& SceneQueryData, FChaosVDSceneQueryVisualizationDataContext& VisualizationContext, const FSceneView* View, FPrimitiveDrawInterface* PDI)
 {
+	const UChaosVDSceneQueriesVisualizationSettings* DebugDrawSettings = Cast<UChaosVDSceneQueriesVisualizationSettings>(VisualizationContext.DebugDrawSettings);
+	if (!DebugDrawSettings)
+	{
+		return;
+	}
+
 	PDI->SetHitProxy(new HChaosVDSceneQueryProxy(Component, VisualizationContext.DataSelectionHandle));
 
-	const FText DebugText = FText::FormatOrdered(LOCTEXT("LineTraceDebugDrawText", "Type: Line Trace \n Tag {1} \n Owner Tag {2}"), FText::AsCultureInvariant(SceneQueryData.CollisionQueryParams.TraceTag.ToString()), FText::AsCultureInvariant(SceneQueryData.CollisionQueryParams.OwnerTag.ToString()));
+	const FText DebugText = DebugDrawSettings->bShowText ? FText::FormatOrdered(LOCTEXT("LineTraceDebugDrawText", "Type: Line Trace \n Tag {1} \n Owner Tag {2}"), FText::AsCultureInvariant(SceneQueryData.CollisionQueryParams.TraceTag.ToString()), FText::AsCultureInvariant(SceneQueryData.CollisionQueryParams.OwnerTag.ToString()))
+							: FText::GetEmpty();
 
 	FVector EndLocationToDraw;
 	if (SceneQueryData.SQVisitData.IsValidIndex(SceneQueryData.CurrentVisitIndex))
@@ -141,12 +189,19 @@ void FChaosVDSceneQueryDataComponentVisualizer::DrawLineTraceQuery(const UActorC
 
 void FChaosVDSceneQueryDataComponentVisualizer::DrawOverlapQuery(const UActorComponent* Component, const FChaosVDQueryDataWrapper& SceneQueryData, FChaosVDSceneQueryVisualizationDataContext& VisualizationContext, const FSceneView* View, FPrimitiveDrawInterface* PDI)
 {
+	const UChaosVDSceneQueriesVisualizationSettings* DebugDrawSettings = Cast<UChaosVDSceneQueriesVisualizationSettings>(VisualizationContext.DebugDrawSettings);
+	if (!DebugDrawSettings)
+	{
+		return;
+	}
+	
 	PDI->SetHitProxy(new HChaosVDSceneQueryProxy(Component, VisualizationContext.DataSelectionHandle));
 
 	const Chaos::FConstImplicitObjectPtr InputShapePtr = VisualizationContext.InputGeometry;
 	if (ensure(InputShapePtr))
 	{
-		const FText DebugText = FText::FormatOrdered(LOCTEXT("OverlapDebugDrawText", "Type: Overlap \n Tag {1} \n Owner Tag {2}"), FText::AsCultureInvariant(SceneQueryData.CollisionQueryParams.TraceTag.ToString()), FText::AsCultureInvariant(SceneQueryData.CollisionQueryParams.OwnerTag.ToString()));
+		const FText DebugText = DebugDrawSettings->bShowText ? FText::FormatOrdered(LOCTEXT("OverlapDebugDrawText", "Type: Overlap \n Tag {1} \n Owner Tag {2}"), FText::AsCultureInvariant(SceneQueryData.CollisionQueryParams.TraceTag.ToString()), FText::AsCultureInvariant(SceneQueryData.CollisionQueryParams.OwnerTag.ToString()))
+								: FText::GetEmpty();
 		FChaosVDDebugDrawUtils::DrawImplicitObject(PDI, VisualizationContext.GeometryGenerator.Pin(), InputShapePtr, FTransform(SceneQueryData.GeometryOrientation, SceneQueryData.StartLocation), VisualizationContext.DebugDrawColor, DebugText, ESceneDepthPriorityGroup::SDPG_Foreground);
 	}
 
@@ -160,6 +215,12 @@ void FChaosVDSceneQueryDataComponentVisualizer::DrawOverlapQuery(const UActorCom
 
 void FChaosVDSceneQueryDataComponentVisualizer::DrawSweepQuery(const UActorComponent* Component, const FChaosVDQueryDataWrapper& SceneQueryData, FChaosVDSceneQueryVisualizationDataContext& VisualizationContext, const FSceneView* View, FPrimitiveDrawInterface* PDI)
 {
+	const UChaosVDSceneQueriesVisualizationSettings* DebugDrawSettings = Cast<UChaosVDSceneQueriesVisualizationSettings>(VisualizationContext.DebugDrawSettings);
+	if (!DebugDrawSettings)
+	{
+		return;
+	}
+
 	PDI->SetHitProxy(new HChaosVDSceneQueryProxy(Component, VisualizationContext.DataSelectionHandle));
 
 	const Chaos::FConstImplicitObjectPtr InputShapePtr = VisualizationContext.InputGeometry;
@@ -184,7 +245,8 @@ void FChaosVDSceneQueryDataComponentVisualizer::DrawSweepQuery(const UActorCompo
 
 	PDI->SetHitProxy(nullptr);
 
-	const FText DebugText = FText::FormatOrdered(LOCTEXT("SweepDebugDrawText", "Type: Sweep \n Tag {1} \n Owner Tag {2}"), FText::AsCultureInvariant(SceneQueryData.CollisionQueryParams.TraceTag.ToString()), FText::AsCultureInvariant(SceneQueryData.CollisionQueryParams.OwnerTag.ToString()));
+	const FText DebugText = DebugDrawSettings->bShowText ? FText::FormatOrdered(LOCTEXT("SweepDebugDrawText", "Type: Sweep \n Tag {1} \n Owner Tag {2}"), FText::AsCultureInvariant(SceneQueryData.CollisionQueryParams.TraceTag.ToString()), FText::AsCultureInvariant(SceneQueryData.CollisionQueryParams.OwnerTag.ToString()))
+							: FText::GetEmpty();
 
 	FChaosVDDebugDrawUtils::DrawArrowVector(PDI, SceneQueryData.StartLocation, SceneQueryData.EndLocation, DebugText, VisualizationContext.DebugDrawColor, SDPG_Foreground);
 
@@ -196,6 +258,12 @@ void FChaosVDSceneQueryDataComponentVisualizer::DrawSweepQuery(const UActorCompo
 
 void FChaosVDSceneQueryDataComponentVisualizer::DrawHits(const UActorComponent* Component, const FChaosVDQueryDataWrapper& SceneQueryData, FPrimitiveDrawInterface* PDI, const FColor& InColor, FChaosVDSceneQueryVisualizationDataContext& VisualizationContext)
 {
+	const UChaosVDSceneQueriesVisualizationSettings* DebugDrawSettings = Cast<UChaosVDSceneQueriesVisualizationSettings>(VisualizationContext.DebugDrawSettings);
+	if (!DebugDrawSettings)
+	{
+		return;
+	}
+
 	for (int32 SQVisitIndex = 0; SQVisitIndex < SceneQueryData.SQVisitData.Num(); ++SQVisitIndex)
 	{
 		const FChaosVDQueryVisitStep& SQVisitData = SceneQueryData.SQVisitData[SQVisitIndex];
@@ -208,7 +276,7 @@ void FChaosVDSceneQueryDataComponentVisualizer::DrawHits(const UActorComponent*
 
 		PDI->SetHitProxy(new HChaosVDSceneQueryProxy(Component, VisualizationContext.DataSelectionHandle));
 
-		const FText HitPointDebugText = FText::FormatOrdered(LOCTEXT("SceneQueryHitDebugText", "Distance {0} \n Face Index {1} \n "), SQVisitData.HitData.Distance, SQVisitData.HitData.FaceIdx);
+		const FText HitPointDebugText = DebugDrawSettings->bShowText ? FText::FormatOrdered(LOCTEXT("SceneQueryHitDebugText", "Distance {0} \n Face Index {1} \n "), SQVisitData.HitData.Distance, SQVisitData.HitData.FaceIdx) : FText::GetEmpty();
 		static FText HitFaceNormalDebugText = FText::AsCultureInvariant(TEXT("Hit Face Normal"));
 		static FText HitWorldNormalDebugText = FText::AsCultureInvariant(TEXT("Hit World Normal"));
 
@@ -218,12 +286,12 @@ void FChaosVDSceneQueryDataComponentVisualizer::DrawHits(const UActorComponent*
 		constexpr int32 CircleSegments = 12;
 		constexpr float NormalScale = 10.5f;
 		FChaosVDDebugDrawUtils::DrawCircle(PDI, SQVisitData.HitData.WorldPosition, CircleRadius, CircleSegments, InColor, Thickness, Axes.GetUnitAxis(EAxis::Y), Axes.GetUnitAxis(EAxis::Z), HitPointDebugText,  ESceneDepthPriorityGroup::SDPG_Foreground);
-		FChaosVDDebugDrawUtils::DrawArrowVector(PDI, SQVisitData.HitData.WorldPosition, SQVisitData.HitData.WorldPosition + SQVisitData.HitData.FaceNormal * NormalScale, HitFaceNormalDebugText, (FLinearColor(InColor) * 0.65f).ToFColorSRGB(), ESceneDepthPriorityGroup::SDPG_Foreground);
+		FChaosVDDebugDrawUtils::DrawArrowVector(PDI, SQVisitData.HitData.WorldPosition, SQVisitData.HitData.WorldPosition + SQVisitData.HitData.FaceNormal * NormalScale, DebugDrawSettings->bShowText ? HitFaceNormalDebugText : FText::GetEmpty(), (FLinearColor(InColor) * 0.65f).ToFColorSRGB(), ESceneDepthPriorityGroup::SDPG_Foreground);
 
 		// Hit Face Normal is not used in line traces
 		if (SceneQueryData.Type != EChaosVDSceneQueryType::RayCast)
 		{
-			FChaosVDDebugDrawUtils::DrawArrowVector(PDI, SQVisitData.HitData.WorldPosition, SQVisitData.HitData.WorldPosition + SQVisitData.HitData.WorldNormal * NormalScale, HitWorldNormalDebugText, InColor, ESceneDepthPriorityGroup::SDPG_Foreground);
+			FChaosVDDebugDrawUtils::DrawArrowVector(PDI, SQVisitData.HitData.WorldPosition, SQVisitData.HitData.WorldPosition + SQVisitData.HitData.WorldNormal * NormalScale, DebugDrawSettings->bShowText ? HitWorldNormalDebugText : FText::GetEmpty(), InColor, ESceneDepthPriorityGroup::SDPG_Foreground);
 		}
 
 		if (SQVisitData.bIsSelectedInEditor)
diff --git a/Engine/Plugins/ChaosVD/Source/ChaosVD/Private/Visualizers/ChaosVDSceneQueryDataComponentVisualizer.h b/Engine/Plugins/ChaosVD/Source/ChaosVD/Private/Visualizers/ChaosVDSceneQueryDataComponentVisualizer.h
--- a/Engine/Plugins/ChaosVD/Source/ChaosVD/Private/Visualizers/ChaosVDSceneQueryDataComponentVisualizer.h
+++ b/Engine/Plugins/ChaosVD/Source/ChaosVD/Private/Visualizers/ChaosVDSceneQueryDataComponentVisualizer.h
@@ -53,22 +53,6 @@ struct HChaosVDSceneQueryProxy : public HComponentVisProxy
 	FChaosVDSceneQuerySelectionHandle DataSelectionHandle;
 };
 
-/** Set of visualization flags options for Scene Queries */
-UENUM(meta = (Bitflags, UseEnumValuesAsMaskValuesInEditor = "true"))
-enum class EChaosVDSceneQueryVisualizationFlags: uint32
-{
-	None					= 0 UMETA(Hidden),
-	EnableDraw				= 1 << 0,
-	DrawLineTraceQueries	= 1 << 1,
-	DrawSweepQueries		= 1 << 2,
-	DrawOverlapQueries		= 1 << 3,
-	DrawHits				= 1 << 4,
-	OnlyDrawSelectedQuery	= 1 << 5,
-	HideEmptyQueries		= 1 << 6,
-	HideSubQueries			= 1 << 7,
-};
-ENUM_CLASS_FLAGS(EChaosVDSceneQueryVisualizationFlags);
-
 /**
  * Component visualizer in charge of generating debug draw visualizations for scene queries in a ChaosVDSceneQueryDataComponent
  */
@@ -76,9 +60,9 @@ class FChaosVDSceneQueryDataComponentVisualizer final : public FComponentVisuali
 {
 
 public:
-	FChaosVDSceneQueryDataComponentVisualizer()
-	{
-	}
+	FChaosVDSceneQueryDataComponentVisualizer();
+	
+	void RegisterVisualizerMenus();
 
 	virtual void DrawVisualization(const UActorComponent* Component, const FSceneView* View, FPrimitiveDrawInterface* PDI) override;
 	virtual bool VisProxyHandleClick(FEditorViewportClient* InViewportClient, HComponentVisProxy* VisProxy, const FViewportClick& Click) override;
diff --git a/Engine/Plugins/ChaosVD/Source/ChaosVD/Private/Visualizers/ChaosVDSolverCollisionDataComponentVisualizer.cpp b/Engine/Plugins/ChaosVD/Source/ChaosVD/Private/Visualizers/ChaosVDSolverCollisionDataComponentVisualizer.cpp
--- a/Engine/Plugins/ChaosVD/Source/ChaosVD/Private/Visualizers/ChaosVDSolverCollisionDataComponentVisualizer.cpp
+++ b/Engine/Plugins/ChaosVD/Source/ChaosVD/Private/Visualizers/ChaosVDSolverCollisionDataComponentVisualizer.cpp
@@ -3,7 +3,6 @@
 #include "ChaosVDSolverCollisionDataComponentVisualizer.h"
 
 #include "ChaosVDCollisionDataDetailsTab.h"
-#include "ChaosVDEditorSettings.h"
 #include "ChaosVDParticleActor.h"
 #include "ChaosVDScene.h"
 #include "ChaosVDTabsIDs.h"
@@ -11,9 +10,17 @@
 #include "EditorViewportClient.h"
 #include "Actors/ChaosVDSolverInfoActor.h"
 #include "Components/ChaosVDSolverCollisionDataComponent.h"
+#include "Settings/ChaosVDCollisionVisualizationSettings.h"
+#include "ToolMenu.h"
+#include "ToolMenus.h"
+#include "ToolMenuEntry.h"
+#include "ToolMenuSection.h"
+#include "Utils/ChaosVDUserInterfaceUtils.h"
 #include "Visualizers/ChaosVDDebugDrawUtils.h"
 #include "Widgets/SChaosVDCollisionDataInspector.h"
 #include "Widgets/SChaosVDMainTab.h"
+#include "Widgets/SChaosVDViewportToolbar.h"
+#include "Widgets/SChaosVDEnumFlagsMenu.h"
 
 IMPLEMENT_HIT_PROXY(HChaosVDContactPointProxy, HComponentVisProxy)
 
@@ -32,12 +39,44 @@ void FChaosVDCollisionDataFinder::SetIsSelected(bool bNewSelected)
 
 FChaosVDSolverCollisionDataComponentVisualizer::FChaosVDSolverCollisionDataComponentVisualizer()
 {
+	RegisterVisualizerMenus();
 }
 
 FChaosVDSolverCollisionDataComponentVisualizer::~FChaosVDSolverCollisionDataComponentVisualizer()
 {
 }
 
+void FChaosVDSolverCollisionDataComponentVisualizer::RegisterVisualizerMenus()
+{
+	UToolMenus* ToolMenus = UToolMenus::Get();
+	
+	if (!ensure(ToolMenus))
+	{
+		return;
+	}
+
+	if (UToolMenu* Menu = ToolMenus->ExtendMenu(SChaosVDViewportToolbar::ShowMenuName))
+	{
+		FToolMenuSection& Section = Menu->AddSection("CollisionDataVisualization.Show", LOCTEXT("CollisionDataVisualizationShowMenuLabel", "Collision Data Visualization"));
+		
+		Section.AddSubMenu(TEXT("CollisionVisualizationFlags"), LOCTEXT("CollisionVisualizationFlagsMenuLabel", "Collision Data Flags"), LOCTEXT("CollisionVisualizationFlagsMenuToolTip", "Set of flags to enable/disable visibility of specific types of collisionData"), FNewToolMenuDelegate::CreateLambda([](UToolMenu* Menu)
+		                   {
+			                   TSharedRef<SWidget> VisualizationFlagsWidget = SNew(SChaosVDEnumFlagsMenu<EChaosVDCollisionVisualizationFlags>)
+				                   .CurrentValue_Static(&UChaosVDCollisionDataVisualizationSettings::GetCollisionDataVisualizationFlags)
+				                   .OnEnumSelectionChanged_Static(&UChaosVDCollisionDataVisualizationSettings::SetCollisionDataVisualizationFlags);
+			
+
+			                   FToolMenuEntry FlagsMenuEntry = FToolMenuEntry::InitWidget("CollisionVisualizationFlagsFlags", VisualizationFlagsWidget,FText::GetEmpty());
+			                   Menu->AddMenuEntry(NAME_None, FlagsMenuEntry);
+		                   }),
+		                   false, FSlateIcon(FAppStyle::Get().GetStyleSetName(), TEXT("EditorViewport.CollisionVisibility")));
+
+		using namespace Chaos::VisualDebugger::Utils;
+		Section.AddSubMenu(TEXT("CollisionVisualizationSettings"), LOCTEXT("CollisionVisualizationSettingsMenuLabel", "Collision Visualization Settings"), LOCTEXT("CollisionVisualizationSettingsMenuToolTip", "Options to change how the recorded collision Data is debug drawn"), FNewToolMenuDelegate::CreateStatic(&CreateMenuEntryForDefaultObject<UChaosVDCollisionDataVisualizationSettings>, EChaosVDSaveSettingsOptions::ShowSaveButton),
+		                   false, FSlateIcon(FAppStyle::Get().GetStyleSetName(), TEXT("Icons.Toolbar.Settings")));
+	}
+}
+
 bool FChaosVDSolverCollisionDataComponentVisualizer::ShowWhenSelected()
 {
 	return false;
@@ -66,11 +105,8 @@ void FChaosVDSolverCollisionDataComponentVisualizer::DrawVisualization(const UAc
 	VisualizationContext.SolverID = SolverInfoContainer->GetSolverID();
 	VisualizationContext.CVDScene = SolverInfoContainer->GetScene();
 	VisualizationContext.SpaceTransform = SolverInfoContainer->GetSimulationTransform();
-
-	if (const UChaosVDEditorSettings* EditorSettings = GetDefault<UChaosVDEditorSettings>())
-	{
-		VisualizationContext.VisualizationFlags = EditorSettings->GlobalCollisionDataVisualizationFlags;
-	}
+	VisualizationContext.VisualizationFlags = static_cast<uint32>(UChaosVDCollisionDataVisualizationSettings::GetCollisionDataVisualizationFlags());
+	VisualizationContext.DebugDrawSettings = GetDefault<UChaosVDCollisionDataVisualizationSettings>();
 
 	if (!EnumHasAnyFlags(EChaosVDCollisionVisualizationFlags::EnableDraw, static_cast<EChaosVDCollisionVisualizationFlags>(VisualizationContext.VisualizationFlags)))
 	{
@@ -163,10 +199,10 @@ void FChaosVDSolverCollisionDataComponentVisualizer::DrawMidPhaseData(const UAct
 		return;
 	}
 
-	FChaosVDContactDebugDrawSettings DebugDrawSettings;
-	if (const UChaosVDEditorSettings* EditorSettings = GetDefault<UChaosVDEditorSettings>())
+	const UChaosVDCollisionDataVisualizationSettings* DebugDrawSettings = Cast<UChaosVDCollisionDataVisualizationSettings>(VisualizationContext.DebugDrawSettings);
+	if (!VisualizationContext.DebugDrawSettings)
 	{
-		DebugDrawSettings = EditorSettings->ContactDebugDrawSettings;
+		return;
 	}
 	
 	if (MidPhase.IsValid())
@@ -290,7 +326,8 @@ void FChaosVDSolverCollisionDataComponentVisualizer::DrawMidPhaseData(const UAct
 				{
 					if (ManifoldPoint.bIsValid && !ManifoldPoint.NetPushOut.IsNearlyZero())
 					{
-						FChaosVDDebugDrawUtils::DrawArrowVector(PDI, WorldPointPlaneLocation, WorldPointPlaneLocation + VisualizationContext.SpaceTransform.TransformPosition(ManifoldPoint.NetPushOut),  UEnum::GetDisplayValueAsText(EChaosVDCollisionVisualizationFlags::NetPushOut), PushOutImpulseColor, DebugDrawSettings.DepthPriority);
+						const FText& DebugText = DebugDrawSettings->bShowDebugText ? UEnum::GetDisplayValueAsText(EChaosVDCollisionVisualizationFlags::NetPushOut) : FText::GetEmpty();
+						FChaosVDDebugDrawUtils::DrawArrowVector(PDI, WorldPointPlaneLocation, WorldPointPlaneLocation + VisualizationContext.SpaceTransform.TransformPosition(ManifoldPoint.NetPushOut), DebugText, PushOutImpulseColor, DebugDrawSettings->DepthPriority);
 					}
 				}
 
@@ -298,7 +335,8 @@ void FChaosVDSolverCollisionDataComponentVisualizer::DrawMidPhaseData(const UAct
 				{
 					if (ManifoldPoint.bIsValid && !ManifoldPoint.NetImpulse.IsNearlyZero())
 					{
-						FChaosVDDebugDrawUtils::DrawArrowVector(PDI, WorldPointPlaneLocation, WorldPointPlaneLocation + VisualizationContext.SpaceTransform.TransformPosition(ManifoldPoint.NetImpulse), UEnum::GetDisplayValueAsText(EChaosVDCollisionVisualizationFlags::NetImpulse), ImpulseColor, DebugDrawSettings.DepthPriority);
+						const FText& DebugText = DebugDrawSettings->bShowDebugText ? UEnum::GetDisplayValueAsText(EChaosVDCollisionVisualizationFlags::NetImpulse) : FText::GetEmpty();
+						FChaosVDDebugDrawUtils::DrawArrowVector(PDI, WorldPointPlaneLocation, WorldPointPlaneLocation + VisualizationContext.SpaceTransform.TransformPosition(ManifoldPoint.NetImpulse), DebugText, ImpulseColor, DebugDrawSettings->DepthPriority);
 					}
 				}
 
@@ -309,26 +347,29 @@ void FChaosVDSolverCollisionDataComponentVisualizer::DrawMidPhaseData(const UAct
 					static FText ManifoldPointDebugText = FText::FormatOrdered(LOCTEXT("ManifoldPointDebugText", "{0} | Manifold Point"), ContactPointsDebugBaseFText);
 					static FText ManifoldInitialPhiDebugText = FText::FormatOrdered(LOCTEXT("ManifoldInitialPhiDebugText", "{0} | Manifold Initial Phi"), ContactPointsDebugBaseFText);
 	
-					FChaosVDDebugDrawUtils::DrawCircle(PDI, WorldPlaneLocation, DebugDrawSettings.ContactCircleRadius, CircleSegments, DiscColor, LinesThickness, Axes.GetUnitAxis(EAxis::Y), Axes.GetUnitAxis(EAxis::Z), ManifoldPlaneDebugText, DebugDrawSettings.DepthPriority);
-					FChaosVDDebugDrawUtils::DrawCircle(PDI, WorldPointLocation, 0.5f * DebugDrawSettings.ContactCircleRadius, CircleSegments, DiscColor, LinesThickness, Axes.GetUnitAxis(EAxis::Y), Axes.GetUnitAxis(EAxis::Z), ManifoldPointDebugText, DebugDrawSettings.DepthPriority);
+					FChaosVDDebugDrawUtils::DrawCircle(PDI, WorldPlaneLocation, DebugDrawSettings->ContactCircleRadius, CircleSegments, DiscColor, LinesThickness, Axes.GetUnitAxis(EAxis::Y), Axes.GetUnitAxis(EAxis::Z), DebugDrawSettings->bShowDebugText ? ManifoldPlaneDebugText : FText::GetEmpty(), DebugDrawSettings->DepthPriority);
+					FChaosVDDebugDrawUtils::DrawCircle(PDI, WorldPointLocation, 0.5f * DebugDrawSettings->ContactCircleRadius, CircleSegments, DiscColor, LinesThickness, Axes.GetUnitAxis(EAxis::Y), Axes.GetUnitAxis(EAxis::Z),  DebugDrawSettings->bShowDebugText ? ManifoldPointDebugText : FText::GetEmpty(), DebugDrawSettings->DepthPriority);
 					if (ManifoldPoint.InitialPhi != 0)
 					{
-						FChaosVDDebugDrawUtils::DrawCircle(PDI, WorldPlaneLocation + ManifoldPoint.InitialPhi * WorldPlaneNormal, 0.25f * DebugDrawSettings.ContactCircleRadius, CircleSegments, InitialPhiColor, LinesThickness, Axes.GetUnitAxis(EAxis::Y), Axes.GetUnitAxis(EAxis::Z), ManifoldInitialPhiDebugText, DebugDrawSettings.DepthPriority);
+						FChaosVDDebugDrawUtils::DrawCircle(PDI, WorldPlaneLocation + ManifoldPoint.InitialPhi * WorldPlaneNormal, 0.25f * DebugDrawSettings->ContactCircleRadius, CircleSegments, InitialPhiColor, LinesThickness, Axes.GetUnitAxis(EAxis::Y), Axes.GetUnitAxis(EAxis::Z),
+							DebugDrawSettings->bShowDebugText ? ManifoldInitialPhiDebugText : FText::GetEmpty(), DebugDrawSettings->DepthPriority);
 					}
 				}
 
 				if (EnumHasAnyFlags(VisualizationFlags, EChaosVDCollisionVisualizationFlags::ContactNormal))
 				{
 					FColor NormalColor = ((ManifoldPoint.ContactPoint.ContactType != EChaosVDContactPointType::EdgeEdge) ? PlaneNormalColor : EdgeNormalColor);
-					const int32 Scale = DebugDrawSettings.ContactNormalScale * ContactLenScale;
-					FChaosVDDebugDrawUtils::DrawArrowVector(PDI, WorldPlaneLocation, WorldPlaneLocation + WorldPlaneNormal * Scale, UEnum::GetDisplayValueAsText(EChaosVDCollisionVisualizationFlags::ContactNormal), NormalColor, DebugDrawSettings.DepthPriority);
+					const int32 Scale = DebugDrawSettings->ContactNormalScale * ContactLenScale;
+					const FText& DebugText = DebugDrawSettings->bShowDebugText ? UEnum::GetDisplayValueAsText(EChaosVDCollisionVisualizationFlags::ContactNormal) : FText::GetEmpty();
+					FChaosVDDebugDrawUtils::DrawArrowVector(PDI, WorldPlaneLocation, WorldPlaneLocation + WorldPlaneNormal * Scale, DebugText, NormalColor, DebugDrawSettings->DepthPriority);
 				}
 
 				if (EnumHasAnyFlags(VisualizationFlags, EChaosVDCollisionVisualizationFlags::AccumulatedImpulse))
 				{
 					if (!Constraint.AccumulatedImpulse.IsNearlyZero())
 					{
-						FChaosVDDebugDrawUtils::DrawArrowVector(PDI, WorldActorTransform0.GetLocation(), WorldActorTransform0.GetLocation() + Constraint.AccumulatedImpulse, UEnum::GetDisplayValueAsText(EChaosVDCollisionVisualizationFlags::AccumulatedImpulse), FColor::White, DebugDrawSettings.DepthPriority);
+						const FText& DebugText = DebugDrawSettings->bShowDebugText ? UEnum::GetDisplayValueAsText(EChaosVDCollisionVisualizationFlags::AccumulatedImpulse) : FText::GetEmpty();
+						FChaosVDDebugDrawUtils::DrawArrowVector(PDI, WorldActorTransform0.GetLocation(), WorldActorTransform0.GetLocation() + Constraint.AccumulatedImpulse, DebugText, FColor::White, DebugDrawSettings->DepthPriority);
 					}
 				}
 
@@ -338,29 +379,29 @@ void FChaosVDSolverCollisionDataComponentVisualizer::DrawMidPhaseData(const UAct
 					Transform.SetRotation(FRotationMatrix::MakeFromZ(WorldPlaneNormal).ToQuat());
 					Transform.SetLocation(WorldPlaneLocation);
 
-					FVector BoxExtents(DebugDrawSettings.ContactCircleRadius, DebugDrawSettings.ContactCircleRadius, 0.01f);
+					FVector BoxExtents(DebugDrawSettings->ContactCircleRadius, DebugDrawSettings->ContactCircleRadius, 0.01f);
 					FColor Color;
 
 					if (Constraint.bWasManifoldRestored)
 					{
 						Color = FColor::Blue;
-						FChaosVDDebugDrawUtils::DrawBox(PDI, BoxExtents, Color, Transform, FText::GetEmpty(), DebugDrawSettings.DepthPriority);
+						FChaosVDDebugDrawUtils::DrawBox(PDI, BoxExtents, Color, Transform, FText::GetEmpty(), DebugDrawSettings->DepthPriority);
 					}
 					else if (ManifoldPoint.bWasRestored)
 					{
 						Color = FColor::Purple;
-						FChaosVDDebugDrawUtils::DrawBox(PDI, BoxExtents, Color, Transform, FText::GetEmpty(), DebugDrawSettings.DepthPriority);
+						FChaosVDDebugDrawUtils::DrawBox(PDI, BoxExtents, Color, Transform, FText::GetEmpty(), DebugDrawSettings->DepthPriority);
 					}
 					else if (ManifoldPoint.bWasReplaced)
 					{
 						Color = FColor::Orange;
-						FChaosVDDebugDrawUtils::DrawBox(PDI, BoxExtents, Color, Transform, FText::GetEmpty(), DebugDrawSettings.DepthPriority);
+						FChaosVDDebugDrawUtils::DrawBox(PDI, BoxExtents, Color, Transform, FText::GetEmpty(), DebugDrawSettings->DepthPriority);
 					}
 
 					if (MidPhase->bIsSleeping)
 					{
 						// This box should surround the debug draw box we made before
-						FChaosVDDebugDrawUtils::DrawBox(PDI, BoxExtents * 1.1f, FColor::Black, Transform, FText::GetEmpty(), DebugDrawSettings.DepthPriority);
+						FChaosVDDebugDrawUtils::DrawBox(PDI, BoxExtents * 1.1f, FColor::Black, Transform, FText::GetEmpty(), DebugDrawSettings->DepthPriority);
 					}
 				}
 
@@ -375,10 +416,10 @@ void FChaosVDSolverCollisionDataComponentVisualizer::DrawMidPhaseData(const UAct
 					SelectionBoxTransform.SetLocation(WorldPlaneLocation);
 
 					// The Selection box should be a bit bigger than the configured circle radius for the debug draw contact
-					const float ContactSelectionBoxSize = DebugDrawSettings.ContactCircleRadius * 1.5f;
+					const float ContactSelectionBoxSize = DebugDrawSettings->ContactCircleRadius * 1.5f;
 
 					FVector SelectionBoxExtents(ContactSelectionBoxSize,ContactSelectionBoxSize,ContactSelectionBoxSize);
-					FChaosVDDebugDrawUtils::DrawBox(PDI, SelectionBoxExtents, FColor::Yellow, SelectionBoxTransform, FText::GetEmpty(), DebugDrawSettings.DepthPriority);
+					FChaosVDDebugDrawUtils::DrawBox(PDI, SelectionBoxExtents, FColor::Yellow, SelectionBoxTransform, FText::GetEmpty(), DebugDrawSettings->DepthPriority);
 				}
 			}
 		}
diff --git a/Engine/Plugins/ChaosVD/Source/ChaosVD/Private/Visualizers/ChaosVDSolverCollisionDataComponentVisualizer.h b/Engine/Plugins/ChaosVD/Source/ChaosVD/Private/Visualizers/ChaosVDSolverCollisionDataComponentVisualizer.h
--- a/Engine/Plugins/ChaosVD/Source/ChaosVD/Private/Visualizers/ChaosVDSolverCollisionDataComponentVisualizer.h
+++ b/Engine/Plugins/ChaosVD/Source/ChaosVD/Private/Visualizers/ChaosVDSolverCollisionDataComponentVisualizer.h
@@ -34,27 +34,14 @@ struct HChaosVDContactPointProxy : public HComponentVisProxy
 	FChaosVDCollisionDataFinder ContactFinder;
 };
 
-UENUM(meta = (Bitflags, UseEnumValuesAsMaskValuesInEditor = "true"))
-enum class EChaosVDCollisionVisualizationFlags: uint32
-{
-	None							= 0 UMETA(Hidden),
-	ContactPoints					= 1 << 0,
-	ContactInfo						= 1 << 1,
-	NetPushOut						= 1 << 2,
-	NetImpulse						= 1 << 3,
-	ContactNormal					= 1 << 4,
-	AccumulatedImpulse				= 1 << 5,
-	DrawInactiveContacts			= 1 << 6,
-	DrawDataOnlyForSelectedParticle	= 1 << 7,
-	EnableDraw						= 1 << 8,
-};
-ENUM_CLASS_FLAGS(EChaosVDCollisionVisualizationFlags);
-
 class FChaosVDSolverCollisionDataComponentVisualizer : public FComponentVisualizer
 {
 public:
+	
 	FChaosVDSolverCollisionDataComponentVisualizer();
-	~FChaosVDSolverCollisionDataComponentVisualizer();
+	virtual ~FChaosVDSolverCollisionDataComponentVisualizer() override;
+	
+	void RegisterVisualizerMenus();
 
 	virtual bool ShowWhenSelected() override;
 	virtual void DrawVisualization(const UActorComponent* Component, const FSceneView* View, FPrimitiveDrawInterface* PDI) override;
diff --git a/Engine/Plugins/ChaosVD/Source/ChaosVD/Private/Visualizers/IChaosVDParticleVisualizationDataProvider.h b/Engine/Plugins/ChaosVD/Source/ChaosVD/Private/Visualizers/IChaosVDParticleVisualizationDataProvider.h
--- a/Engine/Plugins/ChaosVD/Source/ChaosVD/Private/Visualizers/IChaosVDParticleVisualizationDataProvider.h
+++ b/Engine/Plugins/ChaosVD/Source/ChaosVD/Private/Visualizers/IChaosVDParticleVisualizationDataProvider.h
@@ -5,6 +5,8 @@
 #include "HAL/Platform.h"
 #include "Templates/SharedPointer.h"
 
+class UChaosVDSettingsObjectBase;
+class UChaosVDCollisionDataVisualizationSettings;
 struct FChaosVDParticleDataWrapper;
 class FChaosVDScene;
 
@@ -15,6 +17,7 @@ struct FChaosVDVisualizationContext
 	TWeakPtr<FChaosVDScene> CVDScene;
 	int32 SolverID = INDEX_NONE;
 	uint32 VisualizationFlags = 0;
+	const UChaosVDSettingsObjectBase* DebugDrawSettings = nullptr;
 };
 
 /** Interface to be used by any object that contains Particle Data that needs to be visualized.
diff --git a/Engine/Plugins/ChaosVD/Source/ChaosVD/Private/Widgets/SChaosVDEnumFlagsMenu.h b/Engine/Plugins/ChaosVD/Source/ChaosVD/Private/Widgets/SChaosVDEnumFlagsMenu.h
new file mode 100644
--- /dev/null
+++ b/Engine/Plugins/ChaosVD/Source/ChaosVD/Private/Widgets/SChaosVDEnumFlagsMenu.h
@@ -0,0 +1,151 @@
+// Copyright Epic Games, Inc. All Rights Reserved.
+#pragma once
+
+#include "Framework/MultiBox/MultiBoxBuilder.h"
+#include "Framework/Commands/UIAction.h"
+#include "Widgets/SCompoundWidget.h"
+#include "Widgets/SBoxPanel.h"
+#include "Widgets/Text/STextBlock.h"
+
+
+/**
+ * Simple widget that allows that created a checkbox style menu for enum flags
+ */
+template<typename TEnum>
+class SChaosVDEnumFlagsMenu : public SCompoundWidget
+{
+	
+public:
+	DECLARE_DELEGATE_OneParam(FOnEnumSelectionChanged, TEnum);
+	struct FEnumInfo
+	{
+		explicit FEnumInfo(const int32 InIndex, const TEnum InValue, const FText& InDisplayName, const FText& InTooltipText)
+			: Index(InIndex), Value(InValue), DisplayName(InDisplayName), TooltipText(InTooltipText)
+		{}
+		
+		int32 Index = 0;
+		TEnum Value;
+		FText DisplayName;
+		FText TooltipText;
+	};
+	
+	SLATE_BEGIN_ARGS(SChaosVDEnumFlagsMenu)
+		: _CurrentValue()
+	{}
+
+	SLATE_ATTRIBUTE(TEnum, CurrentValue)
+	SLATE_EVENT(FOnEnumSelectionChanged, OnEnumSelectionChanged)
+
+		
+	SLATE_END_ARGS()
+
+	void Construct(const FArguments& InArgs);
+
+	void UpdateEnumFlagValue(FEnumInfo FlagInfo);
+	bool IsEnumFlagSet(FEnumInfo FlagInfo) const;
+
+	TAttribute<TEnum> CurrentValue;
+	
+	FOnEnumSelectionChanged OnEnumSelectionChangedDelegate;
+	
+	const UEnum* Enum = nullptr;
+};
+
+namespace Chaos::VisualDebugger::Utils
+{
+	template<typename Enum>
+	void EnumAddToggleFlag(Enum& Flags, Enum Flag)
+	{
+		using UnderlyingType = __underlying_type(Enum);
+		Flags = (Enum)((UnderlyingType)Flags ^ (UnderlyingType)Flag);
+	}
+}
+
+template <typename TEnum>
+void SChaosVDEnumFlagsMenu<TEnum>::Construct(const FArguments& InArgs)
+{
+	Enum = StaticEnum<TEnum>();
+
+	CurrentValue = InArgs._CurrentValue;
+	OnEnumSelectionChangedDelegate = InArgs._OnEnumSelectionChanged;
+
+	using namespace Chaos::VisualDebugger;
+	
+	static const FName UseEnumValuesAsMaskValuesInEditorName(TEXT("UseEnumValuesAsMaskValuesInEditor"));
+
+	const bool bUseEnumValuesAsMaskValues = Enum ? Enum->GetBoolMetaData(UseEnumValuesAsMaskValuesInEditorName) : false;
+
+	if (!ensure(bUseEnumValuesAsMaskValues))
+	{
+		this->ChildSlot
+		[
+			SNew(SVerticalBox)
+			+SVerticalBox::Slot()
+			[
+				SNew(STextBlock)
+				.Text(NSLOCTEXT("ChaosVisualDebugger", "ChaosVDEnumFlagsMenuErrorMessage", "Incompatible enum. Make sure to add the meta tag UseEnumValuesAsMaskValuesInEditor and it is a valid UEnum"))
+			]
+		];
+		return;
+	}
+
+	constexpr bool bCloseAfterSelection = false;
+	constexpr bool bCloseSelfOnly = true;
+
+	FMenuBuilder MenuBuilder(bCloseAfterSelection, nullptr, nullptr, bCloseSelfOnly);
+
+	const int32 EnumValuesCount = Enum->NumEnums() - 1;
+
+	for (int32 EnumValueIndex = 0; EnumValueIndex < EnumValuesCount ; EnumValueIndex++)
+	{
+		const int32 Value = Enum->GetValueByIndex(EnumValueIndex);
+		const int32 Index = EnumValueIndex;
+		const bool bIsHidden = Enum->HasMetaData(TEXT("Hidden"), Index);
+		if (Value >= 0 && !bIsHidden)
+		{
+			if (!FMath::IsPowerOfTwo(Value))
+			{
+				continue;
+			}
+	
+			FText DisplayName = Enum->GetDisplayNameTextByIndex(Index);
+			FText TooltipText = Enum->GetToolTipTextByIndex(Index);
+			if (TooltipText.IsEmpty())
+			{
+				TooltipText = FText::Format(NSLOCTEXT("ChaosVisualDebugger", "BitmaskDefaultFlagToolTipText", "Toggle {0} on/off"), DisplayName);
+			}
+
+			FEnumInfo EnumEntryInfo(Index, static_cast<TEnum>(Value), DisplayName, TooltipText);
+
+			MenuBuilder.AddMenuEntry(
+				DisplayName,
+				EnumEntryInfo.TooltipText,
+				FSlateIcon(),
+				FUIAction(
+				FExecuteAction::CreateSP(this, &SChaosVDEnumFlagsMenu::UpdateEnumFlagValue, EnumEntryInfo),
+				FCanExecuteAction::CreateLambda([](){ return true; }),
+				FIsActionChecked::CreateSP(this, &SChaosVDEnumFlagsMenu::IsEnumFlagSet, EnumEntryInfo)),
+				NAME_None, EUserInterfaceActionType::ToggleButton);
+		}
+	}
+
+	this->ChildSlot
+	[
+		MenuBuilder.MakeWidget()
+	];
+}
+
+template <typename TEnum>
+void SChaosVDEnumFlagsMenu<TEnum>::UpdateEnumFlagValue(FEnumInfo FlagInfo)
+{
+	// Toggle value
+	TEnum CurrentValueCopy = CurrentValue.Get();
+	Chaos::VisualDebugger::Utils::EnumAddToggleFlag(CurrentValueCopy, FlagInfo.Value);
+	OnEnumSelectionChangedDelegate.ExecuteIfBound(CurrentValueCopy);
+}
+
+template <typename TEnum>
+bool SChaosVDEnumFlagsMenu<TEnum>::IsEnumFlagSet(FEnumInfo FlagInfo) const
+{
+	return EnumHasAnyFlags(CurrentValue.Get(), FlagInfo.Value);
+}
diff --git a/Engine/Plugins/ChaosVD/Source/ChaosVD/Private/Widgets/SChaosVDMainTab.cpp b/Engine/Plugins/ChaosVD/Source/ChaosVD/Private/Widgets/SChaosVDMainTab.cpp
--- a/Engine/Plugins/ChaosVD/Source/ChaosVD/Private/Widgets/SChaosVDMainTab.cpp
+++ b/Engine/Plugins/ChaosVD/Source/ChaosVD/Private/Widgets/SChaosVDMainTab.cpp
@@ -3,7 +3,6 @@
 #include "Widgets/SChaosVDMainTab.h"
 
 #include "ChaosVDEditorModeTools.h"
-#include "ChaosVDEditorVisualizationSettingsTab.h"
 #include "ChaosVDEngine.h"
 #include "ChaosVDModule.h"
 #include "ChaosVDObjectDetailsTab.h"
@@ -74,7 +73,6 @@ void SChaosVDMainTab::Construct(const FArguments& InArgs, TSharedPtr<FChaosVDEng
 	RegisterTabSpawner<FChaosVDOutputLogTab>(FChaosVDTabID::OutputLog);
 	RegisterTabSpawner<FChaosVDPlaybackViewportTab>(FChaosVDTabID::PlaybackViewport);
 	RegisterTabSpawner<FChaosVDSolversTracksTab>(FChaosVDTabID::SolversTrack);
-	RegisterTabSpawner<FChaosVDEditorVisualizationSettingsTab>(FChaosVDTabID::CVDEditorSettings);
 	RegisterTabSpawner<FChaosVDCollisionDataDetailsTab>(FChaosVDTabID::CollisionDataDetails);
 	RegisterTabSpawner<FChaosVDSceneQueryDataInspectorTab>(FChaosVDTabID::SceneQueryDataDetails);
 	RegisterTabSpawner<FChaosVDConstraintDataInspectorTab>(FChaosVDTabID::JointsDataDetails);
@@ -288,7 +286,6 @@ TSharedRef<FTabManager::FLayout> SChaosVDMainTab::GenerateMainLayout()
 					->AddTab(FChaosVDTabID::SceneQueryDataDetails, ETabState::OpenedTab)
 					->AddTab(FChaosVDTabID::JointsDataDetails, ETabState::ClosedTab)
 					->AddTab(FChaosVDTabID::CharacterGroundConstraintDataDetails, ETabState::ClosedTab)
-					->AddTab(FChaosVDTabID::CVDEditorSettings, ETabState::ClosedTab)
 				)
 			)
 		);
diff --git a/Engine/Plugins/ChaosVD/Source/ChaosVD/Private/Widgets/SChaosVDPlaybackViewport.cpp b/Engine/Plugins/ChaosVD/Source/ChaosVD/Private/Widgets/SChaosVDPlaybackViewport.cpp
--- a/Engine/Plugins/ChaosVD/Source/ChaosVD/Private/Widgets/SChaosVDPlaybackViewport.cpp
+++ b/Engine/Plugins/ChaosVD/Source/ChaosVD/Private/Widgets/SChaosVDPlaybackViewport.cpp
@@ -4,15 +4,11 @@
 
 #include "ChaosVDCommands.h"
 #include "ChaosVDEditorMode.h"
-#include "ChaosVDEditorModeTools.h"
-#include "ChaosVDEditorSettings.h"
-#include "ChaosVDModule.h"
 #include "ChaosVDPlaybackController.h"
 #include "ChaosVDPlaybackViewportClient.h"
 #include "ChaosVDScene.h"
 #include "EditorModeManager.h"
 #include "Framework/Application/SlateApplication.h"
-#include "Elements/Framework/TypedElementSelectionSet.h"
 #include "Widgets/ChaosVDPlaybackControlsHelper.h"
 #include "Widgets/SChaosVDTimelineWidget.h"
 #include "Widgets/SChaosVDViewportToolbar.h"
@@ -31,6 +27,12 @@ namespace Chaos::VisualDebugger::Cvars
 
 SChaosVDPlaybackViewport::~SChaosVDPlaybackViewport()
 {
+	if (ExternalInvalidateHandlerHandle.IsValid())
+	{
+		ExternalViewportInvalidationRequestHandler.Remove(ExternalInvalidateHandlerHandle);
+		ExternalInvalidateHandlerHandle.Reset();
+	}
+
 	PlaybackViewportClient->Viewport = nullptr;
 	PlaybackViewportClient.Reset();
 }
@@ -55,7 +57,7 @@ void SChaosVDPlaybackViewport::Construct(const FArguments& InArgs, TWeakPtr<FCha
 
 	// TODO: Add a way to gracefully shutdown (close) the tool when a no recoverable situation like this happens (UE-191876)
 	check(PlaybackViewportClient.IsValid());
-	
+
 	PlaybackViewportClient->SetScene(InScene);
 	
 	if (UChaosVDEditorMode* CVDEdMode = Cast<UChaosVDEditorMode>(EditorModeTools->GetActiveScriptableMode(UChaosVDEditorMode::EM_ChaosVisualDebugger)))
@@ -103,6 +105,8 @@ void SChaosVDPlaybackViewport::Construct(const FArguments& InArgs, TWeakPtr<FCha
 		]	
 	];
 
+	ExternalInvalidateHandlerHandle = ExternalViewportInvalidationRequestHandler.AddSP(this, &SChaosVDPlaybackViewport::HanldeExternalViewportInvalidateRequest);
+
 	RegisterNewController(InPlaybackController);
 }
 
@@ -125,9 +129,29 @@ void SChaosVDPlaybackViewport::BindCommands()
 	if (ensure(Client))
 	{
 		const TSharedRef<FChaosVDPlaybackViewportClient> ViewportClientRef = StaticCastSharedRef<FChaosVDPlaybackViewportClient>(Client.ToSharedRef());
+		FUIAction ToggleObjectTrackingAction;
+		ToggleObjectTrackingAction.ExecuteAction.BindSP(ViewportClientRef, &FChaosVDPlaybackViewportClient::ToggleObjectTrackingIfSelected);
+		ToggleObjectTrackingAction.GetActionCheckState.BindLambda([WeakThis = ViewportClientRef.ToWeakPtr()]()
+		{
+			TSharedPtr<const FChaosVDPlaybackViewportClient> ViewportPtr = StaticCastSharedPtr<const FChaosVDPlaybackViewportClient>(WeakThis.Pin());
+			return ViewportPtr.IsValid() && ViewportPtr->IsAutoTrackingSelectedObject() ? ECheckBoxState::Checked : ECheckBoxState::Unchecked;
+		});
+	
 		CommandList->MapAction(
-			Commands.TrackUntrackSelectedObject,
-			FExecuteAction::CreateSP(ViewportClientRef, &FChaosVDPlaybackViewportClient::ToggleObjectTrackingIfSelected));
+			Commands.ToggleFollowSelectedObject,
+			ToggleObjectTrackingAction);
+
+		FUIAction ToggleOverrideFrameRateAction;
+		ToggleOverrideFrameRateAction.ExecuteAction.BindSP(SharedThis(this), &SChaosVDPlaybackViewport::ToggleUseFrameRateOverride);
+		ToggleOverrideFrameRateAction.GetActionCheckState.BindLambda([WeakThis = AsWeak()]()
+		{
+			TSharedPtr<const SChaosVDPlaybackViewport> ViewportPtr = StaticCastSharedPtr<const SChaosVDPlaybackViewport>(WeakThis.Pin());
+			return ViewportPtr.IsValid() && ViewportPtr->IsUsingFrameRateOverride() ? ECheckBoxState::Checked : ECheckBoxState::Unchecked;
+		});
+
+		CommandList->MapAction(
+			Commands.OverridePlaybackFrameRate,
+			ToggleOverrideFrameRateAction);
 	}
 }
 
@@ -146,6 +170,43 @@ void SChaosVDPlaybackViewport::GoToLocation(const FVector& InLocation) const
 	}
 }
 
+void SChaosVDPlaybackViewport::ToggleUseFrameRateOverride()
+{
+	if (TSharedPtr<FChaosVDPlaybackController> PlaybackControllerPtr = PlaybackController.Pin())
+	{
+		PlaybackControllerPtr->ToggleUseFrameRateOverride();
+	}
+}
+
+bool SChaosVDPlaybackViewport::IsUsingFrameRateOverride() const
+{
+	if (TSharedPtr<FChaosVDPlaybackController> PlaybackControllerPtr = PlaybackController.Pin())
+	{
+		return PlaybackControllerPtr->IsUsingFrameRateOverride();
+	}
+
+	return false;
+}
+
+int32 SChaosVDPlaybackViewport::GetCurrentTargetFrameRateOverride() const
+{
+	TSharedPtr<FChaosVDPlaybackController> PlaybackControllerPtr = PlaybackController.Pin();
+	return PlaybackControllerPtr ? PlaybackControllerPtr->GetFrameRateOverride() : FChaosVDPlaybackController::InvalidFrameRateOverride;
+}
+
+void SChaosVDPlaybackViewport::SetCurrentTargetFrameRateOverride(int32 NewTarget)
+{
+	if (TSharedPtr<FChaosVDPlaybackController> PlaybackControllerPtr = PlaybackController.Pin())
+	{
+		PlaybackControllerPtr->SetFrameRateOverride(NewTarget);
+	}
+}
+
+void SChaosVDPlaybackViewport::ExecuteExternalViewportInvalidateRequest()
+{
+	ExternalViewportInvalidationRequestHandler.Broadcast();
+}
+
 TSharedRef<FEditorViewportClient> SChaosVDPlaybackViewport::MakeEditorViewportClient()
 {
 	TSharedPtr<FChaosVDPlaybackViewportClient> NewViewport = MakeShared<FChaosVDPlaybackViewportClient>(EditorModeTools, GetViewportWidget());
@@ -321,4 +382,13 @@ bool SChaosVDPlaybackViewport::CanPlayback() const
 	return PlaybackControllerPtr && PlaybackControllerPtr->IsRecordingLoaded();
 }
 
+
+void SChaosVDPlaybackViewport::HanldeExternalViewportInvalidateRequest()
+{
+	if (PlaybackViewportClient)
+	{
+		PlaybackViewportClient->Invalidate();
+	}
+}
+
 #undef LOCTEXT_NAMESPACE
diff --git a/Engine/Plugins/ChaosVD/Source/ChaosVD/Private/Widgets/SChaosVDPlaybackViewport.h b/Engine/Plugins/ChaosVD/Source/ChaosVD/Private/Widgets/SChaosVDPlaybackViewport.h
--- a/Engine/Plugins/ChaosVD/Source/ChaosVD/Private/Widgets/SChaosVDPlaybackViewport.h
+++ b/Engine/Plugins/ChaosVD/Source/ChaosVD/Private/Widgets/SChaosVDPlaybackViewport.h
@@ -46,6 +46,17 @@ public:
 
 	void GoToLocation(const FVector& InLocation) const;
 
+	void ToggleUseFrameRateOverride();
+	bool IsUsingFrameRateOverride() const;
+
+	int32 GetCurrentTargetFrameRateOverride() const;
+
+	void SetCurrentTargetFrameRateOverride(int32 NewTarget);
+
+	TWeakPtr<FChaosVDScene> GetCVDScene() { return CVDSceneWeakPtr; }
+	
+	static void ExecuteExternalViewportInvalidateRequest();
+
 protected:
 
 	virtual TSharedRef<FEditorViewportClient> MakeEditorViewportClient() override;
@@ -66,6 +77,8 @@ protected:
 
 	bool CanPlayback() const;
 
+	void HanldeExternalViewportInvalidateRequest();
+
 	TSharedPtr<SChaosVDTimelineWidget> GameFramesTimelineWidget;
 
 	TSharedPtr<FChaosVDPlaybackViewportClient> PlaybackViewportClient;
@@ -75,4 +88,9 @@ protected:
 	TSharedPtr<FExtender> Extender;
 
 	TSharedPtr<FEditorModeTools> EditorModeTools;
+	
+	DECLARE_MULTICAST_DELEGATE(FChaosVDViewportInvalidationRequestHandler)
+	static inline FChaosVDViewportInvalidationRequestHandler ExternalViewportInvalidationRequestHandler = FChaosVDViewportInvalidationRequestHandler();
+	
+	FDelegateHandle ExternalInvalidateHandlerHandle;
 };
diff --git a/Engine/Plugins/ChaosVD/Source/ChaosVD/Private/Widgets/SChaosVDSolverPlaybackControls.cpp b/Engine/Plugins/ChaosVD/Source/ChaosVD/Private/Widgets/SChaosVDSolverPlaybackControls.cpp
--- a/Engine/Plugins/ChaosVD/Source/ChaosVD/Private/Widgets/SChaosVDSolverPlaybackControls.cpp
+++ b/Engine/Plugins/ChaosVD/Source/ChaosVD/Private/Widgets/SChaosVDSolverPlaybackControls.cpp
@@ -2,8 +2,6 @@
 
 #include "Widgets/SChaosVDSolverPlaybackControls.h"
 
-#include "ChaosVDEditorSettings.h"
-#include "ChaosVDModule.h"
 #include "ChaosVDPlaybackController.h"
 #include "ChaosVDScene.h"
 #include "Widgets/ChaosVDPlaybackControlsHelper.h"
diff --git a/Engine/Plugins/ChaosVD/Source/ChaosVD/Private/Widgets/SChaosVDViewportToolbar.cpp b/Engine/Plugins/ChaosVD/Source/ChaosVD/Private/Widgets/SChaosVDViewportToolbar.cpp
--- a/Engine/Plugins/ChaosVD/Source/ChaosVD/Private/Widgets/SChaosVDViewportToolbar.cpp
+++ b/Engine/Plugins/ChaosVD/Source/ChaosVD/Private/Widgets/SChaosVDViewportToolbar.cpp
@@ -3,25 +3,56 @@
 #include "Widgets/SChaosVDViewportToolbar.h"
 
 #include "ChaosVDCommands.h"
-#include "ChaosVDEditorSettings.h"
+
+#include "ChaosVDPlaybackViewportClient.h"
 #include "Framework/MultiBox/MultiBoxBuilder.h"
-#include "PropertyEditorModule.h"
-#include "DetailsViewArgs.h"
-#include "IDetailsView.h"
-#include "Modules/ModuleManager.h"
-#include "PropertyEditorModule.h"
 #include "SCommonEditorViewportToolbarBase.h"
 #include "SEditorViewport.h"
+#include "ShowFlagFilter.h"
+#include "ShowFlagMenuCommands.h"
 #include "Styling/AppStyle.h"
 #include "Widgets/Input/SEditableText.h"
 #include "Widgets/Layout/SBorder.h"
 #include "Widgets/Layout/SBox.h"
+#include "Widgets/Input/SSpinBox.h"
 #include "Widgets/SChaosVDEditorViewportViewMenu.h"
 #include "Widgets/SChaosVDPlaybackViewport.h"
-#include "Widgets/SChaosVDVisualizationControls.h"
+#include "ToolMenus.h"
 
 #define LOCTEXT_NAMESPACE "ChaosVisualDebugger"
 
+namespace Chaos::VisualDebugger::Utils
+{
+	template <typename TNumericValue>
+	TSharedRef<SWidget> GenerateSpinBoxMenuEntryWidget(const FText& ToolTipText, TNumericValue MinValue, TNumericValue MaxValue, typename SSpinBox<TNumericValue>::FOnValueChanged&& ValueChangedDelegate, const TAttribute<TNumericValue>&& ValueAttribute, const TAttribute<bool>&& EnabledAttribute)
+	{
+		return SNew(SBox)
+			.HAlign(HAlign_Right)
+			[
+				SNew(SBox)
+				.Padding(FMargin(4.0f, 0.0f, 0.0f, 0.0f))
+				.WidthOverride(100.0f)
+				[
+					SNew ( SBorder )
+					.BorderImage(FAppStyle::Get().GetBrush("Menu.WidgetBorder"))
+					.Padding(FMargin(1.0f))
+					[
+						SNew(SSpinBox<TNumericValue>)
+						.Style(&FAppStyle::Get(), "Menu.SpinBox")
+						.ToolTipText(ToolTipText)
+						.MinValue(MinValue)
+						.MaxValue(MaxValue)
+						.Font(FAppStyle::GetFontStyle(TEXT("MenuItem.Font")))
+						.Value(ValueAttribute)
+						.OnValueChanged(ValueChangedDelegate)
+						.IsEnabled(EnabledAttribute)
+					]
+				]
+			];
+	}
+
+}
+
 void SChaosVDViewportToolbar::Construct(const FArguments& InArgs, TSharedPtr<ICommonEditorViewportToolbarInfoProvider> InInfoProvider)
 {
 	SCommonEditorViewportToolbarBase::Construct(SCommonEditorViewportToolbarBase::FArguments(), InInfoProvider);
@@ -37,10 +68,38 @@ void SChaosVDViewportToolbar::ExtendOptionsMenu(FMenuBuilder& OptionsMenuBuilder
 {
 	constexpr bool bInShouldCloseWindowAfterMenuSelection = true;
 	FMenuBuilder CVDOptionsMenuBuilder(bInShouldCloseWindowAfterMenuSelection, GetInfoProvider().GetViewportWidget()->GetCommandList());
-
-	CVDOptionsMenuBuilder.BeginSection("CVDViewportViewportOptions", LOCTEXT("OptionsMenuHeader", "Utils"));
 	
-	CVDOptionsMenuBuilder.AddWidget(GenerateGoToLocationWidget(), LOCTEXT("GoToLocation", "Go to Location"));
+	CVDOptionsMenuBuilder.BeginSection("CVDViewportViewportOptions", LOCTEXT("ViewportOptionsMenuHeader", "Viewport Options"));
+	{
+		CVDOptionsMenuBuilder.AddSubMenu(
+			LOCTEXT("FrameRateOptionsMenuLabel", "Playback Framerate"),
+			LOCTEXT("FrameRateOptionsMenuToolTip", "Options that control how CVD plays a recording."),
+			FNewMenuDelegate::CreateSP(this, &SChaosVDViewportToolbar::PopulateFrameRateSubmenu),
+			false,
+			FSlateIcon(FAppStyle::Get().GetStyleSetName(), TEXT("EditorViewport.ToggleFPS")
+		));
+
+		CVDOptionsMenuBuilder.AddSubMenu(
+		LOCTEXT("FrameRateOptionsMenuLabel", "Object Tracking"),
+		LOCTEXT("FrameRateOptionsMenuToolTip", "Options that control how objects are tracked in the scene by the camera."),
+		FNewMenuDelegate::CreateSP(this, &SChaosVDViewportToolbar::PopulateAutoTrackingSubMenu),
+		false,
+		FSlateIcon(FAppStyle::Get().GetStyleSetName(), TEXT("AnimViewportMenu.CameraFollow.Small"))
+		);
+
+		CVDOptionsMenuBuilder.AddSeparator();
+
+		CVDOptionsMenuBuilder.AddWidget(GenerateFOVMenu(), LOCTEXT("FOVAngle", "Field of View (H)"));
+		CVDOptionsMenuBuilder.AddWidget(GenerateFarViewPlaneMenu(), LOCTEXT("FarViewPlane", "Far View Plane"));
+
+		CVDOptionsMenuBuilder.AddSeparator();
+	}
+	CVDOptionsMenuBuilder.EndSection();
+
+	CVDOptionsMenuBuilder.BeginSection("CVDViewportViewportUtils", LOCTEXT("ViewportUtilMenuHeader", "Utils"));
+	{
+		CVDOptionsMenuBuilder.AddWidget(GenerateGoToLocationWidget(), LOCTEXT("GoToLocation", "Go to Location"));
+	}
 	CVDOptionsMenuBuilder.EndSection();
 
 	OptionsMenuBuilder = CVDOptionsMenuBuilder;
@@ -68,20 +127,132 @@ TSharedRef<SWidget> SChaosVDViewportToolbar::GenerateGoToLocationWidget() const
 			];
 }
 
-TSharedRef<SWidget> SChaosVDViewportToolbar::GenerateShowMenu() const
+TSharedRef<SWidget> SChaosVDViewportToolbar::GenerateFrameRateOverrideValueWidget() const
 {
-	FPropertyEditorModule& PropertyEditorModule = FModuleManager::GetModuleChecked<FPropertyEditorModule>("PropertyEditor");
-	FDetailsViewArgs DetailsViewArgs;
-	DetailsViewArgs.bAllowSearch = false;
-	DetailsViewArgs.NameAreaSettings = FDetailsViewArgs::HideNameArea;
+	TAttribute<int32> ValueAttribute;
+	ValueAttribute.Bind(TAttribute<int32>::FGetter::CreateSP(this, &SChaosVDViewportToolbar::OnGetFrameRateOverrideValue));
+	
+	TAttribute<bool> EnabledAttribute;
+	EnabledAttribute.Bind(TAttribute<bool>::FGetter::CreateSP(this, &SChaosVDViewportToolbar::IsUsingFrameRateOverride));
+
+	SSpinBox<int32>::FOnValueChanged ValueChangedDelegate = SSpinBox<int32>::FOnValueChanged::CreateSP(this, &SChaosVDViewportToolbar::OnFrameRateOverrideValueChanged);
+
+	constexpr int32 MinValue = 1;
+	constexpr int32 MaxValue = 1000;
+
+	return Chaos::VisualDebugger::Utils::GenerateSpinBoxMenuEntryWidget(LOCTEXT("FramerateOverrideTooltip", "Target framerate we should play the loaded recording at"), MinValue, MaxValue, MoveTemp(ValueChangedDelegate), MoveTemp(ValueAttribute), MoveTemp(EnabledAttribute));
+}
 
-	TSharedRef<IDetailsView> DetailsPanel = PropertyEditorModule.CreateDetailView(DetailsViewArgs);
+TSharedRef<SWidget> SChaosVDViewportToolbar::GenerateTrackingDistanceValueWidget() const
+{
+	TAttribute<float> ValueAttribute;
+	ValueAttribute.Bind(TAttribute<float>::FGetter::CreateSP(this, &SChaosVDViewportToolbar::OnGetTrackingDistanceValue));
+	
+	TAttribute<bool> EnabledAttribute;
+	EnabledAttribute.Bind(TAttribute<bool>::FGetter::CreateSP(this, &SChaosVDViewportToolbar::IsAutoTrackingEnabled));
+
+	SSpinBox<float>::FOnValueChanged ValueChangedDelegate = SSpinBox<float>::FOnValueChanged::CreateSP(this, &SChaosVDViewportToolbar::OnTrackingDistanceValueChanged);
 
-	DetailsPanel->SetObject(GetMutableDefault<UChaosVDEditorSettings>());
+	constexpr float MinValue = 1.0f;
+	constexpr float MaxValue = 100000.0f;
+
+	return Chaos::VisualDebugger::Utils::GenerateSpinBoxMenuEntryWidget(LOCTEXT("TrackingDistanceTooltip", "Distance from which we want to track the selected object"), MinValue, MaxValue, MoveTemp(ValueChangedDelegate), MoveTemp(ValueAttribute), MoveTemp(EnabledAttribute));
+}
 
+int32 SChaosVDViewportToolbar::OnGetFrameRateOverrideValue() const
+{
+	TSharedRef<SChaosVDPlaybackViewport> ViewportRef = StaticCastSharedRef<SChaosVDPlaybackViewport>(GetInfoProvider().GetViewportWidget());
+	return ViewportRef->GetCurrentTargetFrameRateOverride();
+}
+
+void SChaosVDViewportToolbar::OnFrameRateOverrideValueChanged(int32 NewFrameRate) const
+{
 	TSharedRef<SChaosVDPlaybackViewport> ViewportRef = StaticCastSharedRef<SChaosVDPlaybackViewport>(GetInfoProvider().GetViewportWidget());
+	return ViewportRef->SetCurrentTargetFrameRateOverride(NewFrameRate);
+}
+
+bool SChaosVDViewportToolbar::IsUsingFrameRateOverride() const
+{
+	TSharedRef<SChaosVDPlaybackViewport> ViewportRef = StaticCastSharedRef<SChaosVDPlaybackViewport>(GetInfoProvider().GetViewportWidget());
+	return ViewportRef->IsUsingFrameRateOverride();
+}
+
+float SChaosVDViewportToolbar::OnGetTrackingDistanceValue() const
+{
+	const TSharedPtr<FChaosVDPlaybackViewportClient> CVDViewportClient = StaticCastSharedPtr<FChaosVDPlaybackViewportClient>(GetInfoProvider().GetViewportWidget()->GetViewportClient());
+
+	return CVDViewportClient ? CVDViewportClient->GetAutoTrackingViewDistance() : -1.0f;
+}
+
+void SChaosVDViewportToolbar::OnTrackingDistanceValueChanged(float NewTrackingDistance) const
+{
+	if (const TSharedPtr<FChaosVDPlaybackViewportClient> CVDViewportClient = StaticCastSharedPtr<FChaosVDPlaybackViewportClient>(GetInfoProvider().GetViewportWidget()->GetViewportClient()))
+	{
+		CVDViewportClient->SetAutoTrackingViewDistance(NewTrackingDistance);
+	}
+}
+
+bool SChaosVDViewportToolbar::IsAutoTrackingEnabled() const
+{
+	const TSharedPtr<FChaosVDPlaybackViewportClient> CVDViewportClient = StaticCastSharedPtr<FChaosVDPlaybackViewportClient>(GetInfoProvider().GetViewportWidget()->GetViewportClient());
+
+	return CVDViewportClient ? CVDViewportClient->IsAutoTrackingSelectedObject() : false;
+}
+
+void SChaosVDViewportToolbar::PopulateFrameRateSubmenu(FMenuBuilder& MenuBuilder) const
+{
+	MenuBuilder.AddMenuEntry(FChaosVDCommands::Get().OverridePlaybackFrameRate, NAME_None);
+	MenuBuilder.AddWidget(GenerateFrameRateOverrideValueWidget(), LOCTEXT("FrameRateOverride", "Target Framerate"));
+}
+
+void SChaosVDViewportToolbar::PopulateAutoTrackingSubMenu(FMenuBuilder& MenuBuilder) const
+{
+	MenuBuilder.AddMenuEntry(FChaosVDCommands::Get().ToggleFollowSelectedObject, NAME_None);
+	MenuBuilder.AddWidget(GenerateTrackingDistanceValueWidget(), LOCTEXT("TrackingDistance", "Follow Distance"));
+}
+
+TSharedRef<SWidget> SChaosVDViewportToolbar::GenerateShowMenu() const
+{	
+	GetInfoProvider().OnFloatingButtonClicked();
+
+	if (!UToolMenus::Get()->IsMenuRegistered(ShowMenuName))
+	{
+		if (UToolMenu* Menu = UToolMenus::Get()->RegisterMenu(ShowMenuName))
+		{
+			FToolMenuSection& Section = Menu->AddSection("ChaosVDViewportToolbarBase.Show.CommonViewportFlags", LOCTEXT("ToolbarCommonViewportFlags", "Common Show Flags"));
+
+			FNewToolMenuDelegate CustomShowMenuBuilder = FNewToolMenuDelegate::CreateLambda([](UToolMenu* Menu)
+															   {
+																	// Only include the flags that might be helpful.
+																	static const FShowFlagFilter ShowFlagFilter = FShowFlagFilter(FShowFlagFilter::ExcludeAllFlagsByDefault)
+																		.IncludeFlag(FEngineShowFlags::SF_AntiAliasing)
+																		.IncludeFlag(FEngineShowFlags::SF_Grid)
+																		.IncludeFlag(FEngineShowFlags::SF_Translucency)
+																		.IncludeFlag(FEngineShowFlags::SF_MeshEdges)
+																		.IncludeFlag(FEngineShowFlags::SF_HitProxies)
+																		.IncludeFlag(FEngineShowFlags::SF_Fog)
+																		.IncludeFlag(FEngineShowFlags::SF_Pivot);
+
+																	FShowFlagMenuCommands::Get().BuildShowFlagsMenu(Menu, ShowFlagFilter);
+															   });
+
+			constexpr bool bOpenSubMenuOnClick = false;
+			Section.AddSubMenu(TEXT("CommonViewportFlags"), LOCTEXT("CommonShowFlagsMenuLabel", "Common Show Flags"), LOCTEXT("CommonShowFlagsMenuToolTip", "Set of flags to enable/disable specific viewport features"), CustomShowMenuBuilder,
+							   bOpenSubMenuOnClick, FSlateIcon(FAppStyle::Get().GetStyleSetName(), TEXT("Icons.Toolbar.Settings")));
+		}
+	}
+
+	FToolMenuContext NewMenuContext;
+	UCommonViewportToolbarBaseMenuContext* ContextObject = NewObject<UCommonViewportToolbarBaseMenuContext>();
+	ContextObject->ToolbarWidget = SharedThis(this);
+	NewMenuContext.AddObject(ContextObject);
+
+	if (TSharedPtr<SEditorViewport> ViewportWidget = GetInfoProvider().GetViewportWidget())
+	{
+		NewMenuContext.AppendCommandList(GetInfoProvider().GetViewportWidget()->GetCommandList());
+	}
 
-	return SNew(SChaosVDVisualizationControls, ViewportRef->GetObservedController(), ViewportRef);
+	return UToolMenus::Get()->GenerateWidget(ShowMenuName, NewMenuContext);
 }
 
 void SChaosVDViewportToolbar::HandleGoToLocationCommited(const FText& InLocationAsText, ETextCommit::Type Type) const
diff --git a/Engine/Plugins/ChaosVD/Source/ChaosVD/Private/Widgets/SChaosVDViewportToolbar.h b/Engine/Plugins/ChaosVD/Source/ChaosVD/Private/Widgets/SChaosVDViewportToolbar.h
--- a/Engine/Plugins/ChaosVD/Source/ChaosVD/Private/Widgets/SChaosVDViewportToolbar.h
+++ b/Engine/Plugins/ChaosVD/Source/ChaosVD/Private/Widgets/SChaosVDViewportToolbar.h
@@ -17,21 +17,36 @@ class CHAOSVD_API SChaosVDViewportToolbar : public SCommonEditorViewportToolbarB
 {
 public:
 	SLATE_BEGIN_ARGS(SChaosVDViewportToolbar)
-		{
-		}
+	{
+	}
 	SLATE_END_ARGS()
 
 	void Construct(const FArguments& InArgs, TSharedPtr<ICommonEditorViewportToolbarInfoProvider> InInfoProvider);
 
+	static inline FName ShowMenuName = FName("ChaosVDViewportToolbarBase.Show");
+
 protected:
 	virtual TSharedRef<SEditorViewportViewMenu> MakeViewMenu() override;
 	virtual void ExtendOptionsMenu(FMenuBuilder& OptionsMenuBuilder) const override;
 
 	TSharedRef<SWidget> GenerateGoToLocationWidget() const;
+	TSharedRef<SWidget> GenerateFrameRateOverrideValueWidget() const;
+	TSharedRef<SWidget> GenerateTrackingDistanceValueWidget() const;
+
+	int32 OnGetFrameRateOverrideValue() const;
+	void OnFrameRateOverrideValueChanged(int32 NewFrameRate) const;
+	bool IsUsingFrameRateOverride() const;
+
+	float OnGetTrackingDistanceValue() const;
+	void OnTrackingDistanceValueChanged(float NewTrackingDistance) const;
+	bool IsAutoTrackingEnabled() const;
+
+	void PopulateFrameRateSubmenu(FMenuBuilder& MenuBuilder) const;
+	void PopulateAutoTrackingSubMenu(FMenuBuilder& MenuBuilder) const;
 
 private:
 	virtual TSharedRef<SWidget> GenerateShowMenu() const override;
 
 	void HandleGoToLocationCommited(const FText& InLocationAsText, ETextCommit::Type Type) const;
 
-};
+};
\ No newline at end of file
diff --git a/Engine/Plugins/ChaosVD/Source/ChaosVD/Private/Widgets/SChaosVDVisualizationControls.cpp b/Engine/Plugins/ChaosVD/Source/ChaosVD/Private/Widgets/SChaosVDVisualizationControls.cpp
deleted file mode 100644
diff --git a/Engine/Plugins/ChaosVD/Source/ChaosVD/Private/Widgets/SChaosVDVisualizationControls.h b/Engine/Plugins/ChaosVD/Source/ChaosVD/Private/Widgets/SChaosVDVisualizationControls.h
deleted file mode 100644
diff --git a/Engine/Source/Editor/UnrealEd/Private/SCommonEditorViewportToolbarBase.cpp b/Engine/Source/Editor/UnrealEd/Private/SCommonEditorViewportToolbarBase.cpp
--- a/Engine/Source/Editor/UnrealEd/Private/SCommonEditorViewportToolbarBase.cpp
+++ b/Engine/Source/Editor/UnrealEd/Private/SCommonEditorViewportToolbarBase.cpp
@@ -752,7 +752,9 @@ float SCommonEditorViewportToolbarBase::OnGetFarViewPlaneValue() const
 
 void SCommonEditorViewportToolbarBase::OnFarViewPlaneValueChanged(float NewValue)
 {
-	GetViewportClient().OverrideFarClipPlane(NewValue);
+	FEditorViewportClient& ViewportClient = GetViewportClient();
+	ViewportClient.OverrideFarClipPlane(NewValue);
+	ViewportClient.Invalidate();
 }
 
 FReply SCommonEditorViewportToolbarBase::OnRealtimeWarningClicked()
