[ChaosVD] Fix for CIS warnings caused by my last change

diff --git a/Engine/Plugins/ChaosVD/Source/ChaosVD/Private/ChaosVDPlaybackController.cpp b/Engine/Plugins/ChaosVD/Source/ChaosVD/Private/ChaosVDPlaybackController.cpp
--- a/Engine/Plugins/ChaosVD/Source/ChaosVD/Private/ChaosVDPlaybackController.cpp
+++ b/Engine/Plugins/ChaosVD/Source/ChaosVD/Private/ChaosVDPlaybackController.cpp
@@ -730,19 +730,19 @@ bool FChaosVDPlaybackController::ReleaseExclusivePlaybackControls(const IChaosVD
 float FChaosVDPlaybackController::GetFrameTimeOverride() const
 {
 	constexpr int32 MinimumFrameRateOverride = 1;
-	return CurrentFrameRateOverride >= MinimumFrameRateOverride ? 1.0f / static_cast<float>(CurrentFrameRateOverride) : InvalidFrameRateOverride;
+	return (CurrentFrameRateOverride >= MinimumFrameRateOverride) ? 1.0f / static_cast<float>(CurrentFrameRateOverride) : InvalidFrameRateOverride;
 }
 
 int32 FChaosVDPlaybackController::GetFrameRateOverride() const
 {
 	constexpr int32 MinimumFrameRateOverride = 1;
-	return CurrentFrameRateOverride >= MinimumFrameRateOverride ? CurrentFrameRateOverride : InvalidFrameRateOverride;
+	return (CurrentFrameRateOverride >= MinimumFrameRateOverride) ? CurrentFrameRateOverride : InvalidFrameRateOverride;
 }
 
 void FChaosVDPlaybackController::SetFrameRateOverride(float NewFrameRateOverride)
 {
 	constexpr int32 MinimumFrameRateOverride = 1;
-	CurrentFrameRateOverride = NewFrameRateOverride >= MinimumFrameRateOverride ? NewFrameRateOverride : InvalidFrameRateOverride;
+	CurrentFrameRateOverride = (NewFrameRateOverride >= MinimumFrameRateOverride) ? NewFrameRateOverride : InvalidFrameRateOverride;
 }
 
 float FChaosVDPlaybackController::GetFrameTimeForTrack(EChaosVDTrackType TrackType, int32 TrackID, const FChaosVDTrackInfo& TrackInfo) const
diff --git a/Engine/Plugins/ChaosVD/Source/ChaosVD/Private/Settings/ChaosVDCharacterConstraintsVisualizationSettings.cpp b/Engine/Plugins/ChaosVD/Source/ChaosVD/Private/Settings/ChaosVDCharacterConstraintsVisualizationSettings.cpp
--- a/Engine/Plugins/ChaosVD/Source/ChaosVD/Private/Settings/ChaosVDCharacterConstraintsVisualizationSettings.cpp
+++ b/Engine/Plugins/ChaosVD/Source/ChaosVD/Private/Settings/ChaosVDCharacterConstraintsVisualizationSettings.cpp
@@ -7,7 +7,7 @@ void UChaosVDCharacterConstraintsVisualizationSettings::SetCharacterGroundConstr
 	if (UChaosVDCharacterConstraintsVisualizationSettings* Settings = GetMutableDefault<UChaosVDCharacterConstraintsVisualizationSettings>())
 	{
 		Settings->GlobalCharacterGroundConstraintDataVisualizationFlags = NewFlags;
-		Settings->BroadcastSettingsChanged(Settings);
+		Settings->BroadcastSettingsChanged();
 	}
 }
 
diff --git a/Engine/Plugins/ChaosVD/Source/ChaosVD/Private/Settings/ChaosVDCollisionVisualizationSettings.cpp b/Engine/Plugins/ChaosVD/Source/ChaosVD/Private/Settings/ChaosVDCollisionVisualizationSettings.cpp
--- a/Engine/Plugins/ChaosVD/Source/ChaosVD/Private/Settings/ChaosVDCollisionVisualizationSettings.cpp
+++ b/Engine/Plugins/ChaosVD/Source/ChaosVD/Private/Settings/ChaosVDCollisionVisualizationSettings.cpp
@@ -7,7 +7,7 @@ void UChaosVDCollisionDataVisualizationSettings::SetCollisionDataVisualizationFl
 	if (UChaosVDCollisionDataVisualizationSettings* Settings = GetMutableDefault<UChaosVDCollisionDataVisualizationSettings>())
 	{
 		Settings->CollisionDataVisualizationFlags = NewFlags;
-		Settings->BroadcastSettingsChanged(Settings);
+		Settings->BroadcastSettingsChanged();
 	}
 }
 
diff --git a/Engine/Plugins/ChaosVD/Source/ChaosVD/Private/Settings/ChaosVDCoreSettings.cpp b/Engine/Plugins/ChaosVD/Source/ChaosVD/Private/Settings/ChaosVDCoreSettings.cpp
--- a/Engine/Plugins/ChaosVD/Source/ChaosVD/Private/Settings/ChaosVDCoreSettings.cpp
+++ b/Engine/Plugins/ChaosVD/Source/ChaosVD/Private/Settings/ChaosVDCoreSettings.cpp
@@ -9,23 +9,23 @@ void UChaosVDSettingsObjectBase::PostEditChangeProperty(FPropertyChangedEvent& P
 {
 	UObject::PostEditChangeProperty(PropertyChangedEvent);
 
-	BroadcastSettingsChanged(this);
+	BroadcastSettingsChanged();
 }
 
 void UChaosVDSettingsObjectBase::PostEditUndo()
 {
 	UObject::PostEditUndo();
-	BroadcastSettingsChanged(this);
+	BroadcastSettingsChanged();
 }
 
-void UChaosVDSettingsObjectBase::BroadcastSettingsChanged(UObject* SettingsObject)
+void UChaosVDSettingsObjectBase::BroadcastSettingsChanged()
 {
 	SettingsChangedDelegate.Broadcast(this);
 }
 
-void UChaosVDVisualizationSettingsObjectBase::BroadcastSettingsChanged(UObject* SettingsObject)
+void UChaosVDVisualizationSettingsObjectBase::BroadcastSettingsChanged()
 {
-	Super::BroadcastSettingsChanged(SettingsObject);
+	Super::BroadcastSettingsChanged();
 	SChaosVDPlaybackViewport::ExecuteExternalViewportInvalidateRequest();
 }
 
diff --git a/Engine/Plugins/ChaosVD/Source/ChaosVD/Private/Settings/ChaosVDCoreSettings.h b/Engine/Plugins/ChaosVD/Source/ChaosVD/Private/Settings/ChaosVDCoreSettings.h
--- a/Engine/Plugins/ChaosVD/Source/ChaosVD/Private/Settings/ChaosVDCoreSettings.h
+++ b/Engine/Plugins/ChaosVD/Source/ChaosVD/Private/Settings/ChaosVDCoreSettings.h
@@ -27,7 +27,7 @@ public:
 	virtual void PostEditUndo() override;
 
 protected:
-	virtual void BroadcastSettingsChanged(UObject* SettingsObject);
+	virtual void BroadcastSettingsChanged();
 	
 private:
 	FChaosVDSettingChanged SettingsChangedDelegate;
@@ -38,7 +38,7 @@ class UChaosVDVisualizationSettingsObjectBase : public UChaosVDSettingsObjectBas
 {
 	GENERATED_BODY()
 protected:
-	virtual void BroadcastSettingsChanged(UObject* SettingsObject) override;
+	virtual void BroadcastSettingsChanged() override;
 };
 
 UCLASS(config = Engine)
diff --git a/Engine/Plugins/ChaosVD/Source/ChaosVD/Private/Settings/ChaosVDJointConstraintVisualizationSettings.cpp b/Engine/Plugins/ChaosVD/Source/ChaosVD/Private/Settings/ChaosVDJointConstraintVisualizationSettings.cpp
--- a/Engine/Plugins/ChaosVD/Source/ChaosVD/Private/Settings/ChaosVDJointConstraintVisualizationSettings.cpp
+++ b/Engine/Plugins/ChaosVD/Source/ChaosVD/Private/Settings/ChaosVDJointConstraintVisualizationSettings.cpp
@@ -7,7 +7,7 @@ void UChaosVDJointConstraintsVisualizationSettings::SetJointsDataVisualizationFl
 	if (UChaosVDJointConstraintsVisualizationSettings* Settings = GetMutableDefault<UChaosVDJointConstraintsVisualizationSettings>())
 	{
 		Settings->GlobalJointsDataVisualizationFlags = NewFlags;
-		Settings->BroadcastSettingsChanged(Settings);
+		Settings->BroadcastSettingsChanged();
 	}
 }
 
diff --git a/Engine/Plugins/ChaosVD/Source/ChaosVD/Private/Settings/ChaosVDParticleVisualizationSettings.cpp b/Engine/Plugins/ChaosVD/Source/ChaosVD/Private/Settings/ChaosVDParticleVisualizationSettings.cpp
--- a/Engine/Plugins/ChaosVD/Source/ChaosVD/Private/Settings/ChaosVDParticleVisualizationSettings.cpp
+++ b/Engine/Plugins/ChaosVD/Source/ChaosVD/Private/Settings/ChaosVDParticleVisualizationSettings.cpp
@@ -89,7 +89,7 @@ void UChaosVDParticleVisualizationDebugDrawSettings::SetDataDebugDrawVisualizati
 	if (UChaosVDParticleVisualizationDebugDrawSettings* Settings = GetMutableDefault<UChaosVDParticleVisualizationDebugDrawSettings>())
 	{
 		Settings->ParticleDataVisualizationFlags = Flags;
-		Settings->BroadcastSettingsChanged(Settings);
+		Settings->BroadcastSettingsChanged();
 	}
 }
 
@@ -108,7 +108,7 @@ void UChaosVDParticleVisualizationSettings::SetGeometryVisualizationFlags(EChaos
 	if (UChaosVDParticleVisualizationSettings* Settings = GetMutableDefault<UChaosVDParticleVisualizationSettings>())
 	{
 		Settings->GeometryVisibilityFlags = Flags;
-		Settings->BroadcastSettingsChanged(Settings);
+		Settings->BroadcastSettingsChanged();
 	}
 }
 
diff --git a/Engine/Plugins/ChaosVD/Source/ChaosVD/Private/Settings/ChaosVDSceneQueryVisualizationSettings.cpp b/Engine/Plugins/ChaosVD/Source/ChaosVD/Private/Settings/ChaosVDSceneQueryVisualizationSettings.cpp
--- a/Engine/Plugins/ChaosVD/Source/ChaosVD/Private/Settings/ChaosVDSceneQueryVisualizationSettings.cpp
+++ b/Engine/Plugins/ChaosVD/Source/ChaosVD/Private/Settings/ChaosVDSceneQueryVisualizationSettings.cpp
@@ -7,7 +7,7 @@ void UChaosVDSceneQueriesVisualizationSettings::SetSceneQueryDataVisualizationFl
 	if (UChaosVDSceneQueriesVisualizationSettings* Settings = GetMutableDefault<UChaosVDSceneQueriesVisualizationSettings>())
 	{
 		Settings->GlobalSceneQueriesVisualizationFlags = NewFlags;
-		Settings->BroadcastSettingsChanged(Settings);
+		Settings->BroadcastSettingsChanged();
 	}
 }
 
diff --git a/Engine/Plugins/ChaosVD/Source/ChaosVD/Private/Utils/ChaosVDUserInterfaceUtils.cpp b/Engine/Plugins/ChaosVD/Source/ChaosVD/Private/Utils/ChaosVDUserInterfaceUtils.cpp
--- a/Engine/Plugins/ChaosVD/Source/ChaosVD/Private/Utils/ChaosVDUserInterfaceUtils.cpp
+++ b/Engine/Plugins/ChaosVD/Source/ChaosVD/Private/Utils/ChaosVDUserInterfaceUtils.cpp
@@ -57,7 +57,7 @@ void Chaos::VisualDebugger::Utils::CreateMenuEntryForObject(UToolMenu* Menu, UOb
 			.Text(NSLOCTEXT("ChaosVisualDebugger", "CreateMenuEntryForObjectErrorMessage", "Failed to create menu for object. The provied object is null"))
 		];
 		
-		FToolMenuEntry ErrorMenuEntry = FToolMenuEntry::InitWidget(Object->GetFName(), ErrorMessageWidget, FText::GetEmpty());
+		FToolMenuEntry ErrorMenuEntry = FToolMenuEntry::InitWidget(TEXT("InvalidObject"), ErrorMessageWidget, FText::GetEmpty());
 		Menu->AddMenuEntry(NAME_None, ErrorMenuEntry);
 		return;
 	}
