[ChaosVD] Fix CIS Localization warning in SChaosVDViewportToolbar

diff --git a/Engine/Plugins/ChaosVD/Source/ChaosVD/Private/Widgets/SChaosVDViewportToolbar.cpp b/Engine/Plugins/ChaosVD/Source/ChaosVD/Private/Widgets/SChaosVDViewportToolbar.cpp
--- a/Engine/Plugins/ChaosVD/Source/ChaosVD/Private/Widgets/SChaosVDViewportToolbar.cpp
+++ b/Engine/Plugins/ChaosVD/Source/ChaosVD/Private/Widgets/SChaosVDViewportToolbar.cpp
@@ -80,8 +80,8 @@ void SChaosVDViewportToolbar::ExtendOptionsMenu(FMenuBuilder& OptionsMenuBuilder
 		));
 
 		CVDOptionsMenuBuilder.AddSubMenu(
-		LOCTEXT("FrameRateOptionsMenuLabel", "Object Tracking"),
-		LOCTEXT("FrameRateOptionsMenuToolTip", "Options that control how objects are tracked in the scene by the camera."),
+		LOCTEXT("ObjectTrackingMenuLabel", "Object Tracking"),
+		LOCTEXT("ObjectTrackingMenuToolTip", "Options that control how objects are tracked in the scene by the camera."),
 		FNewMenuDelegate::CreateSP(this, &SChaosVDViewportToolbar::PopulateAutoTrackingSubMenu),
 		false,
 		FSlateIcon(FAppStyle::Get().GetStyleSetName(), TEXT("AnimViewportMenu.CameraFollow.Small"))
