[ChaosVD] Playback logic refactor to improve stability and lay the ground work to support different sync modes

- Migrated the tick logic from the timeline widgets to the Playback Controller (Now CVD better respects a MVC implementation)
- The game frame track no longer automatically plays the solver tracks. When a sync is desired the controller will use the singl;e Sync method that will handle the sync independently if the track is a solver track or game track
- Removed all the coordination via events to update the UI. The widgets now get assigned a SharedRef to the Track info structure representing the state of the track, and it renders the state at the moment it is painting+
- The playback control widgets now just issues commands (Play,Stop,Next,Prev,Pause) to the playback controller instead of telling it what frame it needs to do. the controller now will go to the correct frame if possible. The UI will naturally reflect the correct state as it is reading the data directly
- Switched to use inline allocators for the arrays we use to gather the available tracks in different places. We know that these will usually not be greater than 6, but I set the the allocators to 16 elements just to have more headroom
- Fixed an issue where Character Ground Constraints were not being properly updated
- Fixed an issue where when you scrub from a solver timeline, the other solvers and game frames data is not updated
- Added a new button to the solver playback controls that allows you to enable or disable its syncing with other tracks. This can be used to revert to the old behavior of playing solver tracks independently
- Fixed an issue with RBAN Particle Data Visualization.

diff --git a/Engine/Plugins/ChaosVD/Source/ChaosVD/Private/Actors/ChaosVDSolverInfoActor.cpp b/Engine/Plugins/ChaosVD/Source/ChaosVD/Private/Actors/ChaosVDSolverInfoActor.cpp
--- a/Engine/Plugins/ChaosVD/Source/ChaosVD/Private/Actors/ChaosVDSolverInfoActor.cpp
+++ b/Engine/Plugins/ChaosVD/Source/ChaosVD/Private/Actors/ChaosVDSolverInfoActor.cpp
@@ -37,10 +37,10 @@ AChaosVDSolverInfoActor::AChaosVDSolverInfoActor(const FObjectInitializer& Objec
 
 }
 
-void AChaosVDSolverInfoActor::SetSolverName(const FString& InSolverName)
+void AChaosVDSolverInfoActor::SetSolverName(const FName& InSolverName)
 {
 	SolverName = InSolverName;
-	SetActorLabel(TEXT("Solver Data Container | ") + InSolverName);
+	SetActorLabel(TEXT("Solver Data Container | ") + InSolverName.ToString());
 }
 
 void AChaosVDSolverInfoActor::SetScene(TWeakPtr<FChaosVDScene> InScene)
@@ -100,7 +100,7 @@ bool AChaosVDSolverInfoActor::SelectParticleByID(int32 ParticleIDToSelect)
 	AChaosVDParticleActor* ParticleToSelect = GetParticleActor(ParticleIDToSelect);
 	if (!ParticleToSelect)
 	{
-		UE_LOG(LogChaosVDEditor, Error, TEXT("[%s] Particle ID [%d] not found in Solver [%s]"), ANSI_TO_TCHAR(__FUNCTION__), ParticleIDToSelect, *GetSolverName());
+		UE_LOG(LogChaosVDEditor, Error, TEXT("[%s] Particle ID [%d] not found in Solver [%s]"), ANSI_TO_TCHAR(__FUNCTION__), ParticleIDToSelect, *GetSolverName().ToString());
 		return false;
 	}
 
@@ -273,7 +273,7 @@ FName AChaosVDSolverInfoActor::GetFolderPathForParticleType(EChaosVDParticleType
 	}
 	else
 	{
-		const FStringFormatOrderedArguments Args {SolverName, FString::FromInt(SolverID)};
+		const FStringFormatOrderedArguments Args {SolverName.ToString(), FString::FromInt(SolverID)};
 		const FName ParticleFolderPath = *FPaths::Combine(FString::Format(TEXT("Solver {0} | ID {1}"), Args), UEnum::GetDisplayValueAsText(ParticleType).ToString());
 
 		FolderPathByParticlePath.Add(ParticleType, ParticleFolderPath);
diff --git a/Engine/Plugins/ChaosVD/Source/ChaosVD/Private/Actors/ChaosVDSolverInfoActor.h b/Engine/Plugins/ChaosVD/Source/ChaosVD/Private/Actors/ChaosVDSolverInfoActor.h
--- a/Engine/Plugins/ChaosVD/Source/ChaosVD/Private/Actors/ChaosVDSolverInfoActor.h
+++ b/Engine/Plugins/ChaosVD/Source/ChaosVD/Private/Actors/ChaosVDSolverInfoActor.h
@@ -29,8 +29,8 @@ public:
 	void SetSolverID(int32 InSolverID) { SolverID = InSolverID; }
 	int32 GetSolverID() const { return SolverID; }
 
-	void SetSolverName(const FString& InSolverName);
-	const FString& GetSolverName() { return SolverName; }
+	void SetSolverName(const FName& InSolverName);
+	const FName& GetSolverName() { return SolverName; }
 
 	void SetIsServer(bool bInIsServer) { bIsServer = bInIsServer; }
 	bool GetIsServer() const { return bIsServer; }
@@ -89,7 +89,7 @@ protected:
 	FTransform SimulationTransform;
 
 	UPROPERTY(VisibleAnywhere, Category="Solver Data")
-	FString SolverName;
+	FName SolverName;
 
 	UPROPERTY()
 	TObjectPtr<UChaosVDSolverCollisionDataComponent> CollisionDataComponent;
diff --git a/Engine/Plugins/ChaosVD/Source/ChaosVD/Private/ChaosVDPlaybackController.cpp b/Engine/Plugins/ChaosVD/Source/ChaosVD/Private/ChaosVDPlaybackController.cpp
--- a/Engine/Plugins/ChaosVD/Source/ChaosVD/Private/ChaosVDPlaybackController.cpp
+++ b/Engine/Plugins/ChaosVD/Source/ChaosVD/Private/ChaosVDPlaybackController.cpp
@@ -24,13 +24,21 @@ namespace Chaos::VisualDebugger::Cvars
 		TEXT("If true, each time we get playback a solver frame in a specific stage, we will play all the previous steps from that frame in sequence to ensure we have the correct visualization for what happened in that frame."));
 }
 
+bool FChaosVDTrackInfo::operator==(const FChaosVDTrackInfo& Other) const
+{
+	return Other.TrackID == TrackID && Other.TrackType == TrackType;
+}
+
+bool FChaosVDTrackInfo::AreSameTrack(const TSharedRef<const FChaosVDTrackInfo>& TrackA, const TSharedRef<const FChaosVDTrackInfo>& TrackB)
+{
+	return TrackA->TrackID == TrackB->TrackID && TrackA->TrackType == TrackB->TrackType; 
+}
+
 FChaosVDPlaybackController::FChaosVDPlaybackController(const TWeakPtr<FChaosVDScene>& InSceneToControl)
 {
 	SceneToControl = InSceneToControl;
 
 	RecordingStoppedHandle = FChaosVDRuntimeModule::Get().RegisterRecordingStopCallback(FChaosVDRecordingStateChangedDelegate::FDelegate::CreateRaw(this, &FChaosVDPlaybackController::HandleDisconnectedFromSession));
-
-	CurrentPlaybackInstigator = IChaosVDPlaybackControllerInstigator::InvalidGuid;
 }
 
 FChaosVDPlaybackController::~FChaosVDPlaybackController()
@@ -147,17 +155,13 @@ void FChaosVDPlaybackController::PlayFromClosestKeyFrame_AssumesLocked(const int
 	}
 }
 
-void FChaosVDPlaybackController::EnqueueTrackInfoUpdate(const FChaosVDTrackInfo& InTrackInfo, FGuid InstigatorID)
+void FChaosVDPlaybackController::EnqueueTrackInfoUpdate(const TSharedRef<const FChaosVDTrackInfo>& InTrackInfo, FGuid InstigatorID)
 {
 	// This will be used in the Game Thread on the first tick after this was added, so we need to make a copy of the state right now
 	FChaosVDQueuedTrackInfoUpdate InfoUpdate;
 	InfoUpdate.TrackInfo = InTrackInfo;
 	InfoUpdate.InstigatorID = InstigatorID;
 
-	// TODO: FChaosVDTrackInfo has the track name as an string which will be copied as well.
-	// We should move the names to another structure as these don't change often.
-	// It wasn't an issue for now, but now this copy might have an impact due to the number of updates there could be done.
-	
 	TrackInfoUpdateGTQueue.Enqueue(InfoUpdate);
 }
 
@@ -179,20 +183,13 @@ void FChaosVDPlaybackController::PlaySolverStepData(int32 TrackID, const TShared
 	}
 }
 
-void FChaosVDPlaybackController::GoToRecordedSolverStep_AssumesLocked(const int32 InTrackID, const int32 FrameNumber, const int32 Step, FGuid InstigatorID, int32 Attempts)
+void FChaosVDPlaybackController::GoToRecordedSolverStage_AssumesLocked(const int32 InTrackID, const int32 FrameNumber, const int32 StageNumber, FGuid InstigatorID)
 {
 	if (const TSharedPtr<FChaosVDScene> SceneToControlSharedPtr = SceneToControl.Pin())
 	{
 		if (ensure(LoadedRecording.IsValid()))
 		{
-			TSharedPtr<FChaosVDTrackInfo> CurrentTrackInfo;
-			if (TrackInfoByIDMap* TrackInfoByID = TrackInfoPerType.Find(EChaosVDTrackType::Solver))
-			{
-				if (const TSharedPtr<FChaosVDTrackInfo>* TrackInfo = TrackInfoByID->Find(InTrackID))
-				{
-					CurrentTrackInfo = *TrackInfo;
-				}
-			}
+			TSharedPtr<FChaosVDTrackInfo> CurrentTrackInfo = GetMutableTrackInfo(EChaosVDTrackType::Solver, InTrackID);
 
 			if (!ensure(CurrentTrackInfo.IsValid()))
 			{
@@ -202,6 +199,12 @@ void FChaosVDPlaybackController::GoToRecordedSolverStep_AssumesLocked(const int3
 
 			if (FChaosVDSolverFrameData* SolverFrameData = LoadedRecording->GetSolverFrameData_AssumesLocked(InTrackID, FrameNumber))
 			{
+				// All solver frames need to be played with a valid specified solver stage. If we don't have just early out
+				if (!SolverFrameData->SolverSteps.IsValidIndex(StageNumber))
+				{
+					return;
+				}
+
 				const int32 FrameDiff = FrameNumber - CurrentTrackInfo->CurrentFrame;
 				constexpr int32 FrameDriftTolerance = 1;
 
@@ -218,26 +221,34 @@ void FChaosVDPlaybackController::GoToRecordedSolverStep_AssumesLocked(const int3
 					PlayFromClosestKeyFrame_AssumesLocked(InTrackID, FrameNumber, *SceneToControlSharedPtr.Get());
 				}
 
-				const int32 StepDiff = Step - CurrentTrackInfo->CurrentStep;
-				const bool bNeedsPlayPreviousSteps = CurrentTrackInfo->CurrentFrame != FrameNumber || StepDiff < 0 || FMath::Abs(StepDiff) > FrameDriftTolerance;
+				const int32 StageNumberDiff = StageNumber - CurrentTrackInfo->CurrentStage;
+				const bool bNeedsPlayPreviousSteps = CurrentTrackInfo->CurrentFrame != FrameNumber || StageNumberDiff < 0 || FMath::Abs(StageNumberDiff) > FrameDriftTolerance;
 
 				if (Chaos::VisualDebugger::Cvars::bPlayAllPreviousFrameSteps && bNeedsPlayPreviousSteps)
 				{
-					for (int32 StepIndex = 0; StepIndex <= Step; StepIndex++)
+					for (int32 StageIndex = 0; StageIndex <= StageNumber; StageIndex++)
 					{
-						PlaySolverStepData(InTrackID, SceneToControlSharedPtr.ToSharedRef(), *SolverFrameData, StepIndex);
+						PlaySolverStepData(InTrackID, SceneToControlSharedPtr.ToSharedRef(), *SolverFrameData, StageIndex);
 					}
 				}
 				else
 				{
-					PlaySolverStepData(InTrackID, SceneToControlSharedPtr.ToSharedRef(), *SolverFrameData, Step);
+					PlaySolverStepData(InTrackID, SceneToControlSharedPtr.ToSharedRef(), *SolverFrameData, StageNumber);
 				}
 
-				CurrentTrackInfo->CurrentFrame = FrameNumber;
-				CurrentTrackInfo->CurrentStep = Step;
+				if (CurrentTrackInfo->CurrentFrame != FrameNumber)
+				{
+					CurrentTrackInfo->CurrentFrame = FrameNumber;
+					SceneToControlSharedPtr->HandleEnterNewSolverFrame(FrameNumber, *SolverFrameData);
+				}
+
+				CurrentTrackInfo->CurrentStage = StageNumber;
 				CurrentTrackInfo->bIsReSimulated = SolverFrameData->bIsResimulated;
+
+				CurrentTrackInfo->CurrentStageNames.Reset();
+				Algo::Transform(SolverFrameData->SolverSteps, CurrentTrackInfo->CurrentStageNames, &FChaosVDStepData::StepName);
 			
-				EnqueueTrackInfoUpdate(*CurrentTrackInfo.Get(), InstigatorID);
+				EnqueueTrackInfoUpdate(CurrentTrackInfo.ToSharedRef(), InstigatorID);
 			}
 		}
 	}
@@ -247,50 +258,46 @@ void FChaosVDPlaybackController::GoToRecordedSolverStep_AssumesLocked(const int3
 	}
 }
 
-void FChaosVDPlaybackController::GoToRecordedGameFrame_AssumesLocked(const int32 FrameNumber, FGuid InstigatorID, int32 Attempts)
+void FChaosVDPlaybackController::GoToRecordedGameFrame_AssumesLocked(const int32 FrameNumber, FGuid InstigatorID)
 {
 	if (const TSharedPtr<FChaosVDScene> SceneToControlSharedPtr = SceneToControl.Pin())
 	{
 		if (ensure(LoadedRecording.IsValid()))
 		{
-			if (TrackInfoByIDMap* TrackInfoByID = TrackInfoPerType.Find(EChaosVDTrackType::Game))
+			const TSharedPtr<FChaosVDTrackInfo> TrackInfoSharedPtr = GetMutableTrackInfo(EChaosVDTrackType::Game, GameTrackID);
+			if (!ensure(TrackInfoSharedPtr))
 			{
-				const TSharedPtr<FChaosVDTrackInfo>* TrackInfoSharedPtrPtr = TrackInfoByID->Find(GameTrackID);
-				if (!ensure(TrackInfoSharedPtrPtr))
-				{
-					return;
-				}
-				
-				if (const TSharedPtr<FChaosVDTrackInfo> TrackInfoSharedPtr = *TrackInfoSharedPtrPtr)
-				{
-					if (const FChaosVDGameFrameData* FoundGameFrameData = LoadedRecording->GetGameFrameData_AssumesLocked(FrameNumber))
-					{
-						TArray<int32> AvailableSolversID;
-						LoadedRecording->GetAvailableSolverIDsAtGameFrameNumber_AssumesLocked(FrameNumber, AvailableSolversID);
+				return;
+			}
 
-						SceneToControlSharedPtr->HandleEnterNewGameFrame(FrameNumber, AvailableSolversID, *FoundGameFrameData);
+			if (const FChaosVDGameFrameData* FoundGameFrameData = LoadedRecording->GetGameFrameData_AssumesLocked(FrameNumber))
+			{
+				TArray<int32, TInlineAllocator<FChaosVDRecording::CommonTrackCount>> AvailableSolversID;
+				TArray<int32, TInlineAllocator<FChaosVDRecording::CommonTrackCount>> RemovedSolversID;
+				
+				AvailableSolversID.Reset();
+				RemovedSolversID.Reset();
+				LoadedRecording->GetAvailableSolverIDsAtGameFrameNumber_AssumesLocked(FrameNumber, AvailableSolversID);
 
-						for (const int32 SolverID : AvailableSolversID)
-						{
-							// When Scrubbing the timeline by game frames instead of solvers, try to go to the first solver frame on the first platform cycle of the game frame.
-							// Game Frames are not in sync with Solver Frames and Solver steps.
-							const int32 SolverFrameNumber = LoadedRecording->GetLowestSolverFrameNumberAtCycle_AssumesLocked(SolverID, FoundGameFrameData->FirstCycle);
-	
-							const int32 StepNumber = GetTrackLastStepAtFrame(EChaosVDTrackType::Solver, SolverID, SolverFrameNumber);
+				SceneToControlSharedPtr->HandleEnterNewGameFrame(FrameNumber, AvailableSolversID, *FoundGameFrameData, RemovedSolversID);
 
-							GoToTrackFrame_AssumesLocked(InstigatorID, EChaosVDTrackType::Solver, SolverID, SolverFrameNumber, StepNumber);
-						}
+				// We need to reset the current frame track info for any solver that is removed as so when (or if) it is loaded again, starts on a clean state
+				for (const int32 SolverID : RemovedSolversID)
+				{
+					if (TSharedPtr<FChaosVDTrackInfo> RemovedSolverTrackPtr = GetMutableTrackInfo(EChaosVDTrackType::Solver, SolverID))
+					{
+						RemovedSolverTrackPtr->CurrentFrame = 0;
 					}
-					
-					TrackInfoSharedPtr->CurrentFrame = FrameNumber;
-					EnqueueTrackInfoUpdate(*TrackInfoSharedPtr.Get(), InstigatorID);
 				}
+
+				TrackInfoSharedPtr->CurrentFrame = FrameNumber;
+				EnqueueTrackInfoUpdate(TrackInfoSharedPtr.ToSharedRef(), InstigatorID);
 			}
 		}
 	}
 }
 
-void FChaosVDPlaybackController::GoToTrackFrame(FGuid InstigatorID, EChaosVDTrackType TrackType, int32 InTrackID, int32 FrameNumber, int32 Step)
+void FChaosVDPlaybackController::GoToTrackFrame(FGuid InstigatorID, EChaosVDTrackType TrackType, int32 InTrackID, int32 FrameNumber, int32 StageNumber)
 {
 	if (!ensure(LoadedRecording.IsValid()))
 	{
@@ -298,10 +305,10 @@ void FChaosVDPlaybackController::GoToTrackFrame(FGuid InstigatorID, EChaosVDTrac
 	}
 
 	FReadScopeLock ReadLock(LoadedRecording->GetRecordingDataLock());
-	GoToTrackFrame_AssumesLocked(InstigatorID, TrackType, InTrackID, FrameNumber, Step);
+	GoToTrackFrame_AssumesLocked(InstigatorID, TrackType, InTrackID, FrameNumber, StageNumber);
 }
 
-void FChaosVDPlaybackController::GoToTrackFrame_AssumesLocked(FGuid InstigatorID, EChaosVDTrackType TrackType, int32 InTrackID, int32 FrameNumber, int32 Step)
+void FChaosVDPlaybackController::GoToTrackFrame_AssumesLocked(FGuid InstigatorID, EChaosVDTrackType TrackType, int32 InTrackID, int32 FrameNumber, int32 StageNumber)
 {
 	switch (TrackType)
 	{
@@ -309,7 +316,7 @@ void FChaosVDPlaybackController::GoToTrackFrame_AssumesLocked(FGuid InstigatorID
 		GoToRecordedGameFrame_AssumesLocked(FrameNumber, InstigatorID);
 		break;
 	case EChaosVDTrackType::Solver:
-		GoToRecordedSolverStep_AssumesLocked(InTrackID, FrameNumber, Step, InstigatorID);
+		GoToRecordedSolverStage_AssumesLocked(InTrackID, FrameNumber, StageNumber, InstigatorID);
 		break;
 	default:
 		ensure(false);
@@ -317,6 +324,31 @@ void FChaosVDPlaybackController::GoToTrackFrame_AssumesLocked(FGuid InstigatorID
 	}
 }
 
+
+void FChaosVDPlaybackController::GoToTrackFrameAndSync(FGuid InstigatorID, EChaosVDTrackType TrackType, int32 InTrackID, int32 FrameNumber, int32 StageNumber)
+{
+	if (!ensure(LoadedRecording.IsValid()))
+	{
+		return;
+	}
+
+	FReadScopeLock ReadLock(LoadedRecording->GetRecordingDataLock());
+	
+	GoToTrackFrame_AssumesLockedAndSync(InstigatorID, TrackType, InTrackID, FrameNumber, StageNumber);
+}
+
+
+void FChaosVDPlaybackController::GoToTrackFrame_AssumesLockedAndSync(FGuid InstigatorID, EChaosVDTrackType TrackType, int32 InTrackID, int32 FrameNumber, int32 StageNumber)
+{
+	GoToTrackFrame_AssumesLocked(InstigatorID, TrackType, InTrackID, FrameNumber, StageNumber);
+	
+	if (TSharedPtr<const FChaosVDTrackInfo> TrackToSyncWith = GetTrackInfo(TrackType, InTrackID))
+	{
+		SyncTracks_AssumesLocked(TrackToSyncWith.ToSharedRef());
+	}
+}
+
+
 int32 FChaosVDPlaybackController::GetTrackStepsNumberAtFrame_AssumesLocked(EChaosVDTrackType TrackType, const int32 InTrackID, const int32 FrameNumber) const
 {
 	if (!LoadedRecording.IsValid())
@@ -360,7 +392,6 @@ const FChaosVDStepsContainer* FChaosVDPlaybackController::GetTrackStepsDataAtFra
 	case EChaosVDTrackType::Game:
 		// Game Tracks do not have steps
 		return nullptr;
-		break;
 	case EChaosVDTrackType::Solver:
 		{
 			if (const FChaosVDSolverFrameData* FrameData = LoadedRecording->GetSolverFrameData_AssumesLocked(InTrackID, FrameNumber))
@@ -408,54 +439,42 @@ int32 FChaosVDPlaybackController::GetTrackFramesNumber(EChaosVDTrackType TrackTy
 	}
 }
 
-int32 FChaosVDPlaybackController::ConvertCurrentFrameToOtherTrackFrame(const FChaosVDTrackInfo* FromTrack, const FChaosVDTrackInfo* ToTrack)
-{
+int32 FChaosVDPlaybackController::ConvertCurrentFrameToOtherTrackFrame_AssumesLocked(const TSharedRef<const FChaosVDTrackInfo>& InFromTrack, const TSharedRef<const FChaosVDTrackInfo>& InToTrack, EChaosVDSyncTimelinesMode TrackSyncMode)
+{	
+	if (!ensure(LoadedRecording.IsValid()))
+	{
+		return INDEX_NONE;
+	}
+
 	// Each track is on a different "space time", because it's source data ticked at different rates when was recorded, and some start/end at different points on time
 	// But all the recorded frame data on all of them use Platform Cycles as timestamps
 	// This method wraps specialized methods in the recording object to convert between these spaces. For example Game frame 1500 could be frame 5 on a specific solver
 	// And Frame 5 of that solver could be frame 30 on another solver.
 
-	if (FromTrack == nullptr || ToTrack == nullptr)
+	if (FChaosVDTrackInfo::AreSameTrack(InFromTrack, InToTrack))
 	{
-		ensureMsgf(false, TEXT("One of provided track infos is not valid"));
-		return INDEX_NONE;
-	}
-	
-	if (!ensure(LoadedRecording.IsValid()))
-    {
-    	return INDEX_NONE;
-    }
-
-	const bool bBothTracksHaveSameID = FromTrack->TrackID == ToTrack->TrackID;
-	const bool bBothTracksHaveSameType = FromTrack->TrackType == ToTrack->TrackType;
-	if (bBothTracksHaveSameType && bBothTracksHaveSameID)
-	{
-		return FromTrack->CurrentFrame;
+		return InFromTrack->CurrentFrame;
 	}
 
-	switch (FromTrack->TrackType)
+	switch (InFromTrack->TrackType)
 	{
 	case EChaosVDTrackType::Game:
 		{
 			// Convert from Game Frame to Solver Frame
-			return LoadedRecording->GetLowestSolverFrameNumberGameFrame(FromTrack->TrackID, FromTrack->CurrentFrame);
-			break;
-		}
-		
+			return LoadedRecording->GetLowestSolverFrameNumberGameFrame_AssumesLocked(InFromTrack->TrackID, InFromTrack->CurrentFrame);
+		}	
 	case EChaosVDTrackType::Solver:
 		{
-			if (ToTrack->TrackType == EChaosVDTrackType::Solver)
+			if (InToTrack->TrackType == EChaosVDTrackType::Solver)
 			{
-				ensureMsgf(false, TEXT("Frame conversion between solver tracks is not supported yet"));
-				return INDEX_NONE;
+				const FChaosVDSolverFrameData* FromSolverFrameData = LoadedRecording->GetSolverFrameData_AssumesLocked(InFromTrack->TrackID, InFromTrack->CurrentFrame);
+				return ensure(FromSolverFrameData) ? LoadedRecording->GetLowestSolverFrameNumberAtCycle_AssumesLocked(InToTrack->TrackID, FromSolverFrameData->FrameCycle) : INDEX_NONE;
 			}
-			return LoadedRecording->GetLowestGameFrameAtSolverFrameNumber(FromTrack->TrackID, FromTrack->CurrentFrame);
-			break;
+			return LoadedRecording->GetLowestGameFrameAtSolverFrameNumber_AssumesLocked(InFromTrack->TrackID, InFromTrack->CurrentFrame);
 		}
 	default:
 		ensure(false);
 		return INDEX_NONE;
-		break;
 	}
 }
 
@@ -480,14 +499,14 @@ int32 FChaosVDPlaybackController::GetTrackCurrentStep(EChaosVDTrackType TrackTyp
 		const TSharedPtr<FChaosVDTrackInfo>* TrackInfoPtrPtr = TrackInfoByID->Find(InTrackID);
 		if (const TSharedPtr<FChaosVDTrackInfo> TrackInfoSharedPtr = TrackInfoPtrPtr ? *TrackInfoPtrPtr : nullptr)
 		{
-			return TrackInfoSharedPtr->CurrentStep;
+			return TrackInfoSharedPtr->CurrentStage;
 		}
 	}
 
 	return INDEX_NONE;
 }
 
-int32 FChaosVDPlaybackController::GetTrackLastStepAtFrame(EChaosVDTrackType TrackType, int32 InTrackID, int32 InFrameNumber) const
+int32 FChaosVDPlaybackController::GetTrackLastStageAtFrame_AssumesLocked(EChaosVDTrackType TrackType, int32 InTrackID, int32 InFrameNumber) const
 {
 	switch (TrackType)
 	{
@@ -503,22 +522,21 @@ int32 FChaosVDPlaybackController::GetTrackLastStepAtFrame(EChaosVDTrackType Trac
 			return INDEX_NONE;
 			break;
 	}
-	
 }
 
-const FChaosVDTrackInfo* FChaosVDPlaybackController::GetTrackInfo(EChaosVDTrackType TrackType, int32 TrackID)
+TSharedPtr<const FChaosVDTrackInfo> FChaosVDPlaybackController::GetTrackInfo(EChaosVDTrackType TrackType, int32 TrackID)
 {
 	return GetMutableTrackInfo(TrackType, TrackID);
 }
 
-FChaosVDTrackInfo* FChaosVDPlaybackController::GetMutableTrackInfo(EChaosVDTrackType TrackType, int32 TrackID)
+TSharedPtr<FChaosVDTrackInfo> FChaosVDPlaybackController::GetMutableTrackInfo(EChaosVDTrackType TrackType, int32 TrackID)
 {
 	if (const TrackInfoByIDMap* TrackInfoByID = TrackInfoPerType.Find(TrackType))
 	{
 		const TSharedPtr<FChaosVDTrackInfo>* TrackInfoPtrPtr = TrackInfoByID->Find(TrackID);
 		if (const TSharedPtr<FChaosVDTrackInfo> TrackInfoSharedPtr = TrackInfoPtrPtr ? *TrackInfoPtrPtr : nullptr)
 		{
-			return TrackInfoSharedPtr.Get();
+			return TrackInfoSharedPtr;
 		}
 	}
 
@@ -527,15 +545,15 @@ FChaosVDTrackInfo* FChaosVDPlaybackController::GetMutableTrackInfo(EChaosVDTrack
 
 void FChaosVDPlaybackController::LockTrackInCurrentStep(EChaosVDTrackType TrackType, int32 TrackID)
 {
-	if (FChaosVDTrackInfo* TrackInfo = GetMutableTrackInfo(TrackType, TrackID))
+	if (TSharedPtr<FChaosVDTrackInfo> TrackInfo = GetMutableTrackInfo(TrackType, TrackID))
 	{
-		TrackInfo->LockedOnStep = TrackInfo->CurrentStep;
+		TrackInfo->LockedOnStep = TrackInfo->CurrentStage;
 	}
 }
 
 void FChaosVDPlaybackController::UnlockTrackStep(EChaosVDTrackType TrackType, int32 TrackID)
 {
-	if (FChaosVDTrackInfo* TrackInfo = GetMutableTrackInfo(TrackType, TrackID))
+	if (TSharedPtr<FChaosVDTrackInfo> TrackInfo = GetMutableTrackInfo(TrackType, TrackID))
 	{
 		TrackInfo->LockedOnStep = INDEX_NONE;
 	}
@@ -547,7 +565,13 @@ void FChaosVDPlaybackController::GetAvailableTracks(EChaosVDTrackType TrackType,
 	TrackInfoPerType.FindOrAdd(TrackType).GenerateValueArray(OutTrackInfo);
 }
 
-void FChaosVDPlaybackController::GetAvailableTrackInfosAtTrackFrame(EChaosVDTrackType TrackTypeToFind, TArray<TSharedPtr<FChaosVDTrackInfo>>& OutTrackInfo, const FChaosVDTrackInfo* TrackFrameInfo)
+void FChaosVDPlaybackController::GetAvailableTrackInfosAtTrackFrame(EChaosVDTrackType TrackTypeToFind, const TSharedRef<const FChaosVDTrackInfo>& InFromTrack, TArray<TSharedPtr<const FChaosVDTrackInfo>>& OutTrackInfo)
+{
+	FReadScopeLock ReadLock(LoadedRecording->GetRecordingDataLock());
+	return GetAvailableTrackInfosAtTrackFrame_AssumesLocked(TrackTypeToFind, InFromTrack, OutTrackInfo);
+}
+
+void FChaosVDPlaybackController::GetAvailableTrackInfosAtTrackFrame_AssumesLocked(EChaosVDTrackType TrackTypeToFind, const TSharedRef<const FChaosVDTrackInfo>& InFromTrack, TArray<TSharedPtr<const FChaosVDTrackInfo>>& OutTrackInfo)
 {
 	OutTrackInfo.Reset();
 
@@ -556,46 +580,58 @@ void FChaosVDPlaybackController::GetAvailableTrackInfosAtTrackFrame(EChaosVDTrac
 		return;
 	}
 
-	if (TrackFrameInfo == nullptr)
+	int32 CorrectedFrameNumber = INDEX_NONE;
+	switch (InFromTrack->TrackType)
 	{
-		return;
+	case EChaosVDTrackType::Game:
+		{
+			CorrectedFrameNumber = InFromTrack->CurrentFrame;
+		}
+		break;
+	case EChaosVDTrackType::Solver:
+		{
+			CorrectedFrameNumber = LoadedRecording->GetLowestGameFrameAtSolverFrameNumber_AssumesLocked(InFromTrack->TrackID, InFromTrack->CurrentFrame);
+			break;
+		}
+	default:
+		ensure(false);
+		break;
 	}
 
-	int32 CorrectedFrameNumber = INDEX_NONE;
-	switch (TrackFrameInfo->TrackType)
+	switch (TrackTypeToFind)
 	{
-		case EChaosVDTrackType::Game:
-			{
-				CorrectedFrameNumber = TrackFrameInfo->CurrentFrame;
-			}
+	case EChaosVDTrackType::Game:
+		{
+			OutTrackInfo.Add(GetTrackInfo(EChaosVDTrackType::Game, GameTrackID));
 			break;
-		case EChaosVDTrackType::Solver:
+		}
+	case EChaosVDTrackType::Solver:
+		{
+			TArray<int32, TInlineAllocator<FChaosVDRecording::CommonTrackCount>> AvailableSolversID;
+			LoadedRecording->GetAvailableSolverIDsAtGameFrameNumber_AssumesLocked(CorrectedFrameNumber, AvailableSolversID);
+	
+			TrackInfoByIDMap& TrackInfoMap = TrackInfoPerType.FindOrAdd(TrackTypeToFind);
+			for (const int32 SolverID : AvailableSolversID)
 			{
-				CorrectedFrameNumber = LoadedRecording->GetLowestGameFrameAtSolverFrameNumber(TrackFrameInfo->TrackID, TrackFrameInfo->CurrentFrame);
-				break;
+				// The recording might have the solver data available added because it was added the trace analysis thread, 
+				// but the playback controller didn't process it in the game thread yet
+				if (const TSharedPtr<FChaosVDTrackInfo>* SolverTrackInfo = TrackInfoMap.Find(SolverID))
+				{
+					OutTrackInfo.Add(*SolverTrackInfo);
+				}	
 			}
-		default:
-			ensure(false);
 			break;
-	}
-
-	TArray<int32> AvailableSolversID;
-	LoadedRecording->GetAvailableSolverIDsAtGameFrameNumber(CorrectedFrameNumber, AvailableSolversID);
-	
-	TrackInfoByIDMap& TrackInfoMap = TrackInfoPerType.FindOrAdd(TrackTypeToFind);
-	for (const int32 SolverID : AvailableSolversID)
-	{
-		// The recording might have the solver data available added because it was added the trace analysis thread, 
-		// but the playback controller didn't process it in the game thread yet
-		if (const TSharedPtr<FChaosVDTrackInfo>* SolverTrackInfo = TrackInfoMap.Find(SolverID))
-		{
-			OutTrackInfo.Add(*SolverTrackInfo);
-		}	
+		}
+	default:
+		ensure(false);
+		break;
 	}
 }
 
 bool FChaosVDPlaybackController::Tick(float DeltaTime)
-{	
+{
+	TickPlayback(DeltaTime);
+
 	if (!GeometryDataUpdateGTQueue.IsEmpty())
 	{
 		FChaosVDGeometryDataUpdate GeometryDataUpdate;
@@ -614,7 +650,6 @@ bool FChaosVDPlaybackController::Tick(float DeltaTime)
 		return true;
 	}
 
-
 	const bool bIsRecordingLoaded = LoadedRecording.IsValid();
 
 	if (bIsRecordingLoaded)
@@ -639,7 +674,7 @@ bool FChaosVDPlaybackController::Tick(float DeltaTime)
 		FChaosVDQueuedTrackInfoUpdate TrackInfoUpdate;
 		while (TrackInfoUpdateGTQueue.Dequeue(TrackInfoUpdate))
 		{
-			OnTrackFrameUpdated().Broadcast(ThisWeakPtr, &TrackInfoUpdate.TrackInfo, TrackInfoUpdate.InstigatorID);
+			OnTrackFrameUpdated().Broadcast(ThisWeakPtr, TrackInfoUpdate.TrackInfo, TrackInfoUpdate.InstigatorID);
 		}
 	}
 
@@ -648,11 +683,20 @@ bool FChaosVDPlaybackController::Tick(float DeltaTime)
 		// Load at least the first frame
 		if (!bPlayedFirstFrame)
 		{
-			if (LoadedRecording->GetAvailableSolversNumber_AssumesLocked() > 0)
+			if (IsPlayingLiveSession())
+			{
+				if (const TSharedPtr<const FChaosVDTrackInfo> GameTrackInfo = GetTrackInfo(EChaosVDTrackType::Game, GameTrackID))
+				{
+					HandleFramePlaybackControlInput(EChaosVDPlaybackButtonsID::Play, GameTrackInfo.ToSharedRef(), PlaybackSelfInstigatorID);
+					bPlayedFirstFrame = true;
+				}
+			}
+			else if (LoadedRecording->GetAvailableGameFramesNumber() > 0 && LoadedRecording->GetAvailableSolversNumber_AssumesLocked() > 0)
 			{
 				constexpr int32 GameFrameToLoad = 0;
-				constexpr int32 Step = 0;
-				GoToTrackFrame(IChaosVDPlaybackControllerInstigator::InvalidGuid, EChaosVDTrackType::Game, GameTrackID, GameFrameToLoad, Step);
+				// Game frames do not have stages
+				constexpr int32 SolverStageToLoad = 0;
+				GoToTrackFrameAndSync(PlaybackSelfInstigatorID, EChaosVDTrackType::Game, GameTrackID, GameFrameToLoad, SolverStageToLoad);
 				bPlayedFirstFrame = true;
 			}
 		}
@@ -660,7 +704,7 @@ bool FChaosVDPlaybackController::Tick(float DeltaTime)
 		// If we are live, make sure we don't lag too much behind
 		if (!bPauseRequested && IsPlayingLiveSession())
 		{
-			if (const FChaosVDTrackInfo* GameTrackInfo = GetTrackInfo(EChaosVDTrackType::Game, GameTrackID))
+			if (const TSharedPtr<const FChaosVDTrackInfo> GameTrackInfo = GetTrackInfo(EChaosVDTrackType::Game, GameTrackID))
 			{
 				const int32 CurrentFrameDeltaFromLast = FMath::Abs(GameTrackInfo->MaxFrames - GameTrackInfo->CurrentFrame);
 				if (CurrentFrameDeltaFromLast > MaxFramesLaggingBehindDuringLiveSession)
@@ -669,7 +713,7 @@ bool FChaosVDPlaybackController::Tick(float DeltaTime)
 					// and we don't want to go to close to the threshold.
 					const int32 GameFrameToLoad = LoadedRecording->GetAvailableGameFramesNumber() - MinFramesLaggingBehindDuringLiveSession;
 					constexpr int32 Step = 0;
-					GoToTrackFrame(IChaosVDPlaybackControllerInstigator::InvalidGuid, EChaosVDTrackType::Game, GameTrackID, GameFrameToLoad, Step);
+					GoToTrackFrameAndSync(PlaybackSelfInstigatorID, EChaosVDTrackType::Game, GameTrackID, GameFrameToLoad, Step);
 				}
 			}
 		}
@@ -694,37 +738,27 @@ void FChaosVDPlaybackController::HandleDisconnectedFromSession()
 	bHasPendingGTUpdateBroadcast = true;
 }
 
-void FChaosVDPlaybackController::RequestStop(const IChaosVDPlaybackControllerInstigator& InPlaybackInstigator)
+void FChaosVDPlaybackController::StopPlayback(const FGuid& InstigatorGUID)
 {
-	constexpr int32 FrameNumber = 0;
-	constexpr int32 StepNumber = 0;
-	GoToTrackFrame(InPlaybackInstigator.GetInstigatorID(), EChaosVDTrackType::Game, GameTrackID, FrameNumber, StepNumber);
-}
-
-bool FChaosVDPlaybackController::AcquireExclusivePlaybackControls(const IChaosVDPlaybackControllerInstigator& InPlaybackInstigator)
-{
-	const FGuid& InstigatorID = InPlaybackInstigator.GetInstigatorID();
-	if (InstigatorID == CurrentPlaybackInstigator || CurrentPlaybackInstigator == IChaosVDPlaybackControllerInstigator::InvalidGuid)
+	if (TSharedPtr<const FChaosVDTrackInfo> GameFramesTrack = GetTrackInfo(EChaosVDTrackType::Game, GameTrackID))
 	{
-		CurrentPlaybackInstigator = InPlaybackInstigator.GetInstigatorID();
-		bHasPendingGTUpdateBroadcast = true;
-		return true;
+		constexpr int32 FrameNumber = 0;
+		constexpr int32 StepNumber = 0;
+	
+		GoToTrackFrameAndSync(InstigatorGUID, GameFramesTrack->TrackType, GameFramesTrack->TrackID, FrameNumber, StepNumber);
 	}
+	
+	CurrentPlayingTrack = nullptr;
 
-	return false;
-}
-
-bool FChaosVDPlaybackController::ReleaseExclusivePlaybackControls(const IChaosVDPlaybackControllerInstigator& InPlaybackInstigator)
-{
-	const FGuid& InstigatorID = InPlaybackInstigator.GetInstigatorID();
-	if (InstigatorID == CurrentPlaybackInstigator || CurrentPlaybackInstigator == IChaosVDPlaybackControllerInstigator::InvalidGuid)
+	VisitAvailableTracks([this](const TSharedPtr<FChaosVDTrackInfo>& TrackInfo)
 	{
-		CurrentPlaybackInstigator = IChaosVDPlaybackControllerInstigator::InvalidGuid;
-		bHasPendingGTUpdateBroadcast = true;
-		return true;
-	}
+		if (TrackInfo)
+		{
+			TrackInfo->bIsPlaying = false;
+		}
 
-	return false;
+		return true;
+	});
 }
 
 float FChaosVDPlaybackController::GetFrameTimeOverride() const
@@ -745,10 +779,10 @@ void FChaosVDPlaybackController::SetFrameRateOverride(float NewFrameRateOverride
 	CurrentFrameRateOverride = (NewFrameRateOverride >= MinimumFrameRateOverride) ? NewFrameRateOverride : InvalidFrameRateOverride;
 }
 
-float FChaosVDPlaybackController::GetFrameTimeForTrack(EChaosVDTrackType TrackType, int32 TrackID, const FChaosVDTrackInfo& TrackInfo) const
+float FChaosVDPlaybackController::GetFrameTimeForTrack(EChaosVDTrackType TrackType, int32 TrackID, const TSharedRef<const FChaosVDTrackInfo>& InTrackInfo) const
 {
 	const float TargetFrameTimeOverride = GetFrameTimeOverride();
-	const bool bHasValidFrameRateOverride = bUseFrameRateOverride && !FMath::IsNearlyEqual(TargetFrameTimeOverride, FChaosVDPlaybackController::InvalidFrameRateOverride);
+	const bool bHasValidFrameRateOverride = bUseFrameRateOverride && !FMath::IsNearlyEqual(TargetFrameTimeOverride, InvalidFrameRateOverride);
 	if (bHasValidFrameRateOverride)
 	{
 		return TargetFrameTimeOverride;
@@ -761,7 +795,7 @@ float FChaosVDPlaybackController::GetFrameTimeForTrack(EChaosVDTrackType TrackTy
 		{
 			case EChaosVDTrackType::Solver:
 				{
-					if (const FChaosVDSolverFrameData* FrameData = LoadedRecording->GetSolverFrameData_AssumesLocked(TrackID, TrackInfo.CurrentFrame))
+					if (const FChaosVDSolverFrameData* FrameData = LoadedRecording->GetSolverFrameData_AssumesLocked(TrackID, InTrackInfo->CurrentFrame))
 					{
 						CurrentTargetFrameTime = FrameData->GetFrameTime();
 					}
@@ -769,7 +803,7 @@ float FChaosVDPlaybackController::GetFrameTimeForTrack(EChaosVDTrackType TrackTy
 				}
 			case EChaosVDTrackType::Game:
 				{
-					if (const FChaosVDGameFrameData* FrameData = LoadedRecording->GetGameFrameData_AssumesLocked(TrackInfo.CurrentFrame))
+					if (const FChaosVDGameFrameData* FrameData = LoadedRecording->GetGameFrameData_AssumesLocked(InTrackInfo->CurrentFrame))
 					{
 						CurrentTargetFrameTime = FrameData->GetFrameTime();
 					}
@@ -806,6 +840,228 @@ void FChaosVDPlaybackController::UpdateTrackVisibility(EChaosVDTrackType Type, i
 	}
 }
 
+void FChaosVDPlaybackController::HandleFramePlaybackControlInput(EChaosVDPlaybackButtonsID ButtonID, const TSharedRef<const FChaosVDTrackInfo>& InTrackInfoRef, FGuid Instigator)
+{
+	switch (ButtonID)
+	{
+	case EChaosVDPlaybackButtonsID::Play:
+		{
+			if (!ensure(!CurrentPlayingTrack))
+			{
+				UE_LOG(LogChaosVDEditor, Warning, TEXT("[%s] Attempted to play when thera was another track playing. This should not happen, this is probably a widget with controls not properly disabled."), ANSI_TO_TCHAR(__FUNCTION__));
+				CurrentPlayingTrack->bIsPlaying = false;
+			}
+
+			bPauseRequested = false;
+
+			// We receive a const ref from the UI as the track info referenced there is read only, but once we are in the controller again we can and want to modify it.
+			CurrentPlayingTrack = ConstCastSharedRef<FChaosVDTrackInfo>(InTrackInfoRef);
+			CurrentPlayingTrack->bIsPlaying = true;
+			break;
+		}
+	case EChaosVDPlaybackButtonsID::Pause:
+		{
+			if (ensure(CurrentPlayingTrack.IsValid()))
+			{
+				CurrentPlayingTrack->bIsPlaying = false;
+			}
+			else
+			{
+				UE_LOG(LogChaosVDEditor, Warning, TEXT("[%s] Attempted to pause but there was not a track currently playing. This should not happen."), ANSI_TO_TCHAR(__FUNCTION__));
+			}
+
+			bPauseRequested = true;
+
+			CurrentPlayingTrack = nullptr;
+
+			break;
+		}
+	case EChaosVDPlaybackButtonsID::Stop:
+		{
+			StopPlayback(Instigator);
+			break;
+		}
+	case EChaosVDPlaybackButtonsID::Next:
+		{
+			const int32 LastSolverStage = InTrackInfoRef->CurrentStageNames.Num() - 1;
+			const int32 NextFrame = InTrackInfoRef->CurrentFrame +1;
+			GoToTrackFrameAndSync(Instigator, InTrackInfoRef->TrackType, InTrackInfoRef->TrackID, NextFrame, LastSolverStage);
+		}
+		break;
+	case EChaosVDPlaybackButtonsID::Prev:
+		{
+			const int32 LastSolverStage = InTrackInfoRef->CurrentStageNames.Num() - 1;
+			const int32 PrevFrame = InTrackInfoRef->CurrentFrame -1;
+			GoToTrackFrameAndSync(Instigator, InTrackInfoRef->TrackType, InTrackInfoRef->TrackID, PrevFrame, LastSolverStage);
+		}
+		break;
+	default:
+		break;
+	}
+}
+
+void FChaosVDPlaybackController::HandleFrameStagePlaybackControlInput(EChaosVDPlaybackButtonsID ButtonID, const TSharedRef<const FChaosVDTrackInfo>& InTrackInfoRef, FGuid Instigator)
+{
+	switch (ButtonID)
+	{
+	case EChaosVDPlaybackButtonsID::Next:
+		{
+			int32 NextSolverStage = InTrackInfoRef->CurrentStage +1;
+			GoToTrackFrame(Instigator, EChaosVDTrackType::Solver, InTrackInfoRef->TrackID, InTrackInfoRef->CurrentFrame, NextSolverStage);
+			break;
+		}
+	case EChaosVDPlaybackButtonsID::Prev:
+		{
+			int32 PrevSolverStage = InTrackInfoRef->CurrentStage -1;
+			GoToTrackFrame(Instigator, EChaosVDTrackType::Solver, InTrackInfoRef->TrackID, InTrackInfoRef->CurrentFrame, PrevSolverStage);
+		}
+		break;
+	case EChaosVDPlaybackButtonsID::Play:
+	case EChaosVDPlaybackButtonsID::Pause:
+	case EChaosVDPlaybackButtonsID::Stop:
+	default:
+		{
+			ensure(false);
+			UE_LOG(LogChaosVDEditor, Warning, TEXT("[%s] Unsupported Input type | [%s]"), ANSI_TO_TCHAR(__FUNCTION__), *(UEnum::GetDisplayValueAsText(ButtonID).ToString()));
+			break;
+		}
+	}
+}
+
+void FChaosVDPlaybackController::TickPlayback(float DeltaTime)
+{
+	if (!CurrentPlayingTrack)
+	{
+		return;
+	}
+
+	if (!CurrentPlayingTrack->bIsPlaying)
+	{
+		return;
+	}
+
+	if (!IsPlayingLiveSession() && CurrentPlayingTrack->CurrentFrame == CurrentPlayingTrack->MaxFrames - 1)
+	{
+		StopPlayback(PlaybackSelfInstigatorID);
+		return;
+	}
+
+	CurrentPlaybackTime += DeltaTime;
+
+	const float CurrentFrameTime = GetFrameTimeForTrack(CurrentPlayingTrack->TrackType, CurrentPlayingTrack->TrackID, CurrentPlayingTrack.ToSharedRef());
+
+	const bool bIsGameTrack = CurrentPlayingTrack->TrackType == EChaosVDTrackType::Game;
+
+	while (CurrentPlaybackTime > CurrentFrameTime)
+	{
+		CurrentPlaybackTime -= CurrentFrameTime;
+		const int32 NextFrame = CurrentPlayingTrack->CurrentFrame + 1;
+		const int32 NextStageNumber =  bIsGameTrack ? 0 : GetTrackLastStageAtFrame_AssumesLocked(CurrentPlayingTrack->TrackType, CurrentPlayingTrack->TrackID, NextFrame);
+		GoToTrackFrameAndSync(PlaybackSelfInstigatorID, CurrentPlayingTrack->TrackType, CurrentPlayingTrack->TrackID, NextFrame, NextStageNumber);
+	}
+}
+
+void FChaosVDPlaybackController::GetTracksByType(EChaosVDTrackType Type, TArray<TSharedPtr<FChaosVDTrackInfo>>& OutTracks)
+{
+	OutTracks.Reset();
+	if (TrackInfoByIDMap* TracksById = TrackInfoPerType.Find(Type))
+	{
+		TracksById->GenerateValueArray(OutTracks);
+	}
+}
+
+void FChaosVDPlaybackController::SyncTracks(const TSharedRef<const FChaosVDTrackInfo>& FromTrack, EChaosVDSyncTimelinesMode TrackSyncMode)
+{
+	if (!ensure(LoadedRecording.IsValid()))
+	{
+		return;
+	}
+
+	FReadScopeLock ReadLock(LoadedRecording->GetRecordingDataLock());
+	
+	SyncTracks_AssumesLocked(FromTrack, TrackSyncMode);
+}
+
+void FChaosVDPlaybackController::SyncTracks_AssumesLocked(const TSharedRef<const FChaosVDTrackInfo>& FromTrack, EChaosVDSyncTimelinesMode TrackSyncMode)
+{
+	if (!ensure(TrackSyncMode == EChaosVDSyncTimelinesMode::RecordedTimestamp))
+	{
+		UE_LOG(LogChaosVDEditor, Error, TEXT("[%s] Attempted to Sync tracks using an unsuported mode. |[%s] is not supported yet."), ANSI_TO_TCHAR(__FUNCTION__), *(UEnum::GetDisplayValueAsText(TrackSyncMode).ToString()));
+		return;
+	}
+
+	if (!FromTrack->bTrackSyncEnabled)
+	{
+		return;
+	}
+
+	TArray<TSharedPtr<const FChaosVDTrackInfo>> AvailableSolverTracks;
+	GetAvailableTrackInfosAtTrackFrame_AssumesLocked(EChaosVDTrackType::Solver, FromTrack, AvailableSolverTracks);
+
+	switch (FromTrack->TrackType)
+	{
+	case EChaosVDTrackType::Game:
+		{
+			for (const TSharedPtr<const FChaosVDTrackInfo>& SolverTrack : AvailableSolverTracks)
+			{
+				if (!SolverTrack->bTrackSyncEnabled)
+				{
+					continue;
+				}
+
+				const int32 SolverFrameNumber = LoadedRecording->GetLowestSolverFrameNumberGameFrame_AssumesLocked(SolverTrack->TrackID, FromTrack->CurrentFrame);
+				const int32 StageNumber = GetTrackLastStageAtFrame_AssumesLocked(EChaosVDTrackType::Solver, SolverTrack->TrackID, SolverFrameNumber);
+
+				GoToTrackFrame_AssumesLocked(PlaybackSelfInstigatorID, EChaosVDTrackType::Solver, SolverTrack->TrackID, SolverFrameNumber, StageNumber);
+			}
+
+			break;
+		}	
+	case EChaosVDTrackType::Solver:
+		{
+			if (const TSharedPtr<const FChaosVDTrackInfo>& GameFramesTrackInfo = GetTrackInfo(EChaosVDTrackType::Game, GameTrackID))
+			{
+				const int32 TargetGameFrameNumber = ConvertCurrentFrameToOtherTrackFrame_AssumesLocked(FromTrack, GameFramesTrackInfo.ToSharedRef());
+				constexpr int32 StageNumber = 0;
+				GoToTrackFrame_AssumesLocked(PlaybackSelfInstigatorID, GameFramesTrackInfo->TrackType, GameFramesTrackInfo->TrackID, TargetGameFrameNumber, StageNumber);
+			}
+
+			for (const TSharedPtr<const FChaosVDTrackInfo>& SolverTrack : AvailableSolverTracks)
+			{
+				if (!SolverTrack->bTrackSyncEnabled)
+				{
+					continue;
+				}
+
+				if (FChaosVDTrackInfo::AreSameTrack(FromTrack, SolverTrack.ToSharedRef()))
+				{
+					continue;
+				}
+
+				const int32 SolverFrameNumber = ConvertCurrentFrameToOtherTrackFrame_AssumesLocked(FromTrack, SolverTrack.ToSharedRef());
+				const int32 StageNumber = GetTrackLastStageAtFrame_AssumesLocked(EChaosVDTrackType::Solver, SolverTrack->TrackID, SolverFrameNumber);
+
+				GoToTrackFrame_AssumesLocked(PlaybackSelfInstigatorID, EChaosVDTrackType::Solver, SolverTrack->TrackID, SolverFrameNumber, StageNumber);
+			}
+			break;
+		}
+	default:
+		ensure(false);
+		break;
+	}
+}
+
+void FChaosVDPlaybackController::ToggleTrackSyncEnabled(const TSharedRef<const FChaosVDTrackInfo>& InTrackInfoRef)
+{
+	TSharedRef<FChaosVDTrackInfo> MutableTrackInfoRef = ConstCastSharedRef<FChaosVDTrackInfo>(InTrackInfoRef);
+	MutableTrackInfoRef->bTrackSyncEnabled = !MutableTrackInfoRef->bTrackSyncEnabled;
+}
+
+bool FChaosVDPlaybackController::IsPlaying() const
+{
+	return CurrentPlayingTrack && CurrentPlayingTrack->bIsPlaying;
+}
+
 void FChaosVDPlaybackController::UpdateSolverTracksData()
 {
 	if (!LoadedRecording.IsValid())
@@ -828,12 +1084,12 @@ void FChaosVDPlaybackController::UpdateSolverTracksData()
 		{
 			SolverTrackInfo = MakeShared<FChaosVDTrackInfo>();
 			SolverTrackInfo->CurrentFrame = 0;
-			SolverTrackInfo->CurrentStep = 0;
+			SolverTrackInfo->CurrentStage = 0;
 		};
 		
 		SolverTrackInfo->TrackID = SolverIDPair.Key;
 		SolverTrackInfo->MaxFrames = GetTrackFramesNumber(EChaosVDTrackType::Solver, SolverIDPair.Key);
-		SolverTrackInfo->TrackName = LoadedRecording->GetSolverName(SolverIDPair.Key);
+		SolverTrackInfo->TrackName = LoadedRecording->GetSolverFName(SolverIDPair.Key);
 		SolverTrackInfo->TrackType = EChaosVDTrackType::Solver;
 	}
 }
@@ -849,8 +1105,9 @@ void FChaosVDPlaybackController::HandleCurrentRecordingUpdated()
 	if (!GameTrackInfo.IsValid())
 	{
 		GameTrackInfo = MakeShared<FChaosVDTrackInfo>();
+		GameTrackInfo->TrackID = GameTrackID;
 		GameTrackInfo->CurrentFrame = 0;
-		GameTrackInfo->CurrentStep = 0;
+		GameTrackInfo->CurrentStage = 0;
 	}
 
 	GameTrackInfo->MaxFrames = LoadedRecording.IsValid() ? LoadedRecording->GetAvailableGameFrames_AssumesLocked().Num() : INDEX_NONE;
diff --git a/Engine/Plugins/ChaosVD/Source/ChaosVD/Private/ChaosVDPlaybackController.h b/Engine/Plugins/ChaosVD/Source/ChaosVD/Private/ChaosVDPlaybackController.h
--- a/Engine/Plugins/ChaosVD/Source/ChaosVD/Private/ChaosVDPlaybackController.h
+++ b/Engine/Plugins/ChaosVD/Source/ChaosVD/Private/ChaosVDPlaybackController.h
@@ -8,7 +8,7 @@
 #include "Templates/SharedPointer.h"
 #include "Delegates/DelegateCombinations.h"
 #include "Delegates/Delegate.h"
-#include "HAL/ThreadSafeBool.h"
+#include "Widgets/SChaosVDTimelineWidget.h"
 
 class IChaosVDPlaybackControllerInstigator;
 class UChaosVDCoreSettings;
@@ -20,14 +20,16 @@ class FChaosVDPlaybackController;
 class FString;
 
 DECLARE_MULTICAST_DELEGATE_OneParam(FChaosVDPlaybackControllerUpdated, TWeakPtr<FChaosVDPlaybackController>)
-DECLARE_MULTICAST_DELEGATE_ThreeParams(FChaosVDPlaybackControllerFrameUpdated, TWeakPtr<FChaosVDPlaybackController>, const FChaosVDTrackInfo*, FGuid);
+DECLARE_MULTICAST_DELEGATE_ThreeParams(FChaosVDPlaybackControllerFrameUpdated, TWeakPtr<FChaosVDPlaybackController>, TWeakPtr<const FChaosVDTrackInfo>, FGuid);
 
 /** Enum with the available game track types*/
 enum class EChaosVDTrackType : int32
 {
 	Invalid,
 	Game,
-	Solver
+	Solver,
+	/** Used mostly for search */
+	All
 };
 
 /** Data that represents the current state of a track and ID info*/
@@ -36,16 +38,23 @@ struct FChaosVDTrackInfo
 	int32 TrackID = INDEX_NONE;
 	EChaosVDTrackType TrackType = EChaosVDTrackType::Invalid;
 	int32 CurrentFrame = INDEX_NONE;
-	int32 CurrentStep = INDEX_NONE;
+	int32 CurrentStage = INDEX_NONE;
 	int32 LockedOnStep = INDEX_NONE;
 	int32 MaxFrames = INDEX_NONE;
-	FString TrackName;
+	FName TrackName;
+	TArray<FStringView> CurrentStageNames;
 	bool bIsReSimulated = false;
+	bool bIsPlaying = false;
+	bool bTrackSyncEnabled = true;
+
+	bool operator==(const FChaosVDTrackInfo& Other) const;
+
+	static bool AreSameTrack(const TSharedRef<const FChaosVDTrackInfo>& TrackA,  const TSharedRef<const FChaosVDTrackInfo>& TrackB);
 };
 
 struct FChaosVDQueuedTrackInfoUpdate
 {
-	FChaosVDTrackInfo TrackInfo;
+	TWeakPtr<const FChaosVDTrackInfo> TrackInfo;
 	FGuid InstigatorID;
 };
 
@@ -64,6 +73,15 @@ enum class EChaosVDUnloadRecordingFlags : uint8
 };
 ENUM_CLASS_FLAGS(EChaosVDUnloadRecordingFlags)
 
+UENUM()
+enum class EChaosVDSyncTimelinesMode : uint8
+{
+	None = 0,
+	RecordedTimestamp = 1 << 0,
+	NetworkTick = 1 << 1
+};
+ENUM_CLASS_FLAGS(EChaosVDSyncTimelinesMode)
+
 typedef TMap<int32, TSharedPtr<FChaosVDTrackInfo>> TrackInfoByIDMap;
 
 /** Loads,unloads and owns a Chaos VD recording file */
@@ -72,9 +90,11 @@ class FChaosVDPlaybackController : public TSharedFromThis<FChaosVDPlaybackContro
 public:
 
 	/** ID used for the Game Track */
-	static constexpr int32 GameTrackID  = 0 ;
-	static constexpr int32 InvalidFrameRateOverride  = -1 ;
-	static constexpr float FallbackFrameTime = 1.0f / 60.0f ;
+	static constexpr int32 GameTrackID  = 0;
+	static constexpr int32 InvalidFrameRateOverride  = -1;
+	static constexpr float FallbackFrameTime = 1.0f / 60.0f;
+
+	static inline FGuid PlaybackSelfInstigatorID = FGuid::NewGuid();
 
 	FChaosVDPlaybackController(const TWeakPtr<FChaosVDScene>& InSceneToControl);
 	virtual ~FChaosVDPlaybackController() override;
@@ -99,10 +119,13 @@ public:
 	 * @param TrackType Type of the track to to move
 	 * @param InTrackID ID of the track to move
 	 * @param FrameNumber Frame number to go
-	 * @param Step Step number to go
+	 * @param StageNumber Step number to go
 	 */
-	void GoToTrackFrame(FGuid InstigatorID, EChaosVDTrackType TrackType, int32 InTrackID, int32 FrameNumber, int32 Step);
-	void GoToTrackFrame_AssumesLocked(FGuid InstigatorID, EChaosVDTrackType TrackType, int32 InTrackID, int32 FrameNumber, int32 Step);
+	void GoToTrackFrame(FGuid InstigatorID, EChaosVDTrackType TrackType, int32 InTrackID, int32 FrameNumber, int32 StageNumber);
+	void GoToTrackFrame_AssumesLocked(FGuid InstigatorID, EChaosVDTrackType TrackType, int32 InTrackID, int32 FrameNumber, int32 StageNumber);
+
+	void GoToTrackFrameAndSync(FGuid InstigatorID, EChaosVDTrackType TrackType, int32 InTrackID, int32 FrameNumber, int32 StageNumber);
+	void GoToTrackFrame_AssumesLockedAndSync(FGuid InstigatorID, EChaosVDTrackType TrackType, int32 InTrackID, int32 FrameNumber, int32 StageNumber);
 
 	/**
 	 * Gets the number of available steps in a track at the specified frame
@@ -152,14 +175,15 @@ public:
 	 * @param InTrackID ID of the track to evaluate
 	 * @return Number of the last step
 	 */
-	int32 GetTrackLastStepAtFrame(EChaosVDTrackType TrackType, int32 InTrackID, int32 InFrameNumber) const;
+	int32 GetTrackLastStageAtFrame_AssumesLocked(EChaosVDTrackType TrackType, int32 InTrackID, int32 InFrameNumber) const;
 
 	/** Converts the current frame number of a track, to a frame number in other tracks space time
-	 * @param FromTrack Track info with the current frame number we want to convert
-	 * @param ToTrack Track info we want to use to convert the frame to
+	 * @param InFromTrack Track info with the current frame number we want to convert
+	 * @param InToTrack Track info we want to use to convert the frame to
+	 * @param TrackSyncMode Criteria or "mode" that should be used to sync a track frame with another
 	 * @return Converted Frame Number
 	 */
-	int32 ConvertCurrentFrameToOtherTrackFrame(const FChaosVDTrackInfo* FromTrack, const FChaosVDTrackInfo* ToTrack);
+	int32 ConvertCurrentFrameToOtherTrackFrame_AssumesLocked(const TSharedRef<const FChaosVDTrackInfo>& InFromTrack, const TSharedRef<const FChaosVDTrackInfo>& InToTrack, EChaosVDSyncTimelinesMode TrackSyncMode = EChaosVDSyncTimelinesMode::RecordedTimestamp);
 
 	/**
 	 * Gets all the ids of the tracks, of the specified type, that are available available on the loaded recording
@@ -173,9 +197,10 @@ public:
 	 * Gets all the ids of the tracks, of the specified type, that are available available on the loaded recording, at a specified frame
 	 * @param TrackTypeToFind Type of the tracks we are interested in
 	 * @param OutTrackInfo Array where any found track info data will be added
-	 * @param TrackFrameInfo Ptr to the track info with the current frame to evaluate
+	 * @param InFromTrack Ptr to the track info with the current frame to evaluate
 	 */
-	void GetAvailableTrackInfosAtTrackFrame(EChaosVDTrackType TrackTypeToFind, TArray<TSharedPtr<FChaosVDTrackInfo>>& OutTrackInfo, const FChaosVDTrackInfo* TrackFrameInfo);
+	void GetAvailableTrackInfosAtTrackFrame(EChaosVDTrackType TrackTypeToFind, const TSharedRef<const FChaosVDTrackInfo>& InFromTrack, TArray<TSharedPtr<const FChaosVDTrackInfo>>& OutTrackInfo);
+	void GetAvailableTrackInfosAtTrackFrame_AssumesLocked(EChaosVDTrackType TrackTypeToFind, const TSharedRef<const FChaosVDTrackInfo>& InFromTrack, TArray<TSharedPtr<const FChaosVDTrackInfo>>& OutTrackInfo);
 
 	/**
 	 * Gets the track info of the specified type with the specified ID
@@ -183,7 +208,7 @@ public:
 	 * @param TrackID ID of the track to find
 	 * @return Ptr to the found track info data - Null if nothing was found
 	 */
-	const FChaosVDTrackInfo* GetTrackInfo(EChaosVDTrackType TrackType, int32 TrackID);
+	TSharedPtr<const FChaosVDTrackInfo> GetTrackInfo(EChaosVDTrackType TrackType, int32 TrackID);
 
 	/**
 	 * Gets the track info of the specified type with the specified ID
@@ -191,7 +216,7 @@ public:
 	 * @param TrackID ID of the track to find
 	 * @return Ptr to the found track info data - Null if nothing was found.
 	 */
-	FChaosVDTrackInfo* GetMutableTrackInfo(EChaosVDTrackType TrackType, int32 TrackID);
+	TSharedPtr<FChaosVDTrackInfo> GetMutableTrackInfo(EChaosVDTrackType TrackType, int32 TrackID);
 
 	/**
 	 * Locks the steps timeline of a given track so each time you move between frames, it will automatically scrub to the locked in step
@@ -224,16 +249,7 @@ public:
 	/** Updates the loaded recording state to indicate is not longer receiving live updates */
 	void HandleDisconnectedFromSession();
 
-	void RequestPause() { bPauseRequested = true; }
-	void RequestUnpause() { bPauseRequested = false; }
-	bool HasPauseRequest() const { return bPauseRequested; }
-	void RequestStop(const IChaosVDPlaybackControllerInstigator& InPlaybackInstigator);
-
-	/** Returns the ID of the current playback instigator that is controlling the playback */
-	const FGuid& GetPlaybackInstigatorWithExclusiveControlsID() const { return CurrentPlaybackInstigator; }
-
-	bool AcquireExclusivePlaybackControls(const IChaosVDPlaybackControllerInstigator& InPlaybackInstigator);
-	bool ReleaseExclusivePlaybackControls(const IChaosVDPlaybackControllerInstigator& InPlaybackInstigator);
+	void StopPlayback(const FGuid& InstigatorGUID);
 
 	bool IsUsingFrameRateOverride() const { return bUseFrameRateOverride; }
 
@@ -243,20 +259,35 @@ public:
 	int32 GetFrameRateOverride() const;
 	void SetFrameRateOverride(float NewFrameRateOverride);
 
-	float GetFrameTimeForTrack(EChaosVDTrackType TrackType, int32 TrackID, const FChaosVDTrackInfo& TrackInfo) const;
+	float GetFrameTimeForTrack(EChaosVDTrackType TrackType, int32 TrackID, const TSharedRef<const FChaosVDTrackInfo>& InTrackInfo) const;
 
 	void UpdateTrackVisibility(EChaosVDTrackType Type, int32 TrackID, bool bNewVisibility);
 
+	void HandleFramePlaybackControlInput(EChaosVDPlaybackButtonsID ButtonID, const TSharedRef<const FChaosVDTrackInfo>& InTrackInfoRef, FGuid Instigator);
+	void HandleFrameStagePlaybackControlInput(EChaosVDPlaybackButtonsID ButtonID, const TSharedRef<const FChaosVDTrackInfo>& InTrackInfoRef, FGuid Instigator);
+
+	void TickPlayback(float DeltaTime);
+
+	TSharedPtr<FChaosVDTrackInfo> GetCurrentPlayingTrackInfo() const { return CurrentPlayingTrack; }
+
+	void GetTracksByType(EChaosVDTrackType Type, TArray<TSharedPtr<FChaosVDTrackInfo>>& OutTracks);
+	void SyncTracks(const TSharedRef<const FChaosVDTrackInfo>& FromTrack, EChaosVDSyncTimelinesMode TrackSyncMode = EChaosVDSyncTimelinesMode::RecordedTimestamp);
+	void SyncTracks_AssumesLocked(const TSharedRef<const FChaosVDTrackInfo>& FromTrack, EChaosVDSyncTimelinesMode TrackSyncMode = EChaosVDSyncTimelinesMode::RecordedTimestamp);
+
+	void ToggleTrackSyncEnabled(const TSharedRef<const FChaosVDTrackInfo>& InTrackInfoRef);
+
+	bool IsPlaying() const;
+
 protected:
 
 	/** Updates (or adds) solvers data from the loaded recording to the solver tracks */
 	void UpdateSolverTracksData();
 
 	/** Updates the controlled scene with the loaded data at specified game frame */
-	void GoToRecordedGameFrame_AssumesLocked(int32 FrameNumber, FGuid InstigatorID, int32 Attempts = 0);
+	void GoToRecordedGameFrame_AssumesLocked(int32 FrameNumber, FGuid InstigatorID);
 
 	/** Updates the controlled scene with the loaded data at specified solver frame and solver step */
-	void GoToRecordedSolverStep_AssumesLocked(int32 InTrackID, int32 FrameNumber, int32 Step, FGuid InstigatorID, int32 Attempts = 0);
+	void GoToRecordedSolverStage_AssumesLocked(int32 InTrackID, int32 FrameNumber, int32 StageNumber, FGuid InstigatorID);
 
 	/** Handles any data changes on the loaded recording - Usually called during Trace analysis */
 	void HandleCurrentRecordingUpdated();
@@ -265,13 +296,16 @@ protected:
 	void PlayFromClosestKeyFrame_AssumesLocked(int32 InTrackID, int32 FrameNumber, FChaosVDScene& InSceneToControl) const;
 
 	/** Add the provided track info update to the queue. The update will be broadcast in the game thread */
-	void EnqueueTrackInfoUpdate(const FChaosVDTrackInfo& InTrackInfo, FGuid InstigatorID);
+	void EnqueueTrackInfoUpdate(const TSharedRef<const FChaosVDTrackInfo>& InTrackInfo, FGuid InstigatorID);
 
 	/** Add the provided Geometry info data to the queue. The update will be broadcast in the game thread */
 	void EnqueueGeometryDataUpdate(const Chaos::FConstImplicitObjectPtr& NewGeometry, const uint32 GeometryID);
 
 	void PlaySolverStepData(int32 TrackID, const TSharedRef<FChaosVDScene>& InSceneToControlSharedPtr, const FChaosVDSolverFrameData& InSolverFrameData, int32 StepIndex);
 
+	template <typename TVisitorCallback>
+	void VisitAvailableTracks(const TVisitorCallback& VisitorCallback);
+
 	/** Map containing all track info, by track type*/
 	TMap<EChaosVDTrackType, TrackInfoByIDMap> TrackInfoPerType;
 
@@ -284,7 +318,7 @@ protected:
 	/** Delegate called when the data on the loaded recording changes */
 	FChaosVDPlaybackControllerUpdated ControllerUpdatedDelegate;
 
-	/** Delegate called when the a in a track changes */
+	/** Delegate called when the data in a track changes */
 	FChaosVDPlaybackControllerFrameUpdated ControllerFrameUpdatedDelegate;
 
 	/** Set to true when the recording data controlled by this Playback Controller is updated, the update delegate will be called on the GT */
@@ -312,5 +346,22 @@ protected:
 
 	FDelegateHandle RecordingStoppedHandle;
 
-	FGuid CurrentPlaybackInstigator;
+	TSharedPtr<FChaosVDTrackInfo> CurrentPlayingTrack;
+
+	float CurrentPlaybackTime = 0.0f;
 };
+
+template <typename TVisitorCallback>
+void FChaosVDPlaybackController::VisitAvailableTracks(const TVisitorCallback& VisitorCallback)
+{
+	for (const TPair<EChaosVDTrackType, TrackInfoByIDMap>& TracksByType : TrackInfoPerType)
+	{
+		for (const TPair<int32, TSharedPtr<FChaosVDTrackInfo>>& TracksById : TracksByType.Value )
+		{
+			if (!VisitorCallback(TracksById.Value))
+			{
+				return;
+			}
+		}
+	}
+}
diff --git a/Engine/Plugins/ChaosVD/Source/ChaosVD/Private/ChaosVDPlaybackControllerObserver.h b/Engine/Plugins/ChaosVD/Source/ChaosVD/Private/ChaosVDPlaybackControllerObserver.h
--- a/Engine/Plugins/ChaosVD/Source/ChaosVD/Private/ChaosVDPlaybackControllerObserver.h
+++ b/Engine/Plugins/ChaosVD/Source/ChaosVD/Private/ChaosVDPlaybackControllerObserver.h
@@ -26,7 +26,7 @@ public:
 protected:
 	virtual void RegisterNewController(TWeakPtr<FChaosVDPlaybackController> NewController);
 	virtual void HandlePlaybackControllerDataUpdated(TWeakPtr<FChaosVDPlaybackController> InController) {};
-	virtual void HandleControllerTrackFrameUpdated(TWeakPtr<FChaosVDPlaybackController> InController, const FChaosVDTrackInfo* UpdatedTrackInfo, FGuid InstigatorGuid){};
+	virtual void HandleControllerTrackFrameUpdated(TWeakPtr<FChaosVDPlaybackController> InController, TWeakPtr<const FChaosVDTrackInfo> UpdatedTrackInfo, FGuid InstigatorGuid){};
 	
 	virtual void HandlePostSelectionChange(const UTypedElementSelectionSet* ChangesSelectionSet) override{};
 	
diff --git a/Engine/Plugins/ChaosVD/Source/ChaosVD/Private/ChaosVDScene.cpp b/Engine/Plugins/ChaosVD/Source/ChaosVD/Private/ChaosVDScene.cpp
--- a/Engine/Plugins/ChaosVD/Source/ChaosVD/Private/ChaosVDScene.cpp
+++ b/Engine/Plugins/ChaosVD/Source/ChaosVD/Private/ChaosVDScene.cpp
@@ -152,7 +152,7 @@ void FChaosVDScene::UpdateFromRecordedStepData(const int32 SolverID, const FChao
 		constexpr float AmountOfWork = 1.0f;
 		const float PercentagePerElement = 1.0f / InRecordedStepData.RecordedParticlesData.Num();
 
-		const FText ProgressBarTitle = FText::Format(FTextFormat(LOCTEXT("ProcessingParticleData", "Processing Particle Data for {0} Solver with ID {1} ...")), FText::FromString(SolverSceneData->GetSolverName()), FText::AsNumber(SolverID));
+		const FText ProgressBarTitle = FText::Format(FTextFormat(LOCTEXT("ProcessingParticleData", "Processing Particle Data for {0} Solver with ID {1} ...")), FText::FromName(SolverSceneData->GetSolverName()), FText::AsNumber(SolverID));
 		FScopedSlowTask UpdatingSceneSlowTask(AmountOfWork, ProgressBarTitle);
 		UpdatingSceneSlowTask.MakeDialogDelayed(ChaosVDSceneUIOptions::DelayToShowProgressDialogThreshold, ChaosVDSceneUIOptions::bShowCancelButton, ChaosVDSceneUIOptions::bAllowInPIE);
 	
@@ -264,10 +264,11 @@ void FChaosVDScene::CreateSolverInfoActor(int32 SolverID)
 		AChaosVDSolverInfoActor* SolverDataInfo = PhysicsVDWorld->SpawnActor<AChaosVDSolverInfoActor>();
 		check(SolverDataInfo);
 
-		FString SolverName = LoadedRecording->GetSolverName_AssumedLocked(SolverID);
-		const bool bIsServer = SolverName.Contains(TEXT("Server"));
+		FName SolverName = LoadedRecording->GetSolverFName_AssumedLocked(SolverID);
+		FString NameAsString = SolverName.ToString();
+		const bool bIsServer = NameAsString.Contains(TEXT("Server"));
 
-		const FStringFormatOrderedArguments Args {SolverName, FString::FromInt(SolverID)};
+		const FStringFormatOrderedArguments Args {NameAsString, FString::FromInt(SolverID)};
 		const FName FolderPath = *FString::Format(TEXT("Solver {0} | ID {1}"), Args);
 
 		SolverDataInfo->SetFolderPath(FolderPath);
@@ -283,7 +284,7 @@ void FChaosVDScene::CreateSolverInfoActor(int32 SolverID)
 	}
 }
 
-void FChaosVDScene::HandleEnterNewGameFrame(int32 FrameNumber, const TArray<int32>& AvailableSolversIds, const FChaosVDGameFrameData& InNewGameFrameData)
+void FChaosVDScene::HandleEnterNewGameFrame(int32 FrameNumber, const TArray<int32, TInlineAllocator<16>>& AvailableSolversIds, const FChaosVDGameFrameData& InNewGameFrameData, TArray<int32, TInlineAllocator<16>>& OutRemovedSolversIds)
 {
 	// Currently the particle actors from all the solvers are in the same level, and we manage them by keeping track
 	// of to which solvers they belong using maps.
@@ -314,6 +315,8 @@ void FChaosVDScene::HandleEnterNewGameFrame(int32 FrameNumber, const TArray<int3
 				PhysicsVDWorld->DestroyActor(SolverInfoActor);
 			}
 
+			OutRemovedSolversIds.Add(RemoveIterator.Key());
+
 			RemoveIterator.RemoveCurrent();
 			AmountRemoved++;
 		}
@@ -331,15 +334,15 @@ void FChaosVDScene::HandleEnterNewGameFrame(int32 FrameNumber, const TArray<int3
 			QueryDataComponent->UpdateQueriesFromFrameData(InNewGameFrameData);
 		}
 	}
+}
 
-	for (int32 SolverID : AvailableSolversIds)
+void FChaosVDScene::HandleEnterNewSolverFrame(int32 FrameNumber, const FChaosVDSolverFrameData& InFrameData)
+{
+	if (AChaosVDSolverInfoActor* SolverDataInfoContainer = SolverDataContainerBySolverID.FindChecked(InFrameData.SolverID))
 	{
-		if (AChaosVDSolverInfoActor* SolverDataInfoContainer = SolverDataContainerBySolverID.FindChecked(SolverID))
+		if (UChaosVDSolverCharacterGroundConstraintDataComponent* DataContainer = SolverDataInfoContainer->GetCharacterGroundConstraintDataComponent())
 		{
-			if (UChaosVDSolverCharacterGroundConstraintDataComponent* DataContainer = SolverDataInfoContainer->GetCharacterGroundConstraintDataComponent())
-			{
-				DataContainer->UpdateConstraintData(InNewGameFrameData.RecordedCharacterGroundConstraints);
-			}
+			DataContainer->UpdateConstraintData(InFrameData.RecordedCharacterGroundConstraints);
 		}
 	}
 }
diff --git a/Engine/Plugins/ChaosVD/Source/ChaosVD/Private/ChaosVDScene.h b/Engine/Plugins/ChaosVD/Source/ChaosVD/Private/ChaosVDScene.h
--- a/Engine/Plugins/ChaosVD/Source/ChaosVD/Private/ChaosVDScene.h
+++ b/Engine/Plugins/ChaosVD/Source/ChaosVD/Private/ChaosVDScene.h
@@ -67,7 +67,8 @@ public:
 	// No need to deprecate the old version since it is not a public API nor inline 
 	void HandleNewGeometryData(const Chaos::FConstImplicitObjectPtr& Geometry, const uint32 GeometryID) const;
 
-	void HandleEnterNewGameFrame(int32 FrameNumber, const TArray<int32>& AvailableSolversIds, const FChaosVDGameFrameData& InNewGameFrameData);
+	void HandleEnterNewGameFrame(int32 FrameNumber, const TArray<int32, TInlineAllocator<16>>& AvailableSolversIds, const FChaosVDGameFrameData& InNewGameFrameData, TArray<int32, TInlineAllocator<16>>& OutRemovedSolversIds);
+	void HandleEnterNewSolverFrame(int32 FrameNumber, const FChaosVDSolverFrameData& InFrameData);
 
 	/** Deletes all actors of the Scene and underlying UWorld */
 	void CleanUpScene();
diff --git a/Engine/Plugins/ChaosVD/Source/ChaosVD/Private/ChaosVDStyle.cpp b/Engine/Plugins/ChaosVD/Source/ChaosVD/Private/ChaosVDStyle.cpp
--- a/Engine/Plugins/ChaosVD/Source/ChaosVD/Private/ChaosVDStyle.cpp
+++ b/Engine/Plugins/ChaosVD/Source/ChaosVD/Private/ChaosVDStyle.cpp
@@ -70,6 +70,9 @@ TSharedRef< FSlateStyleSet > FChaosVDStyle::Create()
 	Style->Set("RecordToLiveIcon", new IMAGE_BRUSH(TEXT("RecordLiveSession"), Icon16x16));
 	Style->Set("ConnectionIcon", new IMAGE_BRUSH_SVG(TEXT("Connection"), Icon16x16));
 
+	Style->Set("LinkedIcon", new IMAGE_BRUSH_SVG(TEXT("Linked"), Icon16x16));
+	Style->Set("UnLinkedIcon", new IMAGE_BRUSH_SVG(TEXT("UnLinked"), Icon16x16));
+
 	return Style;
 }
 
diff --git a/Engine/Plugins/ChaosVD/Source/ChaosVD/Private/Trace/ChaosVDTraceAnalyzer.cpp b/Engine/Plugins/ChaosVD/Source/ChaosVD/Private/Trace/ChaosVDTraceAnalyzer.cpp
--- a/Engine/Plugins/ChaosVD/Source/ChaosVD/Private/Trace/ChaosVDTraceAnalyzer.cpp
+++ b/Engine/Plugins/ChaosVD/Source/ChaosVD/Private/Trace/ChaosVDTraceAnalyzer.cpp
@@ -87,7 +87,7 @@ bool FChaosVDTraceAnalyzer::OnEvent(uint16 RouteId, EStyle Style, const FOnEvent
 
 			FWideStringView DebugNameView;
 			EventData.GetString("DebugName", DebugNameView);
-			NewFrameData.DebugName = DebugNameView;
+			NewFrameData.DebugFName = FName(DebugNameView);
 
 			// Currently not all solvers have an end frame event, so lets just set the end frame time of the previous frame, with the start of this new one.
 			{
diff --git a/Engine/Plugins/ChaosVD/Source/ChaosVD/Private/Trace/ChaosVDTraceProvider.cpp b/Engine/Plugins/ChaosVD/Source/ChaosVD/Private/Trace/ChaosVDTraceProvider.cpp
--- a/Engine/Plugins/ChaosVD/Source/ChaosVD/Private/Trace/ChaosVDTraceProvider.cpp
+++ b/Engine/Plugins/ChaosVD/Source/ChaosVD/Private/Trace/ChaosVDTraceProvider.cpp
@@ -61,11 +61,7 @@ void FChaosVDTraceProvider::StartSolverFrame(const int32 InSolverID, FChaosVDSol
 
 void FChaosVDTraceProvider::CommitProcessedGameFramesToRecording()
 {
-	// This is a kind of made up number, this includes RBAN solvers so we have more than 1 for sure.
-	constexpr int32 MinSolverIDsExpected = 10;
-	
-	static TArray<int32> SolverIDs;
-	SolverIDs.Reserve(MinSolverIDsExpected);
+	TArray<int32, TInlineAllocator<16>> SolverIDs;
 
 	// The Game Frame events are not generated by CVD trace code and we don't have control over them
 	// we use them as general timestamps.
diff --git a/Engine/Plugins/ChaosVD/Source/ChaosVD/Private/Trace/DataProcessors/ChaosVDCharacterGroundConstraintDataProcessor.cpp b/Engine/Plugins/ChaosVD/Source/ChaosVD/Private/Trace/DataProcessors/ChaosVDCharacterGroundConstraintDataProcessor.cpp
--- a/Engine/Plugins/ChaosVD/Source/ChaosVD/Private/Trace/DataProcessors/ChaosVDCharacterGroundConstraintDataProcessor.cpp
+++ b/Engine/Plugins/ChaosVD/Source/ChaosVD/Private/Trace/DataProcessors/ChaosVDCharacterGroundConstraintDataProcessor.cpp
@@ -24,7 +24,7 @@ bool FChaosVDCharacterGroundConstraintDataProcessor::ProcessRawData(const TArray
 
 	if (bSuccess)
 	{
-		if (const TSharedPtr<FChaosVDGameFrameData> CurrentFrameData = ProviderSharedPtr->GetCurrentGameFrame().Pin())
+		if (FChaosVDSolverFrameData* CurrentFrameData = ProviderSharedPtr->GetCurrentSolverFrame(Constraint->SolverID))
 		{
 			CurrentFrameData->RecordedCharacterGroundConstraints.Add(Constraint);
 		}
diff --git a/Engine/Plugins/ChaosVD/Source/ChaosVD/Private/Visualizers/ChaosVDParticleDataComponentVisualizer.cpp b/Engine/Plugins/ChaosVD/Source/ChaosVD/Private/Visualizers/ChaosVDParticleDataComponentVisualizer.cpp
--- a/Engine/Plugins/ChaosVD/Source/ChaosVD/Private/Visualizers/ChaosVDParticleDataComponentVisualizer.cpp
+++ b/Engine/Plugins/ChaosVD/Source/ChaosVD/Private/Visualizers/ChaosVDParticleDataComponentVisualizer.cpp
@@ -249,7 +249,7 @@ void FChaosVDParticleDataComponentVisualizer::DrawVisualizationForParticleData(c
 	}
 
 	const FQuat& OwnerRotation =  InVisualizationContext.SpaceTransform.TransformRotation(InParticleDataViewer.ParticlePositionRotation.MR);
-	const FVector OwnerCoMLocation = OwnerLocation + OwnerRotation *  InVisualizationContext.SpaceTransform.TransformPosition(InParticleDataViewer.ParticleMassProps.MCenterOfMass);
+	const FVector OwnerCoMLocation = InVisualizationContext.SpaceTransform.TransformPosition(InParticleDataViewer.ParticlePositionRotation.MX + (InParticleDataViewer.ParticlePositionRotation.MR *  InParticleDataViewer.ParticleMassProps.MCenterOfMass));
 	
 	FChaosVDScopedParticleHitProxy ScopedHitProxy(PDI, new HChaosVDParticleDataProxy(Component, {InParticleDataViewer.ParticleIndex, InParticleDataViewer.SolverID} ));
 
diff --git a/Engine/Plugins/ChaosVD/Source/ChaosVD/Private/Widgets/ChaosVDPlaybackControlsHelper.cpp b/Engine/Plugins/ChaosVD/Source/ChaosVD/Private/Widgets/ChaosVDPlaybackControlsHelper.cpp
deleted file mode 100644
diff --git a/Engine/Plugins/ChaosVD/Source/ChaosVD/Private/Widgets/ChaosVDPlaybackControlsHelper.h b/Engine/Plugins/ChaosVD/Source/ChaosVD/Private/Widgets/ChaosVDPlaybackControlsHelper.h
deleted file mode 100644
diff --git a/Engine/Plugins/ChaosVD/Source/ChaosVD/Private/Widgets/SChaosVDGameFramesPlaybackControls.cpp b/Engine/Plugins/ChaosVD/Source/ChaosVD/Private/Widgets/SChaosVDGameFramesPlaybackControls.cpp
new file mode 100644
--- /dev/null
+++ b/Engine/Plugins/ChaosVD/Source/ChaosVD/Private/Widgets/SChaosVDGameFramesPlaybackControls.cpp
@@ -0,0 +1,120 @@
+// Copyright Epic Games, Inc. All Rights Reserved.
+
+#include "SChaosVDGameFramesPlaybackControls.h"
+
+void SChaosVDGameFramesPlaybackControls::Construct(const FArguments& InArgs, const TWeakPtr<FChaosVDPlaybackController>& InPlaybackController)
+{
+	this->ChildSlot
+	[
+		SAssignNew(FramesTimelineWidget, SChaosVDTimelineWidget)
+			.IsEnabled_Raw(this, &SChaosVDGameFramesPlaybackControls::CanPlayback)
+			.IsPlaying_Raw(this, &SChaosVDGameFramesPlaybackControls::IsPlaying)
+			.ButtonVisibilityFlags(EChaosVDTimelineElementIDFlags::AllPlayback)
+			.OnFrameChanged_Raw(this, &SChaosVDGameFramesPlaybackControls::OnFrameSelectionUpdated)
+			.OnButtonClicked(this, &SChaosVDGameFramesPlaybackControls::HandleFramePlaybackButtonClicked)
+			.MinFrames_Raw(this, &SChaosVDGameFramesPlaybackControls::GetMinFrames)
+			.MaxFrames_Raw(this, &SChaosVDGameFramesPlaybackControls::GetMaxFrames)
+			.CurrentFrame_Raw(this, &SChaosVDGameFramesPlaybackControls::GetCurrentFrame)
+			.ButtonEnabledFlags_Raw(this, &SChaosVDGameFramesPlaybackControls::GetElementEnabledFlags)
+	];
+
+	RegisterNewController(InPlaybackController);
+}
+
+void SChaosVDGameFramesPlaybackControls::OnFrameSelectionUpdated(int32 NewFrameIndex)
+{
+	if (const TSharedPtr<FChaosVDPlaybackController> PlaybackControllerPtr = PlaybackController.Pin())
+	{
+		constexpr int32 SolverStage = 0;
+		PlaybackControllerPtr->GoToTrackFrameAndSync(GetInstigatorID(), EChaosVDTrackType::Game, FChaosVDPlaybackController::GameTrackID, NewFrameIndex, SolverStage);
+	}
+}
+
+void SChaosVDGameFramesPlaybackControls::RegisterNewController(TWeakPtr<FChaosVDPlaybackController> NewController)
+{
+	FChaosVDPlaybackControllerObserver::RegisterNewController(NewController);
+
+	if (const TSharedPtr<FChaosVDPlaybackController> PlaybackControllerPtr = PlaybackController.Pin())
+	{
+		if (TSharedPtr<FChaosVDTrackInfo> GameTrackInfoPtr = PlaybackControllerPtr->GetMutableTrackInfo(EChaosVDTrackType::Game, FChaosVDPlaybackController::GameTrackID))
+		{
+			GameTrackInfoRef = GameTrackInfoPtr.ToSharedRef();
+		}
+	}
+}
+
+void SChaosVDGameFramesPlaybackControls::HandleFramePlaybackButtonClicked(EChaosVDPlaybackButtonsID ButtonID)
+{
+	if (const TSharedPtr<FChaosVDPlaybackController> PlaybackControllerPtr = PlaybackController.Pin())
+	{
+		PlaybackControllerPtr->HandleFramePlaybackControlInput(ButtonID, GameTrackInfoRef, GetInstigatorID());
+	}
+}
+
+void SChaosVDGameFramesPlaybackControls::HandlePlaybackControllerDataUpdated(TWeakPtr<FChaosVDPlaybackController> InController)
+{
+	if (const TSharedPtr<FChaosVDPlaybackController> PlaybackControllerPtr = PlaybackController.Pin())
+	{
+		if (TSharedPtr<FChaosVDTrackInfo> GameTrackInfoPtr = PlaybackControllerPtr->GetMutableTrackInfo(EChaosVDTrackType::Game, FChaosVDPlaybackController::GameTrackID))
+		{
+			GameTrackInfoRef = GameTrackInfoPtr.ToSharedRef();
+		}
+	}
+}
+
+bool SChaosVDGameFramesPlaybackControls::CanPlayback() const
+{
+	bool bCanControlPlayback = false;
+	if (const TSharedPtr<FChaosVDPlaybackController> PlaybackControllerPtr = PlaybackController.Pin())
+	{
+		bCanControlPlayback = PlaybackControllerPtr->IsRecordingLoaded();
+
+		if (bCanControlPlayback)
+		{
+			TSharedPtr<FChaosVDTrackInfo> CurrentTrackBeingPlayed = PlaybackControllerPtr->GetCurrentPlayingTrackInfo();
+			bCanControlPlayback = !CurrentTrackBeingPlayed || (CurrentTrackBeingPlayed && CurrentTrackBeingPlayed->TrackType == EChaosVDTrackType::Game);
+		}
+	}
+
+	return bCanControlPlayback;
+}
+
+bool SChaosVDGameFramesPlaybackControls::IsPlaying() const
+{
+	return GameTrackInfoRef->bIsPlaying;
+}
+
+int32 SChaosVDGameFramesPlaybackControls::GetCurrentFrame() const
+{
+	return GameTrackInfoRef->CurrentFrame >= 0 ? GameTrackInfoRef->CurrentFrame : 0 ;
+}
+
+int32 SChaosVDGameFramesPlaybackControls::GetMinFrames() const
+{
+	return 0;
+}
+
+int32 SChaosVDGameFramesPlaybackControls::GetMaxFrames() const
+{
+	return GameTrackInfoRef->MaxFrames > 0 ? GameTrackInfoRef->MaxFrames - 1 : 0;
+}
+
+EChaosVDTimelineElementIDFlags SChaosVDGameFramesPlaybackControls::GetElementEnabledFlags() const
+{
+	EChaosVDTimelineElementIDFlags EnabledButtonFlags = EChaosVDTimelineElementIDFlags::All;
+	if (const TSharedPtr<FChaosVDPlaybackController> ControllerSharedPtr = PlaybackController.Pin())
+	{
+		if (ControllerSharedPtr->IsPlayingLiveSession())
+		{
+			constexpr EChaosVDTimelineElementIDFlags PlaybackElementDisabledDuringLiveSession = EChaosVDTimelineElementIDFlags::Stop | EChaosVDTimelineElementIDFlags::Next | EChaosVDTimelineElementIDFlags::Prev;
+			EnumRemoveFlags(EnabledButtonFlags, PlaybackElementDisabledDuringLiveSession);
+			return EnabledButtonFlags;
+		}
+		else
+		{
+			return EnabledButtonFlags;
+		}
+	}
+
+	return EnabledButtonFlags;
+}
diff --git a/Engine/Plugins/ChaosVD/Source/ChaosVD/Private/Widgets/SChaosVDGameFramesPlaybackControls.h b/Engine/Plugins/ChaosVD/Source/ChaosVD/Private/Widgets/SChaosVDGameFramesPlaybackControls.h
new file mode 100644
--- /dev/null
+++ b/Engine/Plugins/ChaosVD/Source/ChaosVD/Private/Widgets/SChaosVDGameFramesPlaybackControls.h
@@ -0,0 +1,52 @@
+// Copyright Epic Games, Inc. All Rights Reserved.
+
+#pragma once
+
+#include "ChaosVDPlaybackController.h"
+#include "ChaosVDPlaybackControllerInstigator.h"
+#include "ChaosVDPlaybackControllerObserver.h"
+#include "SChaosVDTimelineWidget.h"
+#include "Styling/ISlateStyle.h"
+#include "Templates/SharedPointer.h"
+#include "Widgets/SCompoundWidget.h"
+
+enum class EChaosVDPlaybackButtonsID : uint8;
+struct FChaosVDTrackInfo;
+class FChaosVDPlaybackController;
+class SChaosVDTimelineWidget;
+
+class SChaosVDGameFramesPlaybackControls : public SCompoundWidget , public FChaosVDPlaybackControllerObserver, public IChaosVDPlaybackControllerInstigator
+{
+public:
+	
+	explicit SChaosVDGameFramesPlaybackControls()
+	: GameTrackInfoRef(MakeShared<FChaosVDTrackInfo>())
+	{
+	}
+	
+	SLATE_BEGIN_ARGS(SChaosVDGameFramesPlaybackControls){}
+	SLATE_END_ARGS()
+
+	void Construct(const FArguments& InArgs, const TWeakPtr<FChaosVDPlaybackController>& InPlaybackController);
+
+protected:
+	void OnFrameSelectionUpdated(int32 NewFrameIndex);
+
+	virtual void RegisterNewController(TWeakPtr<FChaosVDPlaybackController> NewController) override;
+
+	void HandleFramePlaybackButtonClicked(EChaosVDPlaybackButtonsID ButtonID);
+
+	virtual void HandlePlaybackControllerDataUpdated(TWeakPtr<FChaosVDPlaybackController> InController) override;
+
+	bool CanPlayback() const;
+
+	bool IsPlaying() const;
+
+	int32 GetCurrentFrame() const;
+	int32 GetMinFrames() const;
+	int32 GetMaxFrames() const;
+	EChaosVDTimelineElementIDFlags GetElementEnabledFlags() const;
+
+	TSharedRef<FChaosVDTrackInfo> GameTrackInfoRef;
+	TSharedPtr<SChaosVDTimelineWidget> FramesTimelineWidget;
+};
diff --git a/Engine/Plugins/ChaosVD/Source/ChaosVD/Private/Widgets/SChaosVDPlaybackViewport.cpp b/Engine/Plugins/ChaosVD/Source/ChaosVD/Private/Widgets/SChaosVDPlaybackViewport.cpp
--- a/Engine/Plugins/ChaosVD/Source/ChaosVD/Private/Widgets/SChaosVDPlaybackViewport.cpp
+++ b/Engine/Plugins/ChaosVD/Source/ChaosVD/Private/Widgets/SChaosVDPlaybackViewport.cpp
@@ -8,8 +8,8 @@
 #include "ChaosVDPlaybackViewportClient.h"
 #include "ChaosVDScene.h"
 #include "EditorModeManager.h"
+#include "SChaosVDGameFramesPlaybackControls.h"
 #include "Framework/Application/SlateApplication.h"
-#include "Widgets/ChaosVDPlaybackControlsHelper.h"
 #include "Widgets/SChaosVDTimelineWidget.h"
 #include "Widgets/SChaosVDViewportToolbar.h"
 #include "Widgets/Text/STextBlock.h"
@@ -33,6 +33,8 @@ SChaosVDPlaybackViewport::~SChaosVDPlaybackViewport()
 		ExternalInvalidateHandlerHandle.Reset();
 	}
 
+	UnbindFromSceneUpdateEvents();
+
 	PlaybackViewportClient->Viewport = nullptr;
 	PlaybackViewportClient.Reset();
 }
@@ -95,17 +97,12 @@ void SChaosVDPlaybackViewport::Construct(const FArguments& InArgs, TWeakPtr<FCha
 			]
 			+SVerticalBox::Slot()
 			[
-				SAssignNew(GameFramesTimelineWidget, SChaosVDTimelineWidget)
-				.IsEnabled_Raw(this, &SChaosVDPlaybackViewport::CanPlayback)
-				.ButtonVisibilityFlags(static_cast<uint16>(EChaosVDTimelineElementIDFlags::AllPlayback))
-				.OnFrameChanged_Raw(this, &SChaosVDPlaybackViewport::OnFrameSelectionUpdated)
-				.OnButtonClicked(this, &SChaosVDPlaybackViewport::HandlePlaybackButtonClicked)
-				.MaxFrames(0)
+				SAssignNew(GameFramesPlaybackControls, SChaosVDGameFramesPlaybackControls, InPlaybackController)
 			]
-		]	
+		]
 	];
 
-	ExternalInvalidateHandlerHandle = ExternalViewportInvalidationRequestHandler.AddSP(this, &SChaosVDPlaybackViewport::HanldeExternalViewportInvalidateRequest);
+	ExternalInvalidateHandlerHandle = ExternalViewportInvalidationRequestHandler.AddSP(this, &SChaosVDPlaybackViewport::HandleExternalViewportInvalidateRequest);
 
 	RegisterNewController(InPlaybackController);
 }
@@ -253,71 +250,9 @@ void SChaosVDPlaybackViewport::HandlePlaybackControllerDataUpdated(TWeakPtr<FCha
 		RegisterNewController(InController);
 	}
 
-	const TSharedPtr<FChaosVDPlaybackController> ControllerSharedPtr = PlaybackController.Pin();
-	if (ControllerSharedPtr.IsValid() && ControllerSharedPtr->IsRecordingLoaded())
-	{
-		if (const FChaosVDTrackInfo* TrackInfo = ControllerSharedPtr->GetTrackInfo(EChaosVDTrackType::Game, FChaosVDPlaybackController::GameTrackID))
-		{
-			// Max is inclusive and we use this to request as the index on the recorded frames/steps arrays so we need to -1 to the available frames/steps
-			GameFramesTimelineWidget->UpdateMinMaxValue(0, TrackInfo->MaxFrames != INDEX_NONE ? TrackInfo->MaxFrames -1  : 0);
-
-			const bool bNeedsToBroadcastChange = Chaos::VisualDebugger::Cvars::bBroadcastGameFrameUpdateEvenIfNotChanged || (GameFramesTimelineWidget->GetCurrentFrame() == 0 || GameFramesTimelineWidget->GetCurrentFrame() != TrackInfo->CurrentFrame);
-			GameFramesTimelineWidget->SetCurrentTimelineFrame(TrackInfo->CurrentFrame, bNeedsToBroadcastChange ? EChaosVDSetTimelineFrameFlags::BroadcastChange : EChaosVDSetTimelineFrameFlags::Silent);
-
-			constexpr uint16 PlaybackElementDisabledDuringLiveSession = static_cast<uint16>(EChaosVDTimelineElementIDFlags::Stop | EChaosVDTimelineElementIDFlags::Next | EChaosVDTimelineElementIDFlags::Prev);
-			if (ControllerSharedPtr->IsPlayingLiveSession())
-			{
-				if (!ControllerSharedPtr->HasPauseRequest())
-				{
-					GameFramesTimelineWidget->SetAutoStopEnabled(false);
-					GameFramesTimelineWidget->Play();
-				}
-
-				uint16& CurrentEnabledFlags = GameFramesTimelineWidget->GetMutableElementEnabledFlagsRef();
-				CurrentEnabledFlags = CurrentEnabledFlags &~ PlaybackElementDisabledDuringLiveSession;
-			}
-			else
-			{
-				const FGuid CurrentPlaybackInstigatorID = ControllerSharedPtr->GetPlaybackInstigatorWithExclusiveControlsID();
-				const bool bUserCanControlPlayback = CurrentPlaybackInstigatorID == InvalidGuid || CurrentPlaybackInstigatorID == GetInstigatorID();
-
-				// When it is not a live session, the Game Frames timeline follows the same rule as other timelines. The controls are locked unless we are who started a Play action
-				GameFramesTimelineWidget->SetIsLocked(!bUserCanControlPlayback);
-
-				GameFramesTimelineWidget->SetAutoStopEnabled(true);
-			}
-		}
-	}
-	else
-	{
-		GameFramesTimelineWidget->UpdateMinMaxValue(0, 0);
-		GameFramesTimelineWidget->ResetTimeline();
-	}
-
 	PlaybackViewportClient->bNeedsRedraw = true;
 }
 
-void SChaosVDPlaybackViewport::HandleControllerTrackFrameUpdated(TWeakPtr<FChaosVDPlaybackController> InController, const FChaosVDTrackInfo* UpdatedTrackInfo, FGuid InstigatorGuid)
-{
-	if (const TSharedPtr<FChaosVDPlaybackController> ControllerSharedPtr = InController.Pin())
-	{
-		// The frame number we receive could be from a Solver track, so make sure it is converted to the correct game track frame number
-		if (FChaosVDTrackInfo* GameTrackInfo = ControllerSharedPtr->GetMutableTrackInfo(EChaosVDTrackType::Game, FChaosVDPlaybackController::GameTrackID))
-		{
-			const int32 GameTrackFrame = ControllerSharedPtr->ConvertCurrentFrameToOtherTrackFrame(UpdatedTrackInfo, GameTrackInfo);
-
-			// Something other than us advanced the game track, so make sure the timeline widget is updated
-			if (InstigatorGuid != GetInstigatorID())
-			{
-				GameFramesTimelineWidget->SetCurrentTimelineFrame(GameTrackFrame, EChaosVDSetTimelineFrameFlags::None);
-				GameTrackInfo->CurrentFrame = GameTrackFrame;
-			}
-
-			GameFramesTimelineWidget->SetTargetFrameTime(ControllerSharedPtr->GetFrameTimeForTrack(EChaosVDTrackType::Game, FChaosVDPlaybackController::GameTrackID, *GameTrackInfo));
-		}
-	}
-}
-
 void SChaosVDPlaybackViewport::HandlePostSelectionChange(const UTypedElementSelectionSet* ChangesSelectionSet)
 {
 	PlaybackViewportClient->bNeedsRedraw = true;
@@ -333,57 +268,43 @@ void SChaosVDPlaybackViewport::OnSolverVisibilityUpdated(int32 SolverID, bool bN
 	PlaybackViewportClient->bNeedsRedraw = true;	
 }
 
-void SChaosVDPlaybackViewport::RegisterNewController(TWeakPtr<FChaosVDPlaybackController> NewController)
-{
-	if (PlaybackController != NewController)
+void SChaosVDPlaybackViewport::BindToSceneUpdateEvents()
+{	
+	if (const TSharedPtr<FChaosVDPlaybackController> PlaybackControllerPtr = PlaybackController.Pin())
 	{
-		if (const TSharedPtr<FChaosVDPlaybackController> CurrentPlaybackControllerPtr = PlaybackController.Pin())
-		{
-			if (TSharedPtr<FChaosVDScene> ScenePtr = CurrentPlaybackControllerPtr->GetControllerScene().Pin())
-			{
-				ScenePtr->OnSceneUpdated().RemoveAll(this);
-				ScenePtr->OnSolverVisibilityUpdated().RemoveAll(this);
-			}
-		}
-
-		FChaosVDPlaybackControllerObserver::RegisterNewController(NewController);
-
-		if (const TSharedPtr<FChaosVDPlaybackController> NewPlaybackControllerPtr = PlaybackController.Pin())
+		if (TSharedPtr<FChaosVDScene> ScenePtr = PlaybackControllerPtr->GetControllerScene().Pin())
 		{
-			if (TSharedPtr<FChaosVDScene> ScenePtr = NewPlaybackControllerPtr->GetControllerScene().Pin())
-			{
-				ScenePtr->OnSceneUpdated().AddRaw(this, &SChaosVDPlaybackViewport::OnPlaybackSceneUpdated);
-				ScenePtr->OnSolverVisibilityUpdated().AddRaw(this, &SChaosVDPlaybackViewport::OnSolverVisibilityUpdated);
-			}
+			ScenePtr->OnSceneUpdated().AddSP(this, &SChaosVDPlaybackViewport::OnPlaybackSceneUpdated);
+			ScenePtr->OnSolverVisibilityUpdated().AddSP(this, &SChaosVDPlaybackViewport::OnSolverVisibilityUpdated);
 		}
 	}
 }
 
-void SChaosVDPlaybackViewport::OnFrameSelectionUpdated(int32 NewFrameIndex) const
+void SChaosVDPlaybackViewport::UnbindFromSceneUpdateEvents()
 {
 	if (const TSharedPtr<FChaosVDPlaybackController> PlaybackControllerPtr = PlaybackController.Pin())
 	{
-		constexpr int32 StepNumber = 0;
-		PlaybackControllerPtr->GoToTrackFrame(GetInstigatorID(), EChaosVDTrackType::Game, FChaosVDPlaybackController::GameTrackID, NewFrameIndex, StepNumber);
-
-		PlaybackViewportClient->bNeedsRedraw = true;
+		if (TSharedPtr<FChaosVDScene> ScenePtr = PlaybackControllerPtr->GetControllerScene().Pin())
+		{
+			ScenePtr->OnSceneUpdated().RemoveAll(this);
+			ScenePtr->OnSolverVisibilityUpdated().RemoveAll(this);
+		}
 	}
 }
 
-void SChaosVDPlaybackViewport::HandlePlaybackButtonClicked(EChaosVDPlaybackButtonsID ButtonID)
+void SChaosVDPlaybackViewport::RegisterNewController(TWeakPtr<FChaosVDPlaybackController> NewController)
 {
-	Chaos::VisualDebugger::HandleUserPlaybackInputControl(ButtonID, *this, PlaybackController);
-}
+	if (PlaybackController != NewController)
+	{
+		UnbindFromSceneUpdateEvents();	
 
-bool SChaosVDPlaybackViewport::CanPlayback() const
-{
-	const TSharedPtr<FChaosVDPlaybackController> PlaybackControllerPtr = PlaybackController.Pin();
+		FChaosVDPlaybackControllerObserver::RegisterNewController(NewController);
 
-	return PlaybackControllerPtr && PlaybackControllerPtr->IsRecordingLoaded();
+		BindToSceneUpdateEvents();	
+	}
 }
 
-
-void SChaosVDPlaybackViewport::HanldeExternalViewportInvalidateRequest()
+void SChaosVDPlaybackViewport::HandleExternalViewportInvalidateRequest()
 {
 	if (PlaybackViewportClient)
 	{
diff --git a/Engine/Plugins/ChaosVD/Source/ChaosVD/Private/Widgets/SChaosVDPlaybackViewport.h b/Engine/Plugins/ChaosVD/Source/ChaosVD/Private/Widgets/SChaosVDPlaybackViewport.h
--- a/Engine/Plugins/ChaosVD/Source/ChaosVD/Private/Widgets/SChaosVDPlaybackViewport.h
+++ b/Engine/Plugins/ChaosVD/Source/ChaosVD/Private/Widgets/SChaosVDPlaybackViewport.h
@@ -9,6 +9,7 @@
 #include "Templates/SharedPointer.h"
 #include "Widgets/SCompoundWidget.h"
 
+class SChaosVDGameFramesPlaybackControls;
 enum class EChaosVDPlaybackButtonsID : uint8;
 class FChaosVDEditorModeTools;
 enum class EChaosVDActorTrackingMode;
@@ -65,21 +66,18 @@ protected:
 
 	virtual void RegisterNewController(TWeakPtr<FChaosVDPlaybackController> NewController )override;
 	virtual void HandlePlaybackControllerDataUpdated(TWeakPtr<FChaosVDPlaybackController> InController) override;
-	virtual void HandleControllerTrackFrameUpdated(TWeakPtr<FChaosVDPlaybackController> InController, const FChaosVDTrackInfo* UpdatedTrackInfo, FGuid InstigatorGuid) override;
 	virtual void HandlePostSelectionChange(const UTypedElementSelectionSet* ChangesSelectionSet) override;
 
 	void OnPlaybackSceneUpdated();
 	void OnSolverVisibilityUpdated(int32 SolverID, bool bNewVisibility);
+	void BindToSceneUpdateEvents();
+	void UnbindFromSceneUpdateEvents();
 
-	void OnFrameSelectionUpdated(int32 NewFrameIndex) const;
+	void HandleExternalViewportInvalidateRequest();
 
-	void HandlePlaybackButtonClicked(EChaosVDPlaybackButtonsID ButtonID);
+	TSharedPtr<FChaosVDTrackInfo> CurrentGameTrackInfo;
 
-	bool CanPlayback() const;
-
-	void HanldeExternalViewportInvalidateRequest();
-
-	TSharedPtr<SChaosVDTimelineWidget> GameFramesTimelineWidget;
+	TSharedPtr<SChaosVDGameFramesPlaybackControls> GameFramesPlaybackControls;
 
 	TSharedPtr<FChaosVDPlaybackViewportClient> PlaybackViewportClient;
 	
@@ -93,4 +91,6 @@ protected:
 	static inline FChaosVDViewportInvalidationRequestHandler ExternalViewportInvalidationRequestHandler = FChaosVDViewportInvalidationRequestHandler();
 	
 	FDelegateHandle ExternalInvalidateHandlerHandle;
+
+	bool bIsPlaying = false;
 };
diff --git a/Engine/Plugins/ChaosVD/Source/ChaosVD/Private/Widgets/SChaosVDSceneQueryDataInspector.cpp b/Engine/Plugins/ChaosVD/Source/ChaosVD/Private/Widgets/SChaosVDSceneQueryDataInspector.cpp
--- a/Engine/Plugins/ChaosVD/Source/ChaosVD/Private/Widgets/SChaosVDSceneQueryDataInspector.cpp
+++ b/Engine/Plugins/ChaosVD/Source/ChaosVD/Private/Widgets/SChaosVDSceneQueryDataInspector.cpp
@@ -226,10 +226,13 @@ TSharedRef<SWidget> SChaosVDSceneQueryDataInspector::GenerateVisitStepControls()
 				+SVerticalBox::Slot()
 				[
 					SAssignNew(QueryStepsTimelineWidget, SChaosVDTimelineWidget)
-					.ButtonVisibilityFlags(static_cast<uint16>(EChaosVDTimelineElementIDFlags::AllManualStepping))
+					.ButtonVisibilityFlags(EChaosVDTimelineElementIDFlags::AllManualStepping)
 					.IsEnabled_Raw(this, &SChaosVDSceneQueryDataInspector::GetSQVisitStepsEnabled)
 					.OnFrameChanged_Raw(this, &SChaosVDSceneQueryDataInspector::HandleQueryStepSelectionUpdated)
-					.MaxFrames(0)	
+					.OnButtonClicked(this, &SChaosVDSceneQueryDataInspector::HandleSQVisitTimelineInput)
+					.MinFrames_Raw(this, &SChaosVDSceneQueryDataInspector::GetCurrentMinSQVisitIndex)
+					.MaxFrames_Raw(this, &SChaosVDSceneQueryDataInspector::GetCurrentMaxSQVisitIndex)
+					.CurrentFrame_Raw(this, &SChaosVDSceneQueryDataInspector::GetCurrentSQVisitIndex)
 				]
 			];
 }
@@ -295,13 +298,6 @@ void SChaosVDSceneQueryDataInspector::SetQueryDataToInspect(const FChaosVDSceneQ
 			QueryDataToInspect->CurrentVisitIndex = InQueryDataSelectionHandle.GetSQVisitIndex();
 		}
 
-		if (QueryStepsTimelineWidget)
-		{
-			const int32 SQVisitsNum = QueryDataToInspect->SQVisitData.Num();
-			QueryStepsTimelineWidget->UpdateMinMaxValue(0, SQVisitsNum > 0 ? SQVisitsNum -1 : 0);
-			QueryStepsTimelineWidget->SetCurrentTimelineFrame(QueryDataToInspect->CurrentVisitIndex);
-		}
-
 		if (QueryDataToInspect->SubQueriesIDs.Num() > 0)
 		{
 			TArray<TSharedPtr<FName>> NewSubQueryNameList;
@@ -336,6 +332,13 @@ void SChaosVDSceneQueryDataInspector::HandleQueryStepSelectionUpdated(int32 NewS
 		return;
 	}
 
+	const TSharedPtr<FChaosVDQueryDataWrapper> QueryDataBeingInspected = CurrentSceneQueryBeingInspectedHandle.GetQueryData().Pin();
+	if (!QueryDataBeingInspected)
+	{
+		ClearInspector();
+		return;
+	}
+
 	const TSharedPtr<FChaosVDScene> ScenePtr = SceneWeakPtr.Pin();
 	if (!ScenePtr)
 	{
@@ -350,21 +353,20 @@ void SChaosVDSceneQueryDataInspector::HandleQueryStepSelectionUpdated(int32 NewS
 		return;
 	}
 
-	if (const TSharedPtr<FChaosVDQueryDataWrapper> QueryDataBeingInspected = CurrentSceneQueryBeingInspectedHandle.GetQueryData().Pin())
+	// If we reach this point no need to clear the inspector, we have valid data. This is can be caused by the timeline widget going over as we no longer restrict th button actions
+	if (!QueryDataBeingInspected->SQVisitData.IsValidIndex(NewStepIndex))
 	{
-		const FChaosVDSceneQuerySelectionHandle NewSelection(CurrentSceneQueryBeingInspectedHandle.GetQueryData(), NewStepIndex);
-		QueryDataBeingInspected->CurrentVisitIndex = NewStepIndex;
-		
+		UE_LOG(LogChaosVDEditor, VeryVerbose, TEXT("[%s] Attempted to process and invalid SQ Visit index | Input Index [%d] | Available SQ Visit Data Num [%d]"), ANSI_TO_TCHAR(__FUNCTION__), NewStepIndex, QueryDataBeingInspected->SQVisitData.Num())
+		return;
+	}
 
-		FScopedSQInspectorSilencedSelectionEvents IgnoreSelectionEventsScope(*this);
-		SQDataComponent->SelectQuery(NewSelection);
+	const FChaosVDSceneQuerySelectionHandle NewSelection(CurrentSceneQueryBeingInspectedHandle.GetQueryData(), NewStepIndex);
+	QueryDataBeingInspected->CurrentVisitIndex = NewStepIndex;
 
-		SetQueryDataToInspect(NewSelection);
-	}
-	else
-	{
-		SetQueryDataToInspect(FChaosVDSceneQuerySelectionHandle());
-	}
+	FScopedSQInspectorSilencedSelectionEvents IgnoreSelectionEventsScope(*this);
+	SQDataComponent->SelectQuery(NewSelection);
+
+	SetQueryDataToInspect(NewSelection);
 
 	if (const TSharedPtr<FEditorModeTools> EditorModeToolsPtr = EditorModeToolsWeakPtr.Pin())
 	{
@@ -534,11 +536,6 @@ void SChaosVDSceneQueryDataInspector::ClearInspector()
 		SceneQueryHitDataDetailsView->SetStructureData(nullptr);
 	}
 
-	if (QueryStepsTimelineWidget)
-	{
-		QueryStepsTimelineWidget->UpdateMinMaxValue(0,0);
-	}
-
 	if (SubQueryNamePickerWidget)
 	{
 		SubQueryNamePickerWidget->UpdateNameList({});
@@ -610,10 +607,57 @@ bool SChaosVDSceneQueryDataInspector::GetSQVisitStepsEnabled() const
 	return QueryDataBeingInspected && QueryDataBeingInspected->SQVisitData.Num() > 0;
 }
 
-TSharedPtr<FChaosVDQueryDataWrapper> SChaosVDSceneQueryDataInspector::GetCurrentDataBeingInspected()
+TSharedPtr<FChaosVDQueryDataWrapper> SChaosVDSceneQueryDataInspector::GetCurrentDataBeingInspected() const
 {
 	const TSharedPtr<FChaosVDQueryDataWrapper> QueryDataBeingInspected = CurrentSceneQueryBeingInspectedHandle.GetQueryData().Pin();
 	return QueryDataBeingInspected;
 }
 
+int32 SChaosVDSceneQueryDataInspector::GetCurrentMinSQVisitIndex() const
+{
+	return 0;
+}
+
+int32 SChaosVDSceneQueryDataInspector::GetCurrentMaxSQVisitIndex() const
+{
+	if (const TSharedPtr<FChaosVDQueryDataWrapper> QueryDataBeingInspected = GetCurrentDataBeingInspected())
+	{
+		const int32 SQVisitsNum = QueryDataBeingInspected->SQVisitData.Num();
+		return SQVisitsNum > 0 ? SQVisitsNum -1 : 0;
+	}
+
+	return 0;
+}
+
+int32 SChaosVDSceneQueryDataInspector::GetCurrentSQVisitIndex() const
+{
+	const TSharedPtr<FChaosVDQueryDataWrapper> QueryDataBeingInspected = GetCurrentDataBeingInspected();
+	return  QueryDataBeingInspected ? QueryDataBeingInspected->CurrentVisitIndex : 0;
+}
+
+void SChaosVDSceneQueryDataInspector::HandleSQVisitTimelineInput(EChaosVDPlaybackButtonsID InputID)
+{
+	if (const TSharedPtr<FChaosVDQueryDataWrapper> QueryDataBeingInspected = GetCurrentDataBeingInspected())
+	{
+		switch (InputID)
+		{
+		case EChaosVDPlaybackButtonsID::Next:
+			{
+				HandleQueryStepSelectionUpdated(QueryDataBeingInspected->CurrentVisitIndex +1);
+				break;	
+			}
+		case EChaosVDPlaybackButtonsID::Prev:
+			{
+				HandleQueryStepSelectionUpdated(QueryDataBeingInspected->CurrentVisitIndex -1);
+				break;	
+			}
+		case EChaosVDPlaybackButtonsID::Play:
+		case EChaosVDPlaybackButtonsID::Pause:
+		case EChaosVDPlaybackButtonsID::Stop:
+		default:
+			break;
+		}
+	}
+}
+
 #undef LOCTEXT_NAMESPACE
diff --git a/Engine/Plugins/ChaosVD/Source/ChaosVD/Private/Widgets/SChaosVDSceneQueryDataInspector.h b/Engine/Plugins/ChaosVD/Source/ChaosVD/Private/Widgets/SChaosVDSceneQueryDataInspector.h
--- a/Engine/Plugins/ChaosVD/Source/ChaosVD/Private/Widgets/SChaosVDSceneQueryDataInspector.h
+++ b/Engine/Plugins/ChaosVD/Source/ChaosVD/Private/Widgets/SChaosVDSceneQueryDataInspector.h
@@ -5,6 +5,8 @@
 #include "Components/ChaosVDSceneQueryDataComponent.h"
 #include "Widgets/SCompoundWidget.h"
 
+enum class EChaosVDPlaybackButtonsID : uint8;
+class SChaosVDTimelineWidget;
 struct FChaosVDSceneQuerySelectionHandle;
 class SChaosVDNameListPicker;
 class FEditorModeTools;
@@ -70,7 +72,7 @@ protected:
 	bool GetSelectParticleHitStateEnable() const;
 	bool GetSQVisitStepsEnabled() const;
 
-	TSharedPtr<FChaosVDQueryDataWrapper> GetCurrentDataBeingInspected();
+	TSharedPtr<FChaosVDQueryDataWrapper> GetCurrentDataBeingInspected() const;
 
 	TSharedPtr<SChaosVDTimelineWidget> QueryStepsTimelineWidget;
 	
@@ -92,6 +94,12 @@ protected:
 
 	bool bListenToSelectionEvents = true;
 
+	int32 GetCurrentMinSQVisitIndex() const;
+	int32 GetCurrentMaxSQVisitIndex() const;
+	int32 GetCurrentSQVisitIndex() const;
+
+	void HandleSQVisitTimelineInput(EChaosVDPlaybackButtonsID InputID);
+
 	friend struct FScopedSQInspectorSilencedSelectionEvents;
 };
 
diff --git a/Engine/Plugins/ChaosVD/Source/ChaosVD/Private/Widgets/SChaosVDSolverPlaybackControls.cpp b/Engine/Plugins/ChaosVD/Source/ChaosVD/Private/Widgets/SChaosVDSolverPlaybackControls.cpp
--- a/Engine/Plugins/ChaosVD/Source/ChaosVD/Private/Widgets/SChaosVDSolverPlaybackControls.cpp
+++ b/Engine/Plugins/ChaosVD/Source/ChaosVD/Private/Widgets/SChaosVDSolverPlaybackControls.cpp
@@ -4,26 +4,30 @@
 
 #include "ChaosVDPlaybackController.h"
 #include "ChaosVDScene.h"
-#include "Widgets/ChaosVDPlaybackControlsHelper.h"
+#include "ChaosVDStyle.h"
 #include "Widgets/SBoxPanel.h"
 #include "Widgets/SChaosVDPlaybackViewport.h"
-#include "Widgets/SChaosVDTimelineWidget.h"
 #include "Widgets/Images/SImage.h"
 #include "Widgets/Input/SButton.h"
 #include "Widgets/Text/STextBlock.h"
 
 #define LOCTEXT_NAMESPACE "ChaosVisualDebugger"
 
-void SChaosVDSolverPlaybackControls::Construct(const FArguments& InArgs, int32 InSolverID, const TWeakPtr<FChaosVDPlaybackController>& InPlaybackController)
+void SChaosVDSolverPlaybackControls::Construct(const FArguments& InArgs, const TSharedRef<const FChaosVDTrackInfo>& InSolverTrackInfo, const TWeakPtr<FChaosVDPlaybackController>& InPlaybackController)
 {
-	SolverID = InSolverID;
+	SolverTrackInfoRef = InSolverTrackInfo;
 
 	static const FName NAME_VisibleNotHoveredBrush = TEXT("Level.VisibleIcon16x");
 	static const FName NAME_NotVisibleNotHoveredBrush = TEXT("Level.NotVisibleIcon16x");
+	static const FName NAME_TrackSyncEnabledBrush = TEXT("LinkedIcon");
+	static const FName NAME_TrackSyncDisabledBrush = TEXT("UnLinkedIcon");
 
 	SolverVisibleIconBrush = FAppStyle::Get().GetBrush(NAME_VisibleNotHoveredBrush);
 	SolverHiddenIconBrush = FAppStyle::Get().GetBrush(NAME_NotVisibleNotHoveredBrush);
 
+	SolverTrackSyncEnabledBrush = FChaosVDStyle::Get().GetBrush(NAME_TrackSyncEnabledBrush);
+	SolverTrackSyncDisabledBrush = FChaosVDStyle::Get().GetBrush(NAME_TrackSyncDisabledBrush);
+
 	ChildSlot
 	[
 		SNew(SHorizontalBox)
@@ -47,10 +51,14 @@ void SChaosVDSolverPlaybackControls::Construct(const FArguments& InArgs, int32 I
 				.FillWidth(0.9f)
 				[
 					SAssignNew(FramesTimelineWidget, SChaosVDTimelineWidget)
-					.ButtonVisibilityFlags(static_cast<uint16>(EChaosVDTimelineElementIDFlags::AllPlayback))
+					.IsEnabled_Raw(this, &SChaosVDSolverPlaybackControls::CanPlayback)
+					.ButtonVisibilityFlags(EChaosVDTimelineElementIDFlags::AllPlayback)
+					.IsPlaying_Raw(this, &SChaosVDSolverPlaybackControls::IsPlaying)
+					.MinFrames_Raw(this, &SChaosVDSolverPlaybackControls::GetMinFrames)
+					.MaxFrames_Raw(this, &SChaosVDSolverPlaybackControls::GetMaxFrames)
+					.CurrentFrame_Raw(this, &SChaosVDSolverPlaybackControls::GetCurrentFrame)
 					.OnFrameChanged_Raw(this, &SChaosVDSolverPlaybackControls::OnFrameSelectionUpdated)
-					.OnButtonClicked_Raw(this, &SChaosVDSolverPlaybackControls::HandlePlaybackButtonClicked)
-					.MaxFrames(0)
+					.OnButtonClicked_Raw(this, &SChaosVDSolverPlaybackControls::HandleFramePlaybackButtonClicked)
 				]
 				+SHorizontalBox::Slot()
 				.Padding(6.0f,0.0f)
@@ -67,7 +75,7 @@ void SChaosVDSolverPlaybackControls::Construct(const FArguments& InArgs, int32 I
 						[
 							SNew(STextBlock)
 							.Justification(ETextJustify::Center)
-							.Text_Lambda([this]()->FText{ return bIsReSimFrame ? LOCTEXT("PlaybackViewportWidgetPhysicsFramesResimLabel", "ReSim" ) : LOCTEXT("PlaybackViewportWidgetPhysicsFramesNormalLabel", "Normal" );})
+							.Text_Lambda([this]()->FText{ return SolverTrackInfoRef->bIsReSimulated ? LOCTEXT("PlaybackViewportWidgetPhysicsFramesResimLabel", "ReSim" ) : LOCTEXT("PlaybackViewportWidgetPhysicsFramesNormalLabel", "Normal" );})
 						]
 					]
 				]
@@ -83,7 +91,7 @@ void SChaosVDSolverPlaybackControls::Construct(const FArguments& InArgs, int32 I
 			[
 				SNew(STextBlock)
 				.Justification(ETextJustify::Center)
-				.Text_Lambda([this]()->FText{ return FText::Format(LOCTEXT("PlaybackViewportWidgetStepsLabel","Solver Stage: {0}"), FText::AsCultureInvariant(CurrentStepName));})
+				.Text_Lambda([this]()->FText{ return FText::Format(LOCTEXT("PlaybackViewportWidgetStepsLabel","Solver Stage: {0}"), FText::FromStringView(GetCurrentSolverStageName()));})
 			]
 			+SVerticalBox::Slot()
 			[
@@ -92,45 +100,35 @@ void SChaosVDSolverPlaybackControls::Construct(const FArguments& InArgs, int32 I
 				.FillWidth(0.9f)
 				[
 					SAssignNew(StepsTimelineWidget, SChaosVDTimelineWidget)
-					.ButtonVisibilityFlags(static_cast<uint16>(EChaosVDTimelineElementIDFlags::AllManualStepping))
-					.OnFrameLockStateChanged_Raw(this, &SChaosVDSolverPlaybackControls::HandleLockStateChanged)
-					.OnFrameChanged_Raw(this, &SChaosVDSolverPlaybackControls::OnStepSelectionUpdated)
-					.MaxFrames(0)
+					.IsEnabled_Raw(this, &SChaosVDSolverPlaybackControls::CanPlayback)
+					.ButtonVisibilityFlags(EChaosVDTimelineElementIDFlags::AllManualStepping)
+					.OnFrameChanged_Raw(this, &SChaosVDSolverPlaybackControls::OnSolverStageSelectionUpdated)
+					.MaxFrames_Raw(this, &SChaosVDSolverPlaybackControls::GetMaxSolverStage)
+					.MinFrames_Raw(this, &SChaosVDSolverPlaybackControls::GetMinSolverStage)
+					.CurrentFrame_Raw(this, &SChaosVDSolverPlaybackControls::GetCurrentSolverStage)
+					.OnButtonClicked_Raw(this, &SChaosVDSolverPlaybackControls::HandleSolverStagePlaybackButtonClicked)
 				]
 				+SHorizontalBox::Slot()
 				.AutoWidth()
 				[
 					CreateVisibilityWidget().ToSharedRef()
 				]
+				+SHorizontalBox::Slot()
+				.AutoWidth()
+				[
+					CreateSyncLinkWidget().ToSharedRef()
+				]
 			]
 		]
 	];
 
 	RegisterNewController(InPlaybackController);
-
-	if (const TSharedPtr<FChaosVDPlaybackController> CurrentPlaybackControllerPtr = InPlaybackController.Pin())
-	{
-		if (const FChaosVDTrackInfo* SolverTrackInfo = CurrentPlaybackControllerPtr->GetTrackInfo(EChaosVDTrackType::Solver, SolverID))
-		{
-			HandleControllerTrackFrameUpdated(InPlaybackController, SolverTrackInfo, InvalidGuid);
-		}
-	}
 }
 
 SChaosVDSolverPlaybackControls::~SChaosVDSolverPlaybackControls()
 {
 	if (const TSharedPtr<FChaosVDPlaybackController> PlaybackControllerPtr = PlaybackController.Pin())
-	{
-		// TODO: Doing this here to try to keep this change small enough for a 5.4.1 hotfix release
-		// This should be handled by the playback controller, but currently we are not doing any tracking of what solver tracks are being unloaded/loaded there
-		// This can be done as part of UE-192940 , which is a task to remove a similar dependency in the UI for the playback logic
-		if (FChaosVDTrackInfo* SolverTrackInfo = PlaybackControllerPtr->GetMutableTrackInfo(EChaosVDTrackType::Solver, SolverID))
-		{
-			SolverTrackInfo->CurrentFrame = 0;
-		}
-
-		PlaybackControllerPtr->ReleaseExclusivePlaybackControls(*this);
-		
+	{	
 		if (TSharedPtr<FChaosVDScene> Scene = PlaybackControllerPtr->GetControllerScene().Pin())
 		{
 			Scene->OnSolverVisibilityUpdated().RemoveAll(this);
@@ -138,19 +136,28 @@ SChaosVDSolverPlaybackControls::~SChaosVDSolverPlaybackControls()
 	}
 }
 
-void SChaosVDSolverPlaybackControls::ConditionallyLockPlaybackControl(const TSharedRef<FChaosVDPlaybackController>& InControllerSharedRef)
+bool SChaosVDSolverPlaybackControls::CanPlayback() const
 {
-	const FGuid CurrentPlaybackInstigatorID = InControllerSharedRef->GetPlaybackInstigatorWithExclusiveControlsID();
-	const bool bUserCanControlPlayback = CurrentPlaybackInstigatorID == InvalidGuid || CurrentPlaybackInstigatorID == GetInstigatorID();
+	bool bCanControlPlayback = false;
 
-	// On Live Sessions, only the Game Frames timeline controls are allowed for now
-	FramesTimelineWidget->SetIsLocked(InControllerSharedRef->IsPlayingLiveSession() || !bUserCanControlPlayback);
-	StepsTimelineWidget->SetIsLocked(InControllerSharedRef->IsPlayingLiveSession() || !bUserCanControlPlayback);
+	if (const TSharedPtr<FChaosVDPlaybackController> CurrentPlaybackControllerPtr = PlaybackController.Pin())
+	{
+		bCanControlPlayback = !CurrentPlaybackControllerPtr->IsPlayingLiveSession();
+
+		// When it is not a live session, the Game Frames timeline follows the same rule as other timelines. The controls are locked unless we are who started a Play action
+		if (bCanControlPlayback)
+		{
+			TSharedPtr<FChaosVDTrackInfo> CurrentTrackBeingPlayed = CurrentPlaybackControllerPtr->GetCurrentPlayingTrackInfo();
+			bCanControlPlayback = !CurrentTrackBeingPlayed || (CurrentTrackBeingPlayed && CurrentTrackBeingPlayed->TrackType == EChaosVDTrackType::Solver && CurrentTrackBeingPlayed->TrackID == SolverTrackInfoRef->TrackID);
+		}
+	}
+
+	return bCanControlPlayback;
 }
 
 void SChaosVDSolverPlaybackControls::HandleSolverVisibilityChanged(int32 InSolverID, bool bNewVisibility)
 {
-	if (SolverID != InSolverID)
+	if (SolverTrackInfoRef->TrackID != InSolverID)
 	{
 		return;
 	}
@@ -162,126 +169,60 @@ FReply SChaosVDSolverPlaybackControls::ToggleSolverVisibility() const
 {
 	if (const TSharedPtr<FChaosVDPlaybackController> CurrentPlaybackControllerPtr = PlaybackController.Pin())
 	{
-		CurrentPlaybackControllerPtr->UpdateTrackVisibility(EChaosVDTrackType::Solver, SolverID, !bIsVisible);
+		CurrentPlaybackControllerPtr->UpdateTrackVisibility(EChaosVDTrackType::Solver, SolverTrackInfoRef->TrackID, !bIsVisible);
 	}
 
 	return FReply::Handled();
 }
 
+FReply SChaosVDSolverPlaybackControls::ToggleSolverSyncLink() const
+{
+	if (const TSharedPtr<FChaosVDPlaybackController> CurrentPlaybackControllerPtr = PlaybackController.Pin())
+	{
+		CurrentPlaybackControllerPtr->ToggleTrackSyncEnabled(SolverTrackInfoRef);
+	}
+	
+	return FReply::Handled();
+}
+
 const FSlateBrush* SChaosVDSolverPlaybackControls::GetBrushForCurrentVisibility() const
 {
 	return bIsVisible ? SolverVisibleIconBrush : SolverHiddenIconBrush;
 }
 
-void SChaosVDSolverPlaybackControls::HandlePlaybackControllerDataUpdated(TWeakPtr<FChaosVDPlaybackController> InController)
+const FSlateBrush* SChaosVDSolverPlaybackControls::GetBrushForCurrentLinkState() const
 {
-	if (PlaybackController != InController)
-	{
-		RegisterNewController(InController);
-	}
-
-	const TSharedPtr<FChaosVDPlaybackController> ControllerSharedPtr = PlaybackController.Pin();
-	if (ControllerSharedPtr.IsValid() && ControllerSharedPtr->IsRecordingLoaded())
-	{	
-		const int32 AvailableFrames = ControllerSharedPtr->GetTrackFramesNumber(EChaosVDTrackType::Solver, SolverID);
-
-		int32 AvailableSteps = INDEX_NONE;
-		if (const FChaosVDStepsContainer* StepData = ControllerSharedPtr->GetTrackStepsDataAtFrame_AssumesLocked(EChaosVDTrackType::Solver, SolverID, ControllerSharedPtr->GetTrackCurrentFrame(EChaosVDTrackType::Solver, SolverID)))
-		{
-			AvailableSteps = StepData->Num() > 0 ? StepData->Num() : INDEX_NONE;
-			CurrentStepName =  StepData->Num() > 0 ? (*StepData)[0].StepName : TEXT("NONE");
-		}
-
-		// Max is inclusive and we use this to request as the index on the recorded frames/steps arrays so we need to -1 to the available frames/steps
-		FramesTimelineWidget->UpdateMinMaxValue(0, AvailableFrames != INDEX_NONE ? AvailableFrames -1 : 0);
-
-		//TODO: This will show steps 0/0 if only one step is recorded, we need to add a way to override that functionality
-		// or just set the slider to start from 1 and handle the offset later 
-		StepsTimelineWidget->UpdateMinMaxValue(0,AvailableSteps != INDEX_NONE ? AvailableSteps -1 : 0);
-
-		ConditionallyLockPlaybackControl(ControllerSharedPtr.ToSharedRef());
-	}
-	else
-	{
-		FramesTimelineWidget->UpdateMinMaxValue(0, 0);
-		FramesTimelineWidget->ResetTimeline();
-		StepsTimelineWidget->UpdateMinMaxValue(0,0);
-		StepsTimelineWidget->ResetTimeline();
-	}
+	return SolverTrackInfoRef->bTrackSyncEnabled ? SolverTrackSyncEnabledBrush : SolverTrackSyncDisabledBrush;
 }
 
-void SChaosVDSolverPlaybackControls::UpdateStepsWidgetForFrame(const FChaosVDPlaybackController& InCurrentPlaybackController, int32 FrameNumber, int32 StepNumber, EChaosVDStepsWidgetUpdateFlags OptionsFlags)
+FStringView SChaosVDSolverPlaybackControls::GetCurrentSolverStageName() const
 {
-	if (const FChaosVDStepsContainer* StepsData = InCurrentPlaybackController.GetTrackStepsDataAtFrame_AssumesLocked(EChaosVDTrackType::Solver, SolverID, FrameNumber))
-	{
-		constexpr TCHAR const* UnknownStepName = TEXT("Unknown");
-
-		if (StepsData->Num() > 0 && StepNumber < StepsData->Num())
-		{
-			const FChaosVDStepData& StepData = (*StepsData)[StepNumber];
-
-			if (EnumHasAnyFlags(OptionsFlags, EChaosVDStepsWidgetUpdateFlags::UpdateText))
-			{
-				CurrentStepName = StepData.StepName;
-			}
+	static TCHAR const* UnknownStepName = TEXT("Unknown");
 
-			if (EnumHasAnyFlags(OptionsFlags, EChaosVDStepsWidgetUpdateFlags::SetTimelineStep))
-			{
-				int32 AvailableSteps = StepsData->Num() > 0 ? StepsData->Num() : INDEX_NONE;
-				StepsTimelineWidget->UpdateMinMaxValue(0,AvailableSteps != INDEX_NONE ? AvailableSteps -1 : 0);
-
-				// On Frame updates, always use Step 0
-				StepsTimelineWidget->SetCurrentTimelineFrame(StepNumber, EChaosVDSetTimelineFrameFlags::None);
-			}
-		}
-		else
-		{
-			StepsTimelineWidget->UpdateMinMaxValue(0,0);
-			CurrentStepName = UnknownStepName;
-		}
+	if (SolverTrackInfoRef->CurrentStageNames.IsValidIndex(SolverTrackInfoRef->CurrentStage))
+	{
+		return SolverTrackInfoRef->CurrentStageNames[SolverTrackInfoRef->CurrentStage];
 	}
+
+	return UnknownStepName;
 }
 
-void SChaosVDSolverPlaybackControls::HandleControllerTrackFrameUpdated(TWeakPtr<FChaosVDPlaybackController> InController, const FChaosVDTrackInfo* UpdatedTrackInfo, FGuid InstigatorGuid)
+void SChaosVDSolverPlaybackControls::HandleFramePlaybackButtonClicked(EChaosVDPlaybackButtonsID ButtonID)
 {
-	if (const TSharedPtr<FChaosVDPlaybackController> CurrentPlaybackControllerPtr = InController.Pin())
+	if (const TSharedPtr<FChaosVDPlaybackController> PlaybackControllerPtr = PlaybackController.Pin())
 	{
-		if (const FChaosVDTrackInfo* SolverTrackInfo = CurrentPlaybackControllerPtr->GetTrackInfo(EChaosVDTrackType::Solver, SolverID))
-		{
-			if (InstigatorGuid != GetInstigatorID())
-			{
-				// No Need to manually update the widget state if the widget it-self instigated the update 
-				FramesTimelineWidget->SetCurrentTimelineFrame(SolverTrackInfo->CurrentFrame, EChaosVDSetTimelineFrameFlags::None);
-				UpdateStepsWidgetForFrame(*CurrentPlaybackControllerPtr.Get(), SolverTrackInfo->CurrentFrame, SolverTrackInfo->CurrentStep);
-			}
-
-			bIsReSimFrame = SolverTrackInfo->bIsReSimulated;
-
-			FramesTimelineWidget->SetTargetFrameTime(CurrentPlaybackControllerPtr->GetFrameTimeForTrack(EChaosVDTrackType::Solver, SolverID, *SolverTrackInfo));
-		}
+		PlaybackControllerPtr->HandleFramePlaybackControlInput(ButtonID, SolverTrackInfoRef, GetInstigatorID());
 	}
 }
 
-void SChaosVDSolverPlaybackControls::HandleLockStateChanged(bool NewIsLocked)
+void SChaosVDSolverPlaybackControls::HandleSolverStagePlaybackButtonClicked(EChaosVDPlaybackButtonsID ButtonID)
 {
-	if (const TSharedPtr<FChaosVDPlaybackController> CurrentPlaybackControllerPtr = PlaybackController.Pin())
+	if (const TSharedPtr<FChaosVDPlaybackController> PlaybackControllerPtr = PlaybackController.Pin())
 	{
-		if (NewIsLocked)
-		{
-			CurrentPlaybackControllerPtr->LockTrackInCurrentStep(EChaosVDTrackType::Solver, SolverID);
-		}
-		else
-		{
-			CurrentPlaybackControllerPtr->UnlockTrackStep(EChaosVDTrackType::Solver, SolverID);
-		}
+		PlaybackControllerPtr->HandleFrameStagePlaybackControlInput(ButtonID, SolverTrackInfoRef, GetInstigatorID());
 	}
 }
 
-void SChaosVDSolverPlaybackControls::HandlePlaybackButtonClicked(EChaosVDPlaybackButtonsID ButtonID)
-{
-	Chaos::VisualDebugger::HandleUserPlaybackInputControl(ButtonID, *this, PlaybackController);
-}
-
 const FSlateBrush* SChaosVDSolverPlaybackControls::GetFrameTypeBadgeBrush() const
 {
 	const FButtonStyle& ButtonStyle = FAppStyle::Get().GetWidgetStyle<FButtonStyle>("Menu.Button");
@@ -292,6 +233,7 @@ TSharedPtr<SWidget> SChaosVDSolverPlaybackControls::CreateVisibilityWidget()
 {
 	return SNew(SButton)
 			.OnClicked_Raw(this, &SChaosVDSolverPlaybackControls::ToggleSolverVisibility)
+			.ToolTipText_Raw(this, &SChaosVDSolverPlaybackControls::GetVisibilityButtonToolTipText)
 			[
 				SNew(SImage)
 				.Image_Raw(this,&SChaosVDSolverPlaybackControls::GetBrushForCurrentVisibility)
@@ -300,33 +242,80 @@ TSharedPtr<SWidget> SChaosVDSolverPlaybackControls::CreateVisibilityWidget()
 			];	
 }
 
-void SChaosVDSolverPlaybackControls::OnFrameSelectionUpdated(int32 NewFrameIndex)
+TSharedPtr<SWidget> SChaosVDSolverPlaybackControls::CreateSyncLinkWidget()
 {
-	if (const TSharedPtr<FChaosVDPlaybackController> PlaybackControllerPtr = PlaybackController.Pin())
-	{
-		const int32 LastStepNumber = PlaybackControllerPtr->GetTrackLastStepAtFrame(EChaosVDTrackType::Solver, SolverID, NewFrameIndex);
-		const int32 CurrentStep = PlaybackControllerPtr->GetTrackCurrentStep(EChaosVDTrackType::Solver, SolverID);
+	return SNew(SButton)
+		.OnClicked_Raw(this, &SChaosVDSolverPlaybackControls::ToggleSolverSyncLink)
+		.ToolTipText_Raw(this, &SChaosVDSolverPlaybackControls::GetSyncLinkTipText)
+		[
+			SNew(SImage)
+			.Image_Raw(this,&SChaosVDSolverPlaybackControls::GetBrushForCurrentLinkState)
+			.DesiredSizeOverride(FVector2D(16.0f,16.0f))
+			.ColorAndOpacity(FSlateColor::UseForeground())
+		];	
+}
 
-		// If the steps control is unlocked, each time we go to a new frame we should start at step 0.
-		const int32 StepNumber = StepsTimelineWidget->IsUnlocked() ? LastStepNumber : CurrentStep;
+FText SChaosVDSolverPlaybackControls::GetVisibilityButtonToolTipText() const
+{
+	return bIsVisible ? LOCTEXT("HideVisibilityButtonToolTipText", "Click to hide all the visualization data corresponding to this solver track") : LOCTEXT("ShowVisibilityButtonToolTipText", "Click to show all the visualization data corresponding to this solver track");
+}
 
-		UpdateStepsWidgetForFrame(*PlaybackControllerPtr.Get(), NewFrameIndex, StepNumber);
+FText SChaosVDSolverPlaybackControls::GetSyncLinkTipText() const
+{
+	return SolverTrackInfoRef->bTrackSyncEnabled ? LOCTEXT("DisableSyncLinkToolTipText", "Click to disable track syncing so this timeline can be player independently") : LOCTEXT("EnableSyncLinkToolTipText", "Click to eanble track syncing so this will be played in sync with other tracks");
+}
 
-		PlaybackControllerPtr->GoToTrackFrame(GetInstigatorID(), EChaosVDTrackType::Solver, SolverID, NewFrameIndex, StepNumber);
+bool SChaosVDSolverPlaybackControls::IsPlaying() const
+{
+	return SolverTrackInfoRef->bIsPlaying;
+}
+
+int32 SChaosVDSolverPlaybackControls::GetCurrentFrame() const
+{
+	return SolverTrackInfoRef->CurrentFrame;
+}
+
+int32 SChaosVDSolverPlaybackControls::GetMinFrames() const
+{
+	return 0;
+}
+
+int32 SChaosVDSolverPlaybackControls::GetMaxFrames() const
+{
+	return SolverTrackInfoRef->MaxFrames - 1;
+}
+
+int32 SChaosVDSolverPlaybackControls::GetCurrentSolverStage() const
+{
+	return SolverTrackInfoRef->CurrentStage;
+}
+
+int32 SChaosVDSolverPlaybackControls::GetMinSolverStage() const
+{
+	return 0;
+}
+
+int32 SChaosVDSolverPlaybackControls::GetMaxSolverStage() const
+{
+	return SolverTrackInfoRef->CurrentStageNames.Num() -1;
+}
+
+void SChaosVDSolverPlaybackControls::OnFrameSelectionUpdated(int32 NewFrameIndex)
+{
+	if (const TSharedPtr<FChaosVDPlaybackController> PlaybackControllerPtr = PlaybackController.Pin())
+	{
+		// By default we always playback frames at the last recorded stage as that represents the end of frame state
+		const int32 LastStepNumber = GetMaxSolverStage();
+		PlaybackControllerPtr->GoToTrackFrameAndSync(GetInstigatorID(), EChaosVDTrackType::Solver, SolverTrackInfoRef->TrackID, NewFrameIndex, LastStepNumber);
 	}
 }
 
-void SChaosVDSolverPlaybackControls::OnStepSelectionUpdated(int32 NewStepIndex)
+void SChaosVDSolverPlaybackControls::OnSolverStageSelectionUpdated(int32 NewStepIndex)
 {
 	if (const TSharedPtr<FChaosVDPlaybackController> PlaybackControllerPtr = PlaybackController.Pin())
 	{
 		// On Steps updates. Always use the current Frame
-		int32 CurrentFrame = PlaybackControllerPtr->GetTrackCurrentFrame(EChaosVDTrackType::Solver, SolverID);
-
-		// Only Update the text as if we are here it means manually set the step number already
-		UpdateStepsWidgetForFrame(*PlaybackControllerPtr.Get(), CurrentFrame, NewStepIndex, EChaosVDStepsWidgetUpdateFlags::UpdateText);
-
-		PlaybackControllerPtr->GoToTrackFrame(GetInstigatorID(), EChaosVDTrackType::Solver, SolverID, CurrentFrame, NewStepIndex);
+		PlaybackControllerPtr->GoToTrackFrame(GetInstigatorID(), EChaosVDTrackType::Solver, SolverTrackInfoRef->TrackID, SolverTrackInfoRef->CurrentFrame, NewStepIndex);
 	}
 }
 
diff --git a/Engine/Plugins/ChaosVD/Source/ChaosVD/Private/Widgets/SChaosVDSolverPlaybackControls.h b/Engine/Plugins/ChaosVD/Source/ChaosVD/Private/Widgets/SChaosVDSolverPlaybackControls.h
--- a/Engine/Plugins/ChaosVD/Source/ChaosVD/Private/Widgets/SChaosVDSolverPlaybackControls.h
+++ b/Engine/Plugins/ChaosVD/Source/ChaosVD/Private/Widgets/SChaosVDSolverPlaybackControls.h
@@ -2,8 +2,10 @@
 
 #pragma once
 
+#include "ChaosVDPlaybackController.h"
 #include "ChaosVDPlaybackControllerInstigator.h"
 #include "ChaosVDPlaybackControllerObserver.h"
+#include "SChaosVDTimelineWidget.h"
 #include "Styling/ISlateStyle.h"
 #include "Templates/SharedPointer.h"
 #include "Widgets/SCompoundWidget.h"
@@ -31,47 +33,70 @@ ENUM_CLASS_FLAGS(EChaosVDStepsWidgetUpdateFlags)
 class SChaosVDSolverPlaybackControls : public SCompoundWidget, public FChaosVDPlaybackControllerObserver, public IChaosVDPlaybackControllerInstigator
 {
 public:
-	SLATE_BEGIN_ARGS( SChaosVDSolverPlaybackControls ){}
+
+	explicit SChaosVDSolverPlaybackControls()
+		: SolverTrackInfoRef(MakeShared<const FChaosVDTrackInfo>())
+	{
+	}
+
+
+	SLATE_BEGIN_ARGS(SChaosVDSolverPlaybackControls) {}
 	SLATE_END_ARGS()
 
-	void Construct(const FArguments& InArgs, int32 InSolverID, const TWeakPtr<FChaosVDPlaybackController>& InPlaybackController);
+	void Construct(const FArguments& InArgs, const TSharedRef<const FChaosVDTrackInfo>& InSolverTrackInfo, const TWeakPtr<FChaosVDPlaybackController>& InPlaybackController);
 
-	~SChaosVDSolverPlaybackControls();
+	virtual ~SChaosVDSolverPlaybackControls() override;
 
 private:
 
 	void OnFrameSelectionUpdated(int32 NewFrameIndex);
-	void OnStepSelectionUpdated(int32 NewStepIndex);
+	void OnSolverStageSelectionUpdated(int32 NewStepIndex);
 
 	virtual void RegisterNewController(TWeakPtr<FChaosVDPlaybackController> NewController) override;
 
-	virtual void HandlePlaybackControllerDataUpdated(TWeakPtr<FChaosVDPlaybackController> InController) override;
-	void UpdateStepsWidgetForFrame(const FChaosVDPlaybackController& InCurrentPlaybackController, int32 FrameNumber, int32 StepNumber, EChaosVDStepsWidgetUpdateFlags OptionsFlags = EChaosVDStepsWidgetUpdateFlags::Default);
-	virtual void HandleControllerTrackFrameUpdated(TWeakPtr<FChaosVDPlaybackController> InController, const FChaosVDTrackInfo* UpdatedTrackInfo, FGuid InstigatorGuid) override;
-
-	void HandleLockStateChanged(bool NewIsLocked);
-	void HandlePlaybackButtonClicked(EChaosVDPlaybackButtonsID ButtonID);
+	FStringView GetCurrentSolverStageName() const;
 
-	void ConditionallyLockPlaybackControl(const TSharedRef<FChaosVDPlaybackController>& InControllerSharedRef);
+	void HandleFramePlaybackButtonClicked(EChaosVDPlaybackButtonsID ButtonID);
+	void HandleSolverStagePlaybackButtonClicked(EChaosVDPlaybackButtonsID ButtonID);
 
 	void HandleSolverVisibilityChanged(int32 InSolverID, bool bNewVisibility);
 
 	FReply ToggleSolverVisibility() const;
+	FReply ToggleSolverSyncLink() const;
 
 	const FSlateBrush* GetBrushForCurrentVisibility() const;
+	const FSlateBrush* GetBrushForCurrentLinkState() const;
 
 	const FSlateBrush* GetFrameTypeBadgeBrush() const;
 
 	TSharedPtr<SWidget> CreateVisibilityWidget();
+	TSharedPtr<SWidget> CreateSyncLinkWidget();
+
+	FText GetVisibilityButtonToolTipText() const;
+	FText GetSyncLinkTipText() const;
+	
+	bool CanPlayback() const;
+
+	bool IsPlaying() const;
+
+	int32 GetCurrentFrame() const;
+	int32 GetMinFrames() const;
+	int32 GetMaxFrames() const;
+	
+	int32 GetCurrentSolverStage() const;
+	int32 GetMinSolverStage() const;
+	int32 GetMaxSolverStage() const;
 
-	int32 SolverID = INDEX_NONE;
+	TSharedRef<const FChaosVDTrackInfo> SolverTrackInfoRef;
 	bool bIsReSimFrame = false;
 	FString CurrentStepName;
 	TSharedPtr<SChaosVDTimelineWidget> FramesTimelineWidget;
 	TSharedPtr<SChaosVDTimelineWidget> StepsTimelineWidget;
-	bool bStepsLocked = false;
 	bool bIsVisible = true;
 
 	const FSlateBrush* SolverVisibleIconBrush = nullptr;
 	const FSlateBrush* SolverHiddenIconBrush = nullptr;
+	
+	const FSlateBrush* SolverTrackSyncEnabledBrush = nullptr;
+	const FSlateBrush* SolverTrackSyncDisabledBrush = nullptr;
 };
diff --git a/Engine/Plugins/ChaosVD/Source/ChaosVD/Private/Widgets/SChaosVDSolverTracks.cpp b/Engine/Plugins/ChaosVD/Source/ChaosVD/Private/Widgets/SChaosVDSolverTracks.cpp
--- a/Engine/Plugins/ChaosVD/Source/ChaosVD/Private/Widgets/SChaosVDSolverTracks.cpp
+++ b/Engine/Plugins/ChaosVD/Source/ChaosVD/Private/Widgets/SChaosVDSolverTracks.cpp
@@ -12,7 +12,7 @@ void SChaosVDSolverTracks::Construct(const FArguments& InArgs, TWeakPtr<FChaosVD
 {
 	ChildSlot
 	[
-		SAssignNew(SolverTracksListWidget, SListView<TSharedPtr<FChaosVDTrackInfo>>)
+		SAssignNew(SolverTracksListWidget, SListView<TSharedPtr<const FChaosVDTrackInfo>>)
 		.ListItemsSource(&CachedTrackInfoArray)
 		.SelectionMode( ESelectionMode::None )
 		.ListViewStyle(&FAppStyle::Get().GetWidgetStyle<FTableViewStyle>("SimpleListView"))
@@ -25,9 +25,9 @@ void SChaosVDSolverTracks::Construct(const FArguments& InArgs, TWeakPtr<FChaosVD
 	
 	if (const TSharedPtr<FChaosVDPlaybackController> CurrentPlaybackControllerPtr = InPlaybackController.Pin())
 	{
-		if (const FChaosVDTrackInfo* GameTrackInfo = CurrentPlaybackControllerPtr->GetTrackInfo(EChaosVDTrackType::Game, FChaosVDPlaybackController::GameTrackID))
+		if (const TSharedPtr<const FChaosVDTrackInfo> GameTrackInfo = CurrentPlaybackControllerPtr->GetTrackInfo(EChaosVDTrackType::Game, FChaosVDPlaybackController::GameTrackID))
 		{
-			HandleControllerTrackFrameUpdated(InPlaybackController, GameTrackInfo, InvalidGuid);
+			HandleControllerTrackFrameUpdated(InPlaybackController, GameTrackInfo.ToSharedRef(), InvalidGuid);
 		}
 	}
 	else
@@ -47,17 +47,19 @@ void SChaosVDSolverTracks::HandlePlaybackControllerDataUpdated(TWeakPtr<FChaosVD
 	{
 		// If the controller data was updated, need to update our cache track info data as it could have been changed.
 		// For example this can happen when we load another recording. We use the GameTrack info for that as it is the one that is always valid
-		const FChaosVDTrackInfo* GameTrackInfo = CurrentPlaybackControllerPtr->GetTrackInfo(EChaosVDTrackType::Game, FChaosVDPlaybackController::GameTrackID);
-		UpdatedCachedTrackInfoData(InPlaybackController, GameTrackInfo);
+		if (const TSharedPtr<const FChaosVDTrackInfo> GameTrackInfo = CurrentPlaybackControllerPtr->GetTrackInfo(EChaosVDTrackType::Game, FChaosVDPlaybackController::GameTrackID))
+		{
+			UpdatedCachedTrackInfoData(InPlaybackController, GameTrackInfo.ToSharedRef());
+		}
 	}
 }
 
-void SChaosVDSolverTracks::UpdatedCachedTrackInfoData(TWeakPtr<FChaosVDPlaybackController> InPlaybackController, const FChaosVDTrackInfo* UpdatedTrackInfo)
+void SChaosVDSolverTracks::UpdatedCachedTrackInfoData(TWeakPtr<FChaosVDPlaybackController> InPlaybackController, const TSharedRef<const FChaosVDTrackInfo>& UpdatedTrackInfo)
 {
 	if (const TSharedPtr<FChaosVDPlaybackController> CurrentPlaybackControllerPtr = InPlaybackController.Pin())
 	{
-		TArray<TSharedPtr<FChaosVDTrackInfo>> TrackInfoArray;
-		CurrentPlaybackControllerPtr->GetAvailableTrackInfosAtTrackFrame(EChaosVDTrackType::Solver, TrackInfoArray, UpdatedTrackInfo);
+		TArray<TSharedPtr<const FChaosVDTrackInfo>> TrackInfoArray;
+		CurrentPlaybackControllerPtr->GetAvailableTrackInfosAtTrackFrame(EChaosVDTrackType::Solver, UpdatedTrackInfo, TrackInfoArray);
 
 		if (TrackInfoArray != CachedTrackInfoArray)
 		{
@@ -72,7 +74,7 @@ void SChaosVDSolverTracks::UpdatedCachedTrackInfoData(TWeakPtr<FChaosVDPlaybackC
 	}
 }
 
-void SChaosVDSolverTracks::HandleControllerTrackFrameUpdated(TWeakPtr<FChaosVDPlaybackController> InPlaybackController, const FChaosVDTrackInfo* UpdatedTrackInfo, FGuid InstigatorGuid)
+void SChaosVDSolverTracks::HandleControllerTrackFrameUpdated(TWeakPtr<FChaosVDPlaybackController> InPlaybackController, TWeakPtr<const FChaosVDTrackInfo> UpdatedTrackInfo, FGuid InstigatorGuid)
 {
 	if (InstigatorGuid == GetInstigatorID())
 	{
@@ -80,57 +82,71 @@ void SChaosVDSolverTracks::HandleControllerTrackFrameUpdated(TWeakPtr<FChaosVDPl
 		return;
 	}
 
-	if (UpdatedTrackInfo == nullptr)
+	TSharedPtr<const FChaosVDTrackInfo> UpdatedTrackInfoPtr = UpdatedTrackInfo.Pin();
+
+	if (!UpdatedTrackInfoPtr)
 	{
 		return;
 	}
 
-	if (UpdatedTrackInfo->TrackType == EChaosVDTrackType::Solver)
+	// Only Game Frame Track Update can change the available solvers
+	if (UpdatedTrackInfoPtr->TrackType == EChaosVDTrackType::Solver)
 	{
 		return;
 	}
 
-	UpdatedCachedTrackInfoData(InPlaybackController, UpdatedTrackInfo);
+	UpdatedCachedTrackInfoData(InPlaybackController, UpdatedTrackInfoPtr.ToSharedRef());
 }
 
-TSharedRef<ITableRow> SChaosVDSolverTracks::MakeSolverTrackControlsFromTrackInfo(TSharedPtr<FChaosVDTrackInfo> TrackInfo, const TSharedRef<STableViewBase>& OwnerTable)
+TSharedRef<ITableRow> SChaosVDSolverTracks::MakeSolverTrackControlsFromTrackInfo(TSharedPtr<const FChaosVDTrackInfo> TrackInfo, const TSharedRef<STableViewBase>& OwnerTable)
 {
-	TSharedPtr<SVerticalBox> PlaybackControlsContainer;
-	SAssignNew(PlaybackControlsContainer, SVerticalBox)
-			+SVerticalBox::Slot()
-			.AutoHeight()
-			.Padding(10.0f, 2.0f, 10.0f, 0.0f)
-			[
-				SNew(SExpandableArea)
-					.InitiallyCollapsed(false)
-					.BorderBackgroundColor(FLinearColor::White)
-					.Padding(FMargin(8.f))
-					.HeaderContent()
-					[
-						SNew(SHorizontalBox)
-						+SHorizontalBox::Slot()
-						.HAlign(HAlign_Left)
-						.VAlign(VAlign_Center)
-						.AutoWidth()
-						.Padding(0.f, 0.f, 0.f, 0.f)
+	TSharedPtr<SWidget> RowWidget;
+	if (TrackInfo)
+	{
+		RowWidget = SNew(SVerticalBox)
+				+SVerticalBox::Slot()
+				.AutoHeight()
+				.Padding(10.0f, 2.0f, 10.0f, 0.0f)
+				[
+					SNew(SExpandableArea)
+						.InitiallyCollapsed(false)
+						.BorderBackgroundColor(FLinearColor::White)
+						.Padding(FMargin(8.f))
+						.HeaderContent()
 						[
-							SNew(STextBlock)
-							.Text(FText::FromString(TrackInfo->TrackName))
-							.Font(FCoreStyle::Get().GetFontStyle("ExpandableArea.TitleFont"))
+							SNew(SHorizontalBox)
+							+SHorizontalBox::Slot()
+							.HAlign(HAlign_Left)
+							.VAlign(VAlign_Center)
+							.AutoWidth()
+							.Padding(0.f, 0.f, 0.f, 0.f)
+							[
+								SNew(STextBlock)
+								.Text(FText::FromName(TrackInfo->TrackName))
+								.Font(FCoreStyle::Get().GetFontStyle("ExpandableArea.TitleFont"))
+							]
 						]
-					]
-					.BodyContent()
-					[
-						SNew(SHorizontalBox)
-						+SHorizontalBox::Slot()
-						.Padding(2.f, 4.f, 2.f, 12.f)
+						.BodyContent()
 						[
-							SNew(SChaosVDSolverPlaybackControls, TrackInfo->TrackID, PlaybackController)
+							SNew(SHorizontalBox)
+							+SHorizontalBox::Slot()
+							.Padding(2.f, 4.f, 2.f, 12.f)
+							[
+								SNew(SChaosVDSolverPlaybackControls, TrackInfo.ToSharedRef(), PlaybackController)
+							]
 						]
-					]
-			];
+				];
+	}
+	else
+	{
+		RowWidget = SNew(SVerticalBox)
+		+SVerticalBox::Slot()
+		[
+			SNew(STextBlock)
+			.Text(NSLOCTEXT("ChaosVisualDebugger", "SolverPlaybackControlsErrorMessage", "Failed to read data for solver."))
+		];
+	}
 
-	TSharedPtr<SWidget> RowWidget = PlaybackControlsContainer;
 	return
 		SNew(STableRow< TSharedPtr<FString> >, OwnerTable)
 		[
diff --git a/Engine/Plugins/ChaosVD/Source/ChaosVD/Private/Widgets/SChaosVDSolverTracks.h b/Engine/Plugins/ChaosVD/Source/ChaosVD/Private/Widgets/SChaosVDSolverTracks.h
--- a/Engine/Plugins/ChaosVD/Source/ChaosVD/Private/Widgets/SChaosVDSolverTracks.h
+++ b/Engine/Plugins/ChaosVD/Source/ChaosVD/Private/Widgets/SChaosVDSolverTracks.h
@@ -25,13 +25,13 @@ public:
 
 private:
 
-	void UpdatedCachedTrackInfoData(TWeakPtr<FChaosVDPlaybackController> InPlaybackController, const FChaosVDTrackInfo* UpdatedTrackInfo);
+	void UpdatedCachedTrackInfoData(TWeakPtr<FChaosVDPlaybackController> InPlaybackController, const TSharedRef<const FChaosVDTrackInfo>& UpdatedTrackInfo);
 	virtual void HandlePlaybackControllerDataUpdated(TWeakPtr<FChaosVDPlaybackController> InPlaybackController) override;
-	virtual void HandleControllerTrackFrameUpdated(TWeakPtr<FChaosVDPlaybackController> InPlaybackController, const FChaosVDTrackInfo* UpdatedTrackInfo, FGuid InstigatorGuid) override;
+	virtual void HandleControllerTrackFrameUpdated(TWeakPtr<FChaosVDPlaybackController> InPlaybackController, TWeakPtr<const FChaosVDTrackInfo> UpdatedTrackInfo, FGuid InstigatorGuid) override;
 
-	TSharedRef<ITableRow> MakeSolverTrackControlsFromTrackInfo(TSharedPtr<FChaosVDTrackInfo> TrackInfo, const TSharedRef<STableViewBase>& OwnerTable);
+	TSharedRef<ITableRow> MakeSolverTrackControlsFromTrackInfo(TSharedPtr<const FChaosVDTrackInfo> TrackInfo, const TSharedRef<STableViewBase>& OwnerTable);
 
-	TSharedPtr<SListView<TSharedPtr<FChaosVDTrackInfo>>> SolverTracksListWidget;
+	TSharedPtr<SListView<TSharedPtr<const FChaosVDTrackInfo>>> SolverTracksListWidget;
 
-	TArray<TSharedPtr<FChaosVDTrackInfo>> CachedTrackInfoArray;
+	TArray<TSharedPtr<const FChaosVDTrackInfo>> CachedTrackInfoArray;
 };
diff --git a/Engine/Plugins/ChaosVD/Source/ChaosVD/Private/Widgets/SChaosVDTimelineWidget.cpp b/Engine/Plugins/ChaosVD/Source/ChaosVD/Private/Widgets/SChaosVDTimelineWidget.cpp
--- a/Engine/Plugins/ChaosVD/Source/ChaosVD/Private/Widgets/SChaosVDTimelineWidget.cpp
+++ b/Engine/Plugins/ChaosVD/Source/ChaosVD/Private/Widgets/SChaosVDTimelineWidget.cpp
@@ -1,6 +1,6 @@
 // Copyright Epic Games, Inc. All Rights Reserved.
 
-#include "Widgets/SChaosVDTimelineWidget.h"
+#include "SChaosVDTimelineWidget.h"
 
 #include "ChaosVDStyle.h"
 #include "Input/Reply.h"
@@ -15,13 +15,25 @@
 void SChaosVDTimelineWidget::Construct(const FArguments& InArgs)
 {
 	MaxFrames = InArgs._MaxFrames;
+	MinFrames = InArgs._MinFrames;
+	CurrentFrame = InArgs._CurrentFrame;
 	FrameChangedDelegate = InArgs._OnFrameChanged;
-	FrameLockedDelegate = InArgs._OnFrameLockStateChanged;
 	ButtonClickedDelegate = InArgs._OnButtonClicked;
-	bAutoStopEnabled = InArgs._AutoStopEnabled;
 	ElementVisibilityFlags = InArgs._ButtonVisibilityFlags;
+	ElementEnabledFlags = InArgs._ButtonEnabledFlags;
+	bIsPlaying = InArgs._IsPlaying;
 
-	SetCanTick(false);
+	if (!ElementVisibilityFlags.IsSet())
+	{
+		ElementVisibilityFlags = EChaosVDTimelineElementIDFlags::AllPlayback;
+	}
+	
+	if (!ElementEnabledFlags.IsSet())
+	{
+		ElementEnabledFlags = EChaosVDTimelineElementIDFlags::AllPlayback;
+	}
+
+	SetCanTick(true);
 
 	ChildSlot
 	[
@@ -101,23 +113,6 @@ void SChaosVDTimelineWidget::Construct(const FArguments& InArgs)
 						.ColorAndOpacity( FSlateColor::UseForeground() )
 					]
 				]
-				+SHorizontalBox::Slot()
-				.AutoWidth()
-				[
-					SNew(SButton)
-					.Visibility_Raw(this, &SChaosVDTimelineWidget::GetElementVisibility, EChaosVDTimelineElementIDFlags::Lock)
-					.IsEnabled_Raw(this, &SChaosVDTimelineWidget::GetElementEnabled, EChaosVDTimelineElementIDFlags::Lock)
-					.OnClicked( FOnClicked::CreateRaw(this, &SChaosVDTimelineWidget::ToggleLockState))
-					.ContentPadding( 2.0f )
-					.ForegroundColor( FSlateColor::UseForeground())
-					.IsFocusable( false )
-					[
-						SNew( SImage )
-						.Image_Raw(this, &SChaosVDTimelineWidget::GetLockStateIcon)
-						.DesiredSizeOverride(FVector2D(16.0f,16.0f))
-						.ColorAndOpacity( FSlateColor::UseForeground() )
-					]
-				]
 			]
 			+SHorizontalBox::Slot()
 			.VAlign(VAlign_Center)
@@ -126,13 +121,11 @@ void SChaosVDTimelineWidget::Construct(const FArguments& InArgs)
 			[
 			  SAssignNew(TimelineSlider, SSlider)
 			  .Visibility_Raw(this, &SChaosVDTimelineWidget::GetElementVisibility, EChaosVDTimelineElementIDFlags::Timeline)
-			  //TODO: Enable locking when we have custom images for the slider elements. Currently locking it messes up with the transparency/images
-			  //.Locked_Raw(this, &SChaosVDTimelineWidget::GetElementEnabled, EChaosVDTimelineElementIDFlags::Timeline)
-			  .ToolTipText_Lambda([this]()-> FText{ return FText::AsNumber(CurrentFrame); })
-			  .Value(CurrentFrame)
+			  .ToolTipText_Lambda([this]()-> FText{ return FText::AsNumber(GetCurrentFrame()); })
+			  .Value_Raw(this, &SChaosVDTimelineWidget::GetCurrentFrameAsFloat)
 			  .OnValueChanged_Raw(this, &SChaosVDTimelineWidget::SetCurrentTimelineFrame, EChaosVDSetTimelineFrameFlags::BroadcastChange)
 			  .StepSize(1)
-			  .MaxValue(MaxFrames)
+			  .MaxValue(0)
 			  .MinValue(0)
 			]
 			+SHorizontalBox::Slot()
@@ -143,205 +136,128 @@ void SChaosVDTimelineWidget::Construct(const FArguments& InArgs)
 			[
 				SNew(STextBlock)
 				.ColorAndOpacity(FLinearColor::White)
-				.Text_Lambda([this]()->FText{ return FText::Format(LOCTEXT("FramesCounter","{0} / {1}"), CurrentFrame, MaxFrames);})
+				.Text_Lambda([this]()->FText{ return FText::Format(LOCTEXT("FramesCounter","{0} / {1}"), GetCurrentFrame(), GetCurrentMaxFrames());})
 			]
 	];
 }
 
-void SChaosVDTimelineWidget::UpdateMinMaxValue(float NewMin, float NewMax)
+FReply SChaosVDTimelineWidget::TogglePlay()
 {
-	if (!TimelineSlider.IsValid())
+	if (IsPlaying())
 	{
-		return;
+		Pause();
 	}
-
-	TimelineSlider->SetMinAndMaxValues(NewMin, NewMax);
-
-	MinFrames = NewMin;
-	MaxFrames = NewMax;
-	
-	if (CurrentFrame < NewMin || CurrentFrame > NewMax)
+	else
 	{
-		CurrentFrame = NewMin;
+		Play();
 	}
+
+	return FReply::Handled();
 }
 
-void SChaosVDTimelineWidget::SetTargetFrameTime(float TargetFrameTimeSeconds)
+void SChaosVDTimelineWidget::Play()
 {
-	if (TargetFrameTimeSeconds > 0)
-	{
-		CurrentPlaybackRate = TargetFrameTimeSeconds; 
-	}
-	else
-	{
-		// Default to 60 FPS
-		CurrentPlaybackRate = 1 / 60.0f;
-	}
+	ButtonClickedDelegate.ExecuteIfBound(EChaosVDPlaybackButtonsID::Play);
 }
 
-void SChaosVDTimelineWidget::ResetTimeline()
+FReply SChaosVDTimelineWidget::Stop()
 {
-	TimelineSlider->SetValue(MinFrames);
-	CurrentFrame = MinFrames;
+	ButtonClickedDelegate.ExecuteIfBound(EChaosVDPlaybackButtonsID::Stop);
+
+	return FReply::Handled();
 }
 
 void SChaosVDTimelineWidget::SetCurrentTimelineFrame(float FrameNumber, EChaosVDSetTimelineFrameFlags Options)
 {
-	CurrentFrame = static_cast<int32>(FrameNumber);
-
 	if (TimelineSlider.IsValid())
 	{
-		TimelineSlider->SetValue(CurrentFrame);
 		if (EnumHasAnyFlags(Options, EChaosVDSetTimelineFrameFlags::BroadcastChange))
 		{
-			FrameChangedDelegate.ExecuteIfBound(CurrentFrame);
+			FrameChangedDelegate.ExecuteIfBound(FrameNumber);
 		}
 	}
 }
 
-void SChaosVDTimelineWidget::SetIsLocked(bool NewIsLocked)
+void SChaosVDTimelineWidget::Tick(const FGeometry& AllottedGeometry, const double InCurrentTime, const float InDeltaTime)
 {
-	bIsLocked = NewIsLocked;
-
-	if (bIsLocked)
-	{
-		ElementEnabledFlags = static_cast<uint16>(EChaosVDTimelineElementIDFlags::Lock);	
-	}
-	else
-	{
-		ElementEnabledFlags = DefaultEnabledElementsFlags;
-	}
-}
+	SCompoundWidget::Tick(AllottedGeometry, InCurrentTime, InDeltaTime);
 
-FReply SChaosVDTimelineWidget::TogglePlay()
-{
-	if (bIsPlaying)
-	{
-		Pause();
-		ButtonClickedDelegate.ExecuteIfBound(EChaosVDPlaybackButtonsID::Pause);
-	}
-	else
+	if (TimelineSlider)
 	{
-		Play();
-		ButtonClickedDelegate.ExecuteIfBound(EChaosVDPlaybackButtonsID::Play);
+		const int32 CurrentMinFrames = GetCurrentMinFrames();
+		const int32 CurrentMaxFrames = GetCurrentMaxFrames();
+		if (!FMath::IsNearlyEqual(TimelineSlider->GetMinValue(), CurrentMinFrames) || !FMath::IsNearlyEqual(TimelineSlider->GetMaxValue(), CurrentMaxFrames))
+		{
+			TimelineSlider->SetMinAndMaxValues(CurrentMinFrames, CurrentMaxFrames);
+		}
 	}
-
-	return FReply::Handled();
 }
 
-void SChaosVDTimelineWidget::Play()
+void SChaosVDTimelineWidget::Pause()
 {
-	bIsPlaying = true;
-	SetCanTick(true);
+	ButtonClickedDelegate.ExecuteIfBound(EChaosVDPlaybackButtonsID::Pause);
 }
 
-FReply SChaosVDTimelineWidget::Stop()
+FReply SChaosVDTimelineWidget::Next()
 {
-	CurrentFrame = 0;
-	CurrentPlaybackTime = 0.0f;
-	bIsPlaying = false;
-
-	SetCurrentTimelineFrame(CurrentFrame);
+	ButtonClickedDelegate.ExecuteIfBound(EChaosVDPlaybackButtonsID::Next);
 
-	SetCanTick(false);
+	return FReply::Handled();
+}
 
-	ButtonClickedDelegate.ExecuteIfBound(EChaosVDPlaybackButtonsID::Stop);
+FReply SChaosVDTimelineWidget::Prev()
+{	
+	ButtonClickedDelegate.ExecuteIfBound(EChaosVDPlaybackButtonsID::Prev);
 
 	return FReply::Handled();
 }
 
-void SChaosVDTimelineWidget::Pause()
+const FSlateBrush* SChaosVDTimelineWidget::GetPlayOrPauseIcon() const
 {
-	bIsPlaying = false;
-	SetCanTick(false);
+	return IsPlaying() ? FChaosVDStyle::Get().GetBrush("PauseIcon") : FChaosVDStyle::Get().GetBrush("PlayIcon");
 }
 
-FReply SChaosVDTimelineWidget::Next()
+EVisibility SChaosVDTimelineWidget::GetElementVisibility(EChaosVDTimelineElementIDFlags ElementID) const
 {
-	if (CurrentFrame >= MaxFrames)
-	{
-		CurrentFrame = MaxFrames;
-		return FReply::Handled();
-	}
-
-	CurrentFrame++;
-	
-	SetCurrentTimelineFrame(CurrentFrame);
-
-	ButtonClickedDelegate.ExecuteIfBound(EChaosVDPlaybackButtonsID::Next);
-
-	return FReply::Handled();
+	const bool bIsVisible = ElementVisibilityFlags.IsSet() ? EnumHasAnyFlags(ElementVisibilityFlags.Get(), ElementID) : false;
+	return bIsVisible ? EVisibility::Visible : EVisibility::Collapsed;
 }
 
-FReply SChaosVDTimelineWidget::Prev()
+bool SChaosVDTimelineWidget::GetElementEnabled(EChaosVDTimelineElementIDFlags ElementID) const
 {
-	if (CurrentFrame == 0)
+	if (!ElementEnabledFlags.IsSet())
 	{
-		return FReply::Handled();
+		return false;
 	}
 
-	CurrentFrame--;
-
-	SetCurrentTimelineFrame(CurrentFrame);
 	
-	ButtonClickedDelegate.ExecuteIfBound(EChaosVDPlaybackButtonsID::Prev);
-
-	return FReply::Handled();
+	return EnumHasAnyFlags(ElementEnabledFlags.Get(), ElementID);
 }
 
-FReply SChaosVDTimelineWidget::ToggleLockState()
+bool SChaosVDTimelineWidget::IsPlaying() const
 {
-	SetIsLocked(!bIsLocked);
-
-	FrameLockedDelegate.ExecuteIfBound(bIsLocked);
-
-	return FReply::Handled();
+	return bIsPlaying.IsSet() ? bIsPlaying.Get() : false;
 }
 
-const FSlateBrush* SChaosVDTimelineWidget::GetPlayOrPauseIcon() const
+int32 SChaosVDTimelineWidget::GetCurrentFrame() const
 {
-	return bIsPlaying ? FChaosVDStyle::Get().GetBrush("PauseIcon") : FChaosVDStyle::Get().GetBrush("PlayIcon");
+	return CurrentFrame.IsSet() ? CurrentFrame.Get() : INDEX_NONE;
 }
 
-const FSlateBrush* SChaosVDTimelineWidget::GetLockStateIcon() const
+float SChaosVDTimelineWidget::GetCurrentFrameAsFloat() const
 {
-	return bIsLocked ? FChaosVDStyle::Get().GetBrush("LockIcon") : FChaosVDStyle::Get().GetBrush("UnlockedIcon");
+	return static_cast<float>(GetCurrentFrame());
 }
 
-EVisibility SChaosVDTimelineWidget::GetElementVisibility(EChaosVDTimelineElementIDFlags ElementID) const
+int32 SChaosVDTimelineWidget::GetCurrentMinFrames() const
 {
-	 return ((static_cast<uint16>(ElementID) & ElementVisibilityFlags) == static_cast<uint16>(ElementID)) ? EVisibility::Visible : EVisibility::Collapsed;
+	return MinFrames.IsSet() ? MinFrames.Get() : INDEX_NONE;
 }
 
-bool SChaosVDTimelineWidget::GetElementEnabled(EChaosVDTimelineElementIDFlags ElementID) const
-{
-	return ((static_cast<uint16>(ElementID) & ElementEnabledFlags) == static_cast<uint16>(ElementID));
-}
 
-void SChaosVDTimelineWidget::Tick(const FGeometry& AllottedGeometry, const double InCurrentTime, const float InDeltaTime)
+int32 SChaosVDTimelineWidget::GetCurrentMaxFrames() const
 {
-	SCompoundWidget::Tick(AllottedGeometry, InCurrentTime, InDeltaTime);
-
-	//TODO: We should move the Ticking logic to advance frames outside of this widget.
-	// The logic to update the visual state is already controlled externally
-	if (bIsPlaying)
-	{
-		if (CurrentFrame == MaxFrames)
-		{
-			if (bAutoStopEnabled)
-			{
-				Stop();
-			}
-		}
-
-		CurrentPlaybackTime += InDeltaTime;
-
-		while (CurrentPlaybackTime > CurrentPlaybackRate)
-		{
-			CurrentPlaybackTime -= CurrentPlaybackRate;
-			Next();
-		}
-	}
+	return MaxFrames.IsSet() ? MaxFrames.Get() : INDEX_NONE;
 }
 
 #undef LOCTEXT_NAMESPACE
diff --git a/Engine/Plugins/ChaosVD/Source/ChaosVD/Private/Widgets/SChaosVDTimelineWidget.h b/Engine/Plugins/ChaosVD/Source/ChaosVD/Private/Widgets/SChaosVDTimelineWidget.h
--- a/Engine/Plugins/ChaosVD/Source/ChaosVD/Private/Widgets/SChaosVDTimelineWidget.h
+++ b/Engine/Plugins/ChaosVD/Source/ChaosVD/Private/Widgets/SChaosVDTimelineWidget.h
@@ -9,6 +9,7 @@ class SSlider;
 class FReply;
 struct FSlateBrush;
 
+UENUM()
 enum class EChaosVDPlaybackButtonsID : uint8
 {
 	Play,
@@ -19,7 +20,6 @@ enum class EChaosVDPlaybackButtonsID : uint8
 };
 
 DECLARE_DELEGATE_OneParam(FChaosControlButtonClicked, EChaosVDPlaybackButtonsID)
-DECLARE_DELEGATE_OneParam(FChaosVDFrameChangedDelegate, int32)
 DECLARE_DELEGATE_OneParam(FChaosVDFrameLockStateDelegate, bool)
 
 enum class EChaosVDSetTimelineFrameFlags
@@ -48,89 +48,65 @@ enum class EChaosVDTimelineElementIDFlags : uint16
 };
 ENUM_CLASS_FLAGS(EChaosVDTimelineElementIDFlags)
 
+DECLARE_DELEGATE_OneParam(FChaosVDFrameChangedDelegate, int32)
+
 /** Simple timeline control widget */
 class SChaosVDTimelineWidget : public SCompoundWidget
 {
 public:
 
 	SLATE_BEGIN_ARGS( SChaosVDTimelineWidget ){}
-		SLATE_ARGUMENT(int32, MaxFrames)
-		SLATE_ARGUMENT(uint16, ButtonVisibilityFlags)
-		SLATE_ARGUMENT(bool, AutoStopEnabled)
+		SLATE_ATTRIBUTE(int32, MaxFrames)
+		SLATE_ATTRIBUTE(int32, MinFrames)
+		SLATE_ATTRIBUTE(int32, CurrentFrame)
+		SLATE_ARGUMENT(EChaosVDTimelineElementIDFlags, ButtonVisibilityFlags)
+		SLATE_ATTRIBUTE(EChaosVDTimelineElementIDFlags, ButtonEnabledFlags)
+		SLATE_ATTRIBUTE(bool, IsPlaying)
 		SLATE_EVENT(FChaosVDFrameChangedDelegate, OnFrameChanged)
-		SLATE_EVENT(FChaosVDFrameLockStateDelegate, OnFrameLockStateChanged)
 		SLATE_EVENT(FChaosControlButtonClicked, OnButtonClicked)
 	SLATE_END_ARGS()
 
 	void Construct(const FArguments& InArgs);
 
-	void UpdateMinMaxValue(float NewMin, float NewMax);
-
-	void SetCurrentTimelineFrame(float FrameNumber, EChaosVDSetTimelineFrameFlags Options = EChaosVDSetTimelineFrameFlags::BroadcastChange);
-
-	void SetIsLocked(bool NewIsLocked);
-
-	bool IsUnlocked() const { return !bIsLocked; }
-
-	void SetTargetFrameTime(float TargetFrameTimeSeconds);
-
-	/** Brings back the state of the timeline to its original state*/
-	void ResetTimeline();
-
-	/** Called when a new frame is manually selected or auto-updated during playback */
-	FChaosVDFrameChangedDelegate& OnFrameChanged() { return FrameChangedDelegate; }
-	
-	/** Called when this timeline is locked or unlocked */
-	FChaosVDFrameLockStateDelegate& OnFrameLockStateChanged() { return FrameLockedDelegate; }
-
-	uint16& GetCurrentElementVisibilityFlags() { return ElementVisibilityFlags; }
-	void SetCurrentElementVisibilityFlags(uint16 NewVisibilityFlags) { ElementVisibilityFlags = NewVisibilityFlags; }
-
-	uint16& GetMutableElementEnabledFlagsRef() { return ElementEnabledFlags; }
-
-	void SetAutoStopEnabled(bool bNewEnabled) { bAutoStopEnabled = bNewEnabled; }
-
-	int32 GetCurrentFrame() const { return CurrentFrame; }
+protected:
 
 	void Play();
-	FReply  Stop();
+	FReply Stop();
 
-protected:
+	void SetCurrentTimelineFrame(float FrameNumber, EChaosVDSetTimelineFrameFlags Options);
+
+	void Tick(const FGeometry& AllottedGeometry, const double InCurrentTime, const float InDeltaTime) override;
 
 	void Pause();
 	FReply  TogglePlay();
 	FReply  Next();
 	FReply  Prev();
 
-	FReply ToggleLockState();
-
 	const FSlateBrush* GetPlayOrPauseIcon() const;
-	const FSlateBrush* GetLockStateIcon() const;
 
 	EVisibility GetElementVisibility(EChaosVDTimelineElementIDFlags ElementID) const;
 	bool GetElementEnabled(EChaosVDTimelineElementIDFlags ElementID) const;
 
-	virtual void Tick(const FGeometry& AllottedGeometry, const double InCurrentTime, const float InDeltaTime) override;
+	bool IsPlaying() const;
+
+	int32 GetCurrentFrame() const;
+	float GetCurrentFrameAsFloat() const;
+	int32 GetCurrentMinFrames() const;
+	int32 GetCurrentMaxFrames() const;
 
 	TSharedPtr<SSlider> TimelineSlider;
 
-	int32 CurrentFrame = 0;
-	int32 MinFrames = 0;
-	int32 MaxFrames = 1000;
+	TAttribute<int32> CurrentFrame = 0;
+	TAttribute<int32> MinFrames = 0;
+	TAttribute<int32> MaxFrames = 1000;
+	TAttribute<EChaosVDTimelineElementIDFlags> ElementEnabledFlags = EChaosVDTimelineElementIDFlags::All;
 
-	bool bIsPlaying = false;
-	float CurrentPlaybackTime = 0.0f;
-	float CurrentPlaybackRate = 1.0f / 60.0f;
+	TAttribute<bool> bIsPlaying = false;
 
 	FChaosVDFrameChangedDelegate FrameChangedDelegate;
-	FChaosVDFrameLockStateDelegate FrameLockedDelegate;
 	FChaosControlButtonClicked ButtonClickedDelegate;
 
-	bool bIsLocked = false;
-	bool bAutoStopEnabled = false;
+	TAttribute<EChaosVDTimelineElementIDFlags> ElementVisibilityFlags = EChaosVDTimelineElementIDFlags::All;
 
-	uint16 ElementVisibilityFlags = 0;
-	uint16 ElementEnabledFlags = static_cast<uint16>(EChaosVDTimelineElementIDFlags::All);
-	uint16 DefaultEnabledElementsFlags = static_cast<uint16>(EChaosVDTimelineElementIDFlags::All);
-	
+	EChaosVDTimelineElementIDFlags DefaultEnabledElementsFlags =EChaosVDTimelineElementIDFlags::All;
 };
diff --git a/Engine/Source/Runtime/Experimental/ChaosVDData/Private/ChaosVDRecording.cpp b/Engine/Source/Runtime/Experimental/ChaosVDData/Private/ChaosVDRecording.cpp
--- a/Engine/Source/Runtime/Experimental/ChaosVDData/Private/ChaosVDRecording.cpp
+++ b/Engine/Source/Runtime/Experimental/ChaosVDData/Private/ChaosVDRecording.cpp
@@ -41,16 +41,16 @@ int32 FChaosVDRecording::GetAvailableSolverFramesNumber_AssumesLocked(int32 Solv
 	return INDEX_NONE;
 }
 
-FString FChaosVDRecording::GetSolverName(int32 SolverID)
+FName FChaosVDRecording::GetSolverFName(int32 SolverID)
 {
 	FReadScopeLock ReadLock(RecordingDataLock);
 
-	return GetSolverName_AssumedLocked(SolverID);
+	return GetSolverFName_AssumedLocked(SolverID);
 }
 
-FString FChaosVDRecording::GetSolverName_AssumedLocked(int32 SolverID)
+FName FChaosVDRecording::GetSolverFName_AssumedLocked(int32 SolverID)
 {
-	static FString DefaultName(TEXT("Invalid"));
+	static FName DefaultName(TEXT("Invalid"));
 
 	// Currently we don't create an entry per solver, so we need to get the name from the frame data
 	// TODO: Record Solver specific data per instance and not per frame
@@ -58,7 +58,7 @@ FString FChaosVDRecording::GetSolverName_AssumedLocked(int32 SolverID)
 	{
 		if(!SolverFramesData->IsEmpty())
 		{
-			return (*SolverFramesData)[0].DebugName;
+			return (*SolverFramesData)[0].DebugFName;
 		}
 	}
 
@@ -164,6 +164,12 @@ int32 FChaosVDRecording::FindFirstSolverKeyFrameNumberFromFrame_AssumesLocked(in
 int32 FChaosVDRecording::GetLowestSolverFrameNumberGameFrame(int32 SolverID, int32 GameFrame)
 {
 	FReadScopeLock ReadLock(RecordingDataLock);
+	return GetLowestSolverFrameNumberGameFrame_AssumesLocked(SolverID, GameFrame);
+}
+
+int32 FChaosVDRecording::GetLowestSolverFrameNumberGameFrame_AssumesLocked(int32 SolverID, int32 GameFrame)
+{
+
 	if (!GameFrames.IsValidIndex(GameFrame))
 	{
 		return INDEX_NONE;
@@ -179,9 +185,15 @@ int32 FChaosVDRecording::GetLowestSolverFrameNumberGameFrame(int32 SolverID, int
 	return INDEX_NONE;
 }
 
+
 int32 FChaosVDRecording::GetLowestGameFrameAtSolverFrameNumber(int32 SolverID, int32 SolverFrame)
 {
 	FReadScopeLock ReadLock(RecordingDataLock);
+	return GetLowestGameFrameAtSolverFrameNumber_AssumesLocked(SolverID, SolverFrame);
+}
+
+int32 FChaosVDRecording::GetLowestGameFrameAtSolverFrameNumber_AssumesLocked(int32 SolverID, int32 SolverFrame)
+{
 	if (TArray<FChaosVDSolverFrameData>* SolverFramesPtr = RecordedFramesDataPerSolver.Find(SolverID))
 	{
 		TArray<FChaosVDSolverFrameData>& SolverFrames = *SolverFramesPtr;
@@ -195,6 +207,7 @@ int32 FChaosVDRecording::GetLowestGameFrameAtSolverFrameNumber(int32 SolverID, i
 	return INDEX_NONE;
 }
 
+
 void FChaosVDRecording::AddKeyFrameNumberForSolver(const int32 SolverID, int32 FrameNumber)
 {
 	FWriteScopeLock WriteLock(RecordingDataLock);
@@ -324,54 +337,6 @@ int32 FChaosVDRecording::GetLowestGameFrameNumberAtCycle(uint64 Cycle)
 	return Algo::LowerBoundBy(GameFrames, Cycle, &FChaosVDGameFrameData::FirstCycle);
 }
 
-void FChaosVDRecording::GetAvailableSolverIDsAtGameFrameNumber(int32 FrameNumber, TArray<int32>& OutSolversID)
-{
-	FReadScopeLock ReadLock(RecordingDataLock);
-
-	return GetAvailableSolverIDsAtGameFrameNumber_AssumesLocked(FrameNumber, OutSolversID);
-}
-
-void FChaosVDRecording::GetAvailableSolverIDsAtGameFrameNumber_AssumesLocked(int32 FrameNumber, TArray<int32>& OutSolversID)
-{
-	if (!GameFrames.IsValidIndex(FrameNumber))
-	{
-		return;
-	}
-	
-	GetAvailableSolverIDsAtGameFrame_AssumesLocked(GameFrames[FrameNumber], OutSolversID);
-}
-
-void FChaosVDRecording::GetAvailableSolverIDsAtGameFrame(const FChaosVDGameFrameData& GameFrameData, TArray<int32>& OutSolversID)
-{
-	FReadScopeLock ReadLock(RecordingDataLock);
-	GetAvailableSolverIDsAtGameFrame_AssumesLocked(GameFrameData, OutSolversID);
-}
-
-void FChaosVDRecording::GetAvailableSolverIDsAtGameFrame_AssumesLocked(const FChaosVDGameFrameData& GameFrameData, TArray<int32>& OutSolversID)
-{
-	OutSolversID.Reserve(RecordedFramesDataPerSolver.Num());
-
-	for (const TPair<int32, TArray<FChaosVDSolverFrameData>>& SolverFramesWithIDPair : RecordedFramesDataPerSolver)
-	{
-		if (SolverFramesWithIDPair.Value.IsEmpty())
-		{
-			continue;
-		}
-
-		if (SolverFramesWithIDPair.Value.Num() == 1 && SolverFramesWithIDPair.Value[0].FrameCycle < GameFrameData.FirstCycle)
-		{
-			OutSolversID.Add(SolverFramesWithIDPair.Key);
-		}
-		else
-		{
-			if (GameFrameData.FirstCycle > SolverFramesWithIDPair.Value[0].FrameCycle && GameFrameData.FirstCycle < SolverFramesWithIDPair.Value.Last().FrameCycle)
-			{
-				OutSolversID.Add(SolverFramesWithIDPair.Key);
-			}
-		}	
-	}
-}
-
 void FChaosVDRecording::CollapseSolverFramesRange_AssumesLocked(int32 SolverID, int32 StartFrame, int32 EndFrame, FChaosVDSolverFrameData& OutCollapsedFrameData)
 {
 	// Make sure we start with a clear map
@@ -418,7 +383,7 @@ void FChaosVDRecording::CollapseSolverFramesRange_AssumesLocked(int32 SolverID,
 				OutCollapsedFrameData.bIsKeyFrame = true;
 				OutCollapsedFrameData.SolverID = SolverFrameData->SolverID;
 				OutCollapsedFrameData.SimulationTransform = SolverFrameData->SimulationTransform;
-				OutCollapsedFrameData.DebugName = SolverFrameData->DebugName;
+				OutCollapsedFrameData.DebugFName = SolverFrameData->DebugFName;
 
 				FChaosVDStepData CollapsedStepData;
 				CollapsedStepData.StepName = TEXT("GeneratedStep");
diff --git a/Engine/Source/Runtime/Experimental/ChaosVDData/Public/ChaosVDRecording.h b/Engine/Source/Runtime/Experimental/ChaosVDData/Public/ChaosVDRecording.h
--- a/Engine/Source/Runtime/Experimental/ChaosVDData/Public/ChaosVDRecording.h
+++ b/Engine/Source/Runtime/Experimental/ChaosVDData/Public/ChaosVDRecording.h
@@ -51,6 +51,17 @@ typedef TArray<FChaosVDStepData, TInlineAllocator<16>> FChaosVDStepsContainer;
 
 struct CHAOSVDDATA_API FChaosVDSolverFrameData
 {
+
+PRAGMA_DISABLE_DEPRECATION_WARNINGS
+	FChaosVDSolverFrameData() = default;
+	FChaosVDSolverFrameData(const FChaosVDSolverFrameData& Other) = default;
+	FChaosVDSolverFrameData(FChaosVDSolverFrameData&& Other) noexcept = default;
+	FChaosVDSolverFrameData& operator=(const FChaosVDSolverFrameData& Other) = default;
+	FChaosVDSolverFrameData& operator=(FChaosVDSolverFrameData&& Other) noexcept = default;
+PRAGMA_ENABLE_DEPRECATION_WARNINGS
+	
+	FName DebugFName;
+	UE_DEPRECATED(5.5, "Please use the DebugFName instead")
 	FString DebugName;
 	int32 SolverID = INDEX_NONE;
 	uint64 FrameCycle = 0;
@@ -61,6 +72,7 @@ struct CHAOSVDDATA_API FChaosVDSolverFrameData
 	TSet<int32> ParticlesDestroyedIDs;
 	double StartTime = -1.0;
 	double EndTime = -1.0;
+	TArray<TSharedPtr<FChaosVDCharacterGroundConstraint>> RecordedCharacterGroundConstraints;
 	
 	/** Calculates and returns the frame time for this recorded frame.
 	 * @return Calculated frame time. -1 if it was not recorded
@@ -99,7 +111,6 @@ struct FChaosVDGameFrameData
 	TMap<FName, FChaosVDTrackedLocation> RecordedNonSolverLocationsByID;
 	TMap<FName, FChaosVDTrackedTransform> RecordedNonSolverTransformsByID;
 	TMap<int32, TSharedPtr<FChaosVDQueryDataWrapper>> RecordedSceneQueries;
-	TArray<TSharedPtr<FChaosVDCharacterGroundConstraint>> RecordedCharacterGroundConstraints;
 };
 
 /**
@@ -110,6 +121,10 @@ struct CHAOSVDDATA_API FChaosVDRecording
 {
 	FChaosVDRecording();
 
+	/* Constant used to define inline allocators -
+	* Unless there are some scenarios with a lot of RBAN solvers in the recording, we usually don't go over 3 tracks most of the time so 16 should be plenty by default */
+	static constexpr int32 CommonTrackCount = 16;
+
 	/** Returns the current available recorded solvers number */
 	int32 GetAvailableSolversNumber_AssumesLocked() const { return RecordedFramesDataPerSolver.Num(); }
 	
@@ -134,13 +149,19 @@ struct CHAOSVDDATA_API FChaosVDRecording
 	 * Returns the name of the specified solver id
 	 * @param SolverID ID of the solver 
 	 */
-	FString GetSolverName(int32 SolverID);
+	FName GetSolverFName(int32 SolverID);
+
+	UE_DEPRECATED(5.5, "Please use the GetSolverFName instead")
+	FString GetSolverName(int32 SolverID) { return TEXT(""); }
 
 	/**
 	 * Returns the name of the specified solver id. Must be called from within a ReadLock
 	 * @param SolverID ID of the solver
 	 */
-	FString GetSolverName_AssumedLocked(int32 SolverID);
+	FName GetSolverFName_AssumedLocked(int32 SolverID);
+
+	UE_DEPRECATED(5.5, "Please use the GetSolverFName_AssumedLocked instead")
+	FString GetSolverName_AssumedLocked(int32 SolverID) { return TEXT(""); }
 
 	/**
 	 * Return a ptr to the existing solver frame data from the specified ID and Frame number
@@ -177,6 +198,7 @@ struct CHAOSVDDATA_API FChaosVDRecording
 	 * @return Found frame number. INDEX_NONE if no frame is found for the specified cycle
 	 */
 	int32 GetLowestSolverFrameNumberGameFrame(int32 SolverID, int32 GameFrame);
+	int32 GetLowestSolverFrameNumberGameFrame_AssumesLocked(int32 SolverID, int32 GameFrame);
 	
 	/**
 	 * Searches and returns the lowest game frame number at the specified solver frame
@@ -185,6 +207,7 @@ struct CHAOSVDDATA_API FChaosVDRecording
 	 * @return Found Game frame number. INDEX_NONE if no frame is found for the specified cycle
 	 */
 	int32 GetLowestGameFrameAtSolverFrameNumber(int32 SolverID, int32 SolverFrame);
+	int32 GetLowestGameFrameAtSolverFrameNumber_AssumesLocked(int32 SolverID, int32 SolverFrame);
 
 	/**
 	 * Adds a Solver Frame Data entry for a specific Solver ID. Creates a solver entry if it does not exist 
@@ -231,10 +254,14 @@ struct CHAOSVDDATA_API FChaosVDRecording
      * @param FrameNumber Game Frame number to evaluate
      * @param OutSolversID Solver's ID array to be filled with any IDs found
      */
-	void GetAvailableSolverIDsAtGameFrameNumber(int32 FrameNumber, TArray<int32>& OutSolversID);
-	void GetAvailableSolverIDsAtGameFrameNumber_AssumesLocked(int32 FrameNumber, TArray<int32>& OutSolversID);
-	void GetAvailableSolverIDsAtGameFrame(const FChaosVDGameFrameData& GameFrameData, TArray<int32>& OutSolversID);
-	void GetAvailableSolverIDsAtGameFrame_AssumesLocked(const FChaosVDGameFrameData& GameFrameData, TArray<int32>& OutSolversID);
+	template<typename TAllocator>
+	void GetAvailableSolverIDsAtGameFrameNumber(int32 FrameNumber,TArray<int32, TAllocator>& OutSolversID);
+	template<typename TAllocator>
+	void GetAvailableSolverIDsAtGameFrameNumber_AssumesLocked(int32 FrameNumber, TArray<int32, TAllocator>& OutSolversID);
+	template<typename TAllocator>
+	void GetAvailableSolverIDsAtGameFrame(const FChaosVDGameFrameData& GameFrameData, TArray<int32, TAllocator>& OutSolversID);
+	template<typename TAllocator>
+	void GetAvailableSolverIDsAtGameFrame_AssumesLocked(const FChaosVDGameFrameData& GameFrameData, TArray<int32, TAllocator>& OutSolversID);
 
 	/** Collapses the most important frame data from a range of solver frames into a single solver frame data */
 	void CollapseSolverFramesRange_AssumesLocked(int32 SolverID, int32 StartFrame, int32 EndFrame, FChaosVDSolverFrameData& OutCollapsedFrameData);
@@ -324,3 +351,54 @@ protected:
 	friend class FChaosVDTraceProvider;
 	friend class FChaosVDTraceImplicitObjectProcessor;
 };
+
+template <typename TAllocator>
+void FChaosVDRecording::GetAvailableSolverIDsAtGameFrameNumber(int32 FrameNumber, TArray<int32, TAllocator>& OutSolversID)
+{
+	FReadScopeLock ReadLock(RecordingDataLock);
+	return GetAvailableSolverIDsAtGameFrameNumber_AssumesLocked(FrameNumber, OutSolversID);
+}
+
+template <typename TAllocator>
+void FChaosVDRecording::GetAvailableSolverIDsAtGameFrameNumber_AssumesLocked(int32 FrameNumber, TArray<int32, TAllocator>& OutSolversID)
+{
+	if (!GameFrames.IsValidIndex(FrameNumber))
+	{
+		return;
+	}
+	
+	GetAvailableSolverIDsAtGameFrame_AssumesLocked(GameFrames[FrameNumber], OutSolversID);
+}
+
+template <typename TAllocator>
+void FChaosVDRecording::GetAvailableSolverIDsAtGameFrame(const FChaosVDGameFrameData& GameFrameData, TArray<int32, TAllocator>& OutSolversID)
+{
+	FReadScopeLock ReadLock(RecordingDataLock);
+	GetAvailableSolverIDsAtGameFrame_AssumesLocked(GameFrameData, OutSolversID);
+}
+
+template <typename TAllocator>
+void FChaosVDRecording::GetAvailableSolverIDsAtGameFrame_AssumesLocked(const FChaosVDGameFrameData& GameFrameData, TArray<int32, TAllocator>& OutSolversID)
+{
+	OutSolversID.Reserve(RecordedFramesDataPerSolver.Num());
+
+	for (const TPair<int32, TArray<FChaosVDSolverFrameData>>& SolverFramesWithIDPair : RecordedFramesDataPerSolver)
+	{
+		if (SolverFramesWithIDPair.Value.IsEmpty())
+		{
+			continue;
+		}
+
+		if (SolverFramesWithIDPair.Value.Num() == 1 && SolverFramesWithIDPair.Value[0].FrameCycle < GameFrameData.FirstCycle)
+		{
+			OutSolversID.Add(SolverFramesWithIDPair.Key);
+		}
+		else
+		{
+			if (GameFrameData.FirstCycle > SolverFramesWithIDPair.Value[0].FrameCycle && GameFrameData.FirstCycle < SolverFramesWithIDPair.Value.Last().FrameCycle)
+			{
+				OutSolversID.Add(SolverFramesWithIDPair.Key);
+			}
+		}	
+	}
+}
