[ChaosVD] Playback performance Improvements

- Removed the need to call SetActorLabel to  update the particle debug name in the outliner as it is too expensive.
- Updated CVD Outliner mode to know how to read the particle debug name from the recorded particle data instead of the actor label.
- Limited the number of geometry generation tasks to be launched per frame. There was an issue where in some scenarios we were trying to launch more tasks than the task graph scheduler could accept and therefore we were spending a considerable amount of time waiting in the GT.
- Silenced a log spam in the Implicit Object code which was causing us to spend ~25% on I/O in some scenarios. A proper fix will be done in a future CL.
- Made material instance creation async, and limited to a time budget per frame
- Set CVD ISMC to do RemoveSwap when an instance is removed
- Updated the TrackObject mode so it updated the CVD camera each time the scene updates instead of just when the viewport is drawn. This reduced the stutter when tracking objects moving at high speeds.
- Removed the need to add the ISMC to the instanced component arrays of the pool actor. We were not using that functionality and calling AddInstanceComponents does an AddUnique internally, which becomes too expensive when we are adding a huge number of components.
-Improved how particles subscribe to be aware new geometry is available.
- AChosVDParticleActor now does nothing when modify is called. These actors are completely transient and Modify calls where adding cost to the actor spawning process.
- Made the Scene Outliner mode block updates to the hierarchy while we are playing. The scene outliner will now update itself as soon we stop or pause a recording

diff --git a/Engine/Plugins/ChaosVD/Source/ChaosVD/Private/ChaosVDGeometryBuilder.cpp b/Engine/Plugins/ChaosVD/Source/ChaosVD/Private/ChaosVDGeometryBuilder.cpp
--- a/Engine/Plugins/ChaosVD/Source/ChaosVD/Private/ChaosVDGeometryBuilder.cpp
+++ b/Engine/Plugins/ChaosVD/Source/ChaosVD/Private/ChaosVDGeometryBuilder.cpp
@@ -6,6 +6,7 @@
 #include "ChaosVDGeometryDataComponent.h"
 #include "ChaosVDModule.h"
 #include "ChaosVDTriMeshGenerator.h"
+#include "Components/ChaosVDStaticMeshComponent.h"
 #include "DynamicMeshToMeshDescription.h"
 #include "Generators/CapsuleGenerator.h"
 #include "Generators/MinimalBoxMeshGenerator.h"
@@ -33,6 +34,12 @@ namespace Chaos::VisualDebugger
 			TEXT("p.Chaos.VD.Tool.DisableUVsSupport"),
 			bDisableUVsSupport,
 			TEXT("If true, the generated meshes will not have UV data"));
+
+		static float GeometryGenerationTaskLaunchBudgetSeconds = 0.005f;
+		static FAutoConsoleVariableRef CVarGeometryGenerationTaskLaunchBudgetSeconds(
+			TEXT("p.Chaos.VD.Tool.GeometryGenerationTaskLaunchBudgetSeconds"),
+			GeometryGenerationTaskLaunchBudgetSeconds,
+			TEXT("How much time we can spend on the Geoemtry builder tick launching Geometry Generation Tasks"));
 	}
 
 	void SetTriangleAttributes(const UE::Geometry::FMeshShapeGenerator& Generator, FDynamicMesh3& OutDynamicMesh, int32 AppendedTriangleID, int32 GeneratorTriangleIndex)
@@ -165,7 +172,7 @@ void FChaosVDGeometryBuilder::Initialize(const TWeakPtr<FChaosVDScene>& ChaosVDS
 		const TSharedPtr<FChaosVDGeometryBuilder> GeometryBuilder = WeakThis.Pin();
 		if (!GeometryBuilder)
 		{
-			UE_LOG(LogChaosVDEditor, Verbose, TEXT(" [%s] Failed to update mesh for Handle | Geometry Key [%u] | Handle is invalid"), ANSI_TO_TCHAR(__FUNCTION__), GeometryKey);
+			UE_LOG(LogChaosVDEditor, Verbose, TEXT(" [%s] Failed to update mesh for Handle | Geometry Key [%u] | Geometry Builder is invalid"), ANSI_TO_TCHAR(__FUNCTION__), GeometryKey);
 
 			// If the the builder is no longer valid, just consume the request
 			return true;
@@ -183,8 +190,46 @@ void FChaosVDGeometryBuilder::Initialize(const TWeakPtr<FChaosVDScene>& ChaosVDS
 
 		return false;
 	};
+	
+	auto CreateMaterialForMeshComponent = [WeakThis = AsWeak()](const TWeakObjectPtr<UMeshComponent> Object)
+	{
+		const TSharedPtr<FChaosVDGeometryBuilder> GeometryBuilder = WeakThis.Pin();
+		if (!GeometryBuilder)
+		{
+			UE_LOG(LogChaosVDEditor, Verbose, TEXT(" [%s] Failed to Create Material for Mesh | Geometry builder is no longer valid "), ANSI_TO_TCHAR(__FUNCTION__));
+
+			// If the the builder is no longer valid, just consume the request
+			return true;
+		}
+
+		UMeshComponent* MeshComponent = Object.Get();
+		if (IChaosVDGeometryComponent* CVDMeshComponent = Cast<IChaosVDGeometryComponent>(MeshComponent))
+		{
+			// The Mesh component no longer has instances on it,
+			// this means the component was returned to the pool or is scheduled to be destroyed while we were waiting 
+			if (CVDMeshComponent->GetMeshDataInstanceHandles().IsEmpty())
+			{
+				return true;
+			}
+
+			GeometryBuilder->SetMeshComponentMaterial(CVDMeshComponent->GetMeshComponentAttributeFlags(), MeshComponent);
+		}
+		return true;
+	};
+
+	auto LauchGeometryGenerationTaskDeferred = [WeakThis = AsWeak()](TSharedPtr<FGeometryGenerationTask> GeometryGenerationTask)
+	{
+		UE::Tasks::Launch(TEXT("GeometryGeneration"),
+[GeometryGenerationTask]()
+		{
+			GeometryGenerationTask->GenerateGeometry();
+		});
+		return true;
+	};
 
 	MeshComponentsWaitingForGeometry = MakeUnique<FObjectsWaitingGeometryList<FMeshComponentWeakPtr>>(ProcessMeshComponent, NSLOCTEXT("ChaosVisualDebugger", "GeometryGenNotification","Mesh Components"), ShouldProcessObjectsForKey);
+	MeshComponentsWaitingForMaterial = MakeUnique<FObjectsWaitingProcessingQueue<FMeshComponentWeakPtr>>(CreateMaterialForMeshComponent, NSLOCTEXT("ChaosVisualDebugger", "GeometryMaterialNotification","Material instances"));
+	GeometryTasksPendingLaunch = MakeUnique<FObjectsWaitingProcessingQueue<TSharedPtr<FGeometryGenerationTask>>>(LauchGeometryGenerationTaskDeferred, NSLOCTEXT("ChaosVisualDebugger", "GeometryTaskLauchNotification","Static Meshes"));
 
 	GameThreadTickDelegate = FTSTicker::GetCoreTicker().AddTicker(FTickerDelegate::CreateRaw(this, &FChaosVDGeometryBuilder::GameThreadTick));
 
@@ -194,7 +239,6 @@ void FChaosVDGeometryBuilder::Initialize(const TWeakPtr<FChaosVDScene>& ChaosVDS
 
 void FChaosVDGeometryBuilder::AddReferencedObjects(FReferenceCollector& Collector)
 {
-	Collector.AddStableReferenceMap(DynamicMeshCacheMap);
 	Collector.AddStableReferenceMap(StaticMeshCacheMap);
 	Collector.AddReferencedObjects(MeshComponentsPendingDisposal);
 }
@@ -254,42 +298,17 @@ bool FChaosVDGeometryBuilder::HasGeometryInCache(uint32 GeometryKey)
 
 bool FChaosVDGeometryBuilder::HasGeometryInCache_AssumesLocked(uint32 GeometryKey) const
 {
-	return StaticMeshCacheMap.Contains(GeometryKey) || DynamicMeshCacheMap.Contains(GeometryKey);
+	return StaticMeshCacheMap.Contains(GeometryKey);
 }
 
-UDynamicMesh* FChaosVDGeometryBuilder::CreateAndCacheDynamicMesh(const uint32 GeometryCacheKey, UE::Geometry::FMeshShapeGenerator& MeshGenerator)
+UStaticMesh* FChaosVDGeometryBuilder::GetCachedMeshForImplicit(const uint32 GeometryCacheKey)
 {
+	if (const TObjectPtr<UStaticMesh>* MeshPtrPtr = StaticMeshCacheMap.Find(GeometryCacheKey))
 	{
-		FReadScopeLock ReadLock(GeometryCacheRWLock);
-		if (TObjectPtr<UDynamicMesh>* DynamicMeshPtrPtr = DynamicMeshCacheMap.Find(GeometryCacheKey))
-		{
-			return *DynamicMeshPtrPtr;
-		}
-	}
-
-	TRACE_CPUPROFILER_EVENT_SCOPE(FChaosVDGeometryBuilder::CreateAndCacheDynamicMesh_BUILD);
-
-	UDynamicMesh* Mesh = NewObject<UDynamicMesh>();
-
-	FDynamicMesh3 DynamicMesh;
-
-	if (Chaos::VisualDebugger::Cvars::bUseCVDDynamicMeshGenerator)
-	{
-		Chaos::VisualDebugger::GenerateDynamicMeshFromGenerator(MeshGenerator.Generate(), DynamicMesh);
-	}
-	else
-	{
-		DynamicMesh.Copy(&MeshGenerator.Generate());
+		return MeshPtrPtr->Get();
 	}
 
-	Mesh->SetMesh(DynamicMesh);
-
-	{
-		FWriteScopeLock WriteLock(GeometryCacheRWLock);
-		DynamicMeshCacheMap.Add(GeometryCacheKey, Mesh);
-	}
-
-	return Mesh;
+	return nullptr;
 }
 
 UStaticMesh* FChaosVDGeometryBuilder::CreateAndCacheStaticMesh(const uint32 GeometryCacheKey, UE::Geometry::FMeshShapeGenerator& MeshGenerator, const int32 LODsToGenerateNum)
@@ -384,6 +403,24 @@ UStaticMesh* FChaosVDGeometryBuilder::CreateAndCacheStaticMesh(const uint32 Geom
 	return MainStaticMesh;
 }
 
+void FChaosVDGeometryBuilder::SetMeshComponentMaterial(EChaosVDMeshAttributesFlags MeshComponentAttributeFlags, UMeshComponent* MeshComponent)
+{
+	UMaterialInstanceDynamic* Material = nullptr;
+
+	if (Cast<UChaosVDInstancedStaticMeshComponent>(MeshComponent))
+	{
+		Material = FChaosVDGeometryComponentUtils::CreateMaterialInstance(FChaosVDGeometryComponentUtils::GetMaterialTypeForComponent<UChaosVDInstancedStaticMeshComponent>(MeshComponentAttributeFlags));
+	}
+	else
+	{
+		Material = FChaosVDGeometryComponentUtils::CreateMaterialInstance(FChaosVDGeometryComponentUtils::GetMaterialTypeForComponent<UChaosVDStaticMeshComponent>(MeshComponentAttributeFlags));
+	}
+
+	ensure(Material);
+
+	MeshComponent->SetMaterial(0, Material);
+}
+
 void FChaosVDGeometryBuilder::DestroyMeshComponent(UMeshComponent* MeshComponent)
 {
 	if (Cast<UChaosVDInstancedStaticMeshComponent>(MeshComponent))
@@ -453,13 +490,9 @@ bool FChaosVDGeometryBuilder::ApplyMeshToComponentFromKey(TWeakObjectPtr<UMeshCo
 
 	if (HasGeometryInCache(GeometryKey))
 	{
-		if (UDynamicMeshComponent* DynamicMeshComponent = Cast<UDynamicMeshComponent>(MeshComponent))
+		if (UStaticMeshComponent* StaticMeshComponent = Cast<UStaticMeshComponent>(MeshComponent))
 		{
-			DynamicMeshComponent->SetDynamicMesh(GetCachedMeshForImplicit<UDynamicMesh>(GeometryKey));
-		}
-		else if (UStaticMeshComponent* StaticMeshComponent = Cast<UStaticMeshComponent>(MeshComponent))
-		{
-			StaticMeshComponent->SetStaticMesh(GetCachedMeshForImplicit<UStaticMesh>(GeometryKey));
+			StaticMeshComponent->SetStaticMesh(GetCachedMeshForImplicit(GeometryKey));
 		}
 
 		DataComponent->SetIsMeshReady(true);
@@ -633,11 +666,21 @@ bool FChaosVDGeometryBuilder::ImplicitObjectNeedsUnpacking(const Chaos::FImplici
 
 bool FChaosVDGeometryBuilder::GameThreadTick(float DeltaTime)
 {
-	int32 CurrentGeometryTasksProcessedNum = 0;
+	const float BudgetPerCategory = Chaos::VisualDebugger::Cvars::GeometryGenerationTaskLaunchBudgetSeconds / 3;
+
+	if (GeometryTasksPendingLaunch)
+	{
+		GeometryTasksPendingLaunch->ProcessWaitingTasks(BudgetPerCategory);
+	}
 
 	if (MeshComponentsWaitingForGeometry)
 	{
-		MeshComponentsWaitingForGeometry->ProcessWaitingObjects(CurrentGeometryTasksProcessedNum);
+		MeshComponentsWaitingForGeometry->ProcessWaitingObjects(BudgetPerCategory);
+	}
+
+	if (MeshComponentsWaitingForMaterial)
+	{
+		MeshComponentsWaitingForMaterial->ProcessWaitingTasks(BudgetPerCategory);
 	}
 
 	for (TObjectPtr<UMeshComponent>& MeshComponentPtr : MeshComponentsPendingDisposal)
@@ -705,3 +748,15 @@ void FChaosVDGeometryBuilder::HandleStaticMeshComponentInstanceIndexUpdated(UIns
 		}
 	}
 }
+
+void FGeometryGenerationTask::GenerateGeometry()
+{
+	if (const TSharedPtr<FChaosVDGeometryBuilder> BuilderPtr = Builder.Pin())
+	{
+		BuilderPtr->CreateAndCacheStaticMesh(GeometryKey, *MeshGenerator.Get(), LODsToGenerateNum);
+		{
+			FWriteScopeLock WriteLock(BuilderPtr->GeometryCacheRWLock);
+			BuilderPtr->GeometryBeingGeneratedByKey.Remove(GeometryKey);
+		}
+	}
+}
diff --git a/Engine/Plugins/ChaosVD/Source/ChaosVD/Private/ChaosVDGeometryBuilder.h b/Engine/Plugins/ChaosVD/Source/ChaosVD/Private/ChaosVDGeometryBuilder.h
--- a/Engine/Plugins/ChaosVD/Source/ChaosVD/Private/ChaosVDGeometryBuilder.h
+++ b/Engine/Plugins/ChaosVD/Source/ChaosVD/Private/ChaosVDGeometryBuilder.h
@@ -26,6 +26,7 @@
 #include "UObject/GCObject.h"
 #include "UObject/UObjectGlobals.h"
 
+class FGeometryGenerationTask;
 class UChaosVDStaticMeshComponent;
 class AActor;
 class UDynamicMesh;
@@ -69,15 +70,6 @@ public:
 				StaticMeshByKey.Value->MarkAsGarbage();
 			}
 		}
-
-		for (const TPair<uint32, TObjectPtr<UDynamicMesh>>& DynamicMeshByKey : DynamicMeshCacheMap)
-		{
-			if (DynamicMeshByKey.Value)
-			{
-				DynamicMeshByKey.Value->ClearFlags(RF_Standalone);
-				DynamicMeshByKey.Value->MarkAsGarbage();
-			}
-		}
 	}
 	
 	void Initialize(const TWeakPtr<FChaosVDScene>& ChaosVDScene);
@@ -134,15 +126,9 @@ public:
 	/** Returns an already mesh for the provided implicit object if exists, otherwise returns null
 	 * @param GeometryCacheKey Key to be used to find this geometry in the cache
 	 */
-	template<typename MeshType>
-	MeshType* GetCachedMeshForImplicit(const uint32 GeometryCacheKey);
+	UStaticMesh* GetCachedMeshForImplicit(const uint32 GeometryCacheKey);
 
 private:
-	/** Creates a Dynamic Mesh for the provided Implicit object and generator, and then caches it to be reused later
-	 * @param GeometryCacheKey Key to be used to find this geometry in the cache
-	 * @param MeshGenerator Generator class with the data and rules to create the mesh
-	 */
-	UDynamicMesh* CreateAndCacheDynamicMesh(const uint32 GeometryCacheKey, UE::Geometry::FMeshShapeGenerator& MeshGenerator);
 
 	/** Creates a Dynamic Mesh for the provided Implicit object and generator, and then caches it to be reused later
 	 * @param GeometryCacheKey Key to be used to find this geometry in the cache
@@ -158,8 +144,7 @@ private:
 	template <class ComponentType>
     bool InitializeMeshComponent(AActor* Owner, ComponentType* MeshComponent);
 
-	template <class ComponentType>
-	void SetMeshComponentMaterial(EChaosVDMeshAttributesFlags MeshComponentAttributeFlags, ComponentType* MeshComponent);
+	void SetMeshComponentMaterial(EChaosVDMeshAttributesFlags MeshComponentAttributeFlags, UMeshComponent* MeshComponent);
 
 public:
 	/**
@@ -226,7 +211,7 @@ private:
 	bool ApplyMeshToComponentFromKey(TWeakObjectPtr<UMeshComponent> MeshComponent, const uint32 GeometryKey);
 
 public:
-	/** Creates a mesh generator for the provided Implicit object which will be used to create a Static Mesh or Dynamic Mesh
+	/** Creates a mesh generator for the provided Implicit object which will be used to create a Static Mesh
 	 * @param InImplicit ImplicitObject used a data source for the mesh generator
 	 * @param SimpleShapesComplexityFactor Factor used to reduce or increase the complexity (number of triangles generated) of simple shapes (Sphere/Capsule)
 	 */
@@ -291,9 +276,6 @@ private:
 	/** Handles any changes to the indexes of created instanced mesh components we are managing, making corrections/updates as needed */
 	void HandleStaticMeshComponentInstanceIndexUpdated(UInstancedStaticMeshComponent* InComponent, TArrayView<const FInstancedStaticMeshDelegates::FInstanceIndexUpdateData> InIndexUpdates);
 
-	/** Map containing already generated dynamic mesh for any given implicit object */
-	TMap<uint32, TObjectPtr<UDynamicMesh>> DynamicMeshCacheMap;
-
 	/** Map containing already generated static mesh for any given implicit object */
 	TMap<uint32, TObjectPtr<UStaticMesh>> StaticMeshCacheMap;
 
@@ -309,8 +291,9 @@ private:
 	/** Object containing all the meshes component waiting for geometry, by geometry key*/
 	TUniquePtr<FObjectsWaitingGeometryList<FMeshComponentWeakPtr>> MeshComponentsWaitingForGeometry;
 
-	/** Object containing all the meshes component waiting for geometry, by geometry key*/
-	TUniquePtr<FObjectsWaitingGeometryList<FExtractedGeometryHandle>> HandlesWaitingForGeometry;
+	TUniquePtr<FObjectsWaitingProcessingQueue<FMeshComponentWeakPtr>> MeshComponentsWaitingForMaterial;
+
+	TUniquePtr<FObjectsWaitingProcessingQueue<TSharedPtr<FGeometryGenerationTask>>> GeometryTasksPendingLaunch;
 
 	/** Map containing already initialized Instanced static mesh components for any given geometry key */
 	TMap<uint32, UChaosVDInstancedStaticMeshComponent*> InstancedMeshComponentByGeometryKey;
@@ -351,7 +334,6 @@ public:
 	{
 	}
 
-	template <typename MeshType>
 	void GenerateGeometry();
 
 private:
@@ -450,27 +432,6 @@ void FChaosVDGeometryBuilder::CreateMeshesFromImplicit_Internal(const Chaos::FIm
 	}
 }
 
-template <typename MeshType>
-MeshType* FChaosVDGeometryBuilder::GetCachedMeshForImplicit(const uint32 GeometryCacheKey)
-{
-	if constexpr (std::is_same_v<MeshType, UDynamicMesh>)
-	{
-		if (const TObjectPtr<UDynamicMesh>* MeshPtrPtr = DynamicMeshCacheMap.Find(GeometryCacheKey))
-		{
-			return MeshPtrPtr->Get();
-		}
-	}
-	else if constexpr (std::is_same_v<MeshType, UStaticMesh>)
-	{
-		if (const TObjectPtr<UStaticMesh>* MeshPtrPtr = StaticMeshCacheMap.Find(GeometryCacheKey))
-		{
-			return MeshPtrPtr->Get();
-		}
-	}
-
-	return nullptr;
-}
-
 template <typename ComponentType>
 bool FChaosVDGeometryBuilder::InitializeMeshComponent(AActor* Owner, ComponentType* MeshComponent)
 {
@@ -483,7 +444,6 @@ bool FChaosVDGeometryBuilder::InitializeMeshComponent(AActor* Owner, ComponentTy
 	if (Owner)
 	{
 		Owner->AddOwnedComponent(MeshComponent);
-		Owner->AddInstanceComponent(MeshComponent);
 		MeshComponent->RegisterComponent();
 		MeshComponent->AttachToComponent(Owner->GetRootComponent(), FAttachmentTransformRules::SnapToTargetIncludingScale);
 
@@ -498,19 +458,6 @@ bool FChaosVDGeometryBuilder::InitializeMeshComponent(AActor* Owner, ComponentTy
 	return true;
 }
 
-template <typename ComponentType>
-void FChaosVDGeometryBuilder::SetMeshComponentMaterial(EChaosVDMeshAttributesFlags MeshComponentAttributeFlags, ComponentType* MeshComponent)
-{
-	UMaterialInstanceDynamic* Material = nullptr;
-
-	Material = FChaosVDGeometryComponentUtils::CreateMaterialInstance(FChaosVDGeometryComponentUtils::GetMaterialTypeForComponent<ComponentType>(MeshComponentAttributeFlags));
-
-	ensure(Material);
-
-	MeshComponent->SetMaterial(0, Material);
-}
-
-
 template <typename ComponentType>
 ComponentType* FChaosVDGeometryBuilder::GetAvailableInstancedStaticMeshComponent(const TSharedPtr<FChaosVDExtractedGeometryDataHandle>& InExtractedGeometryDataHandle, AActor* MeshComponentsContainerActor, EChaosVDMeshAttributesFlags MeshComponentAttributeFlags, bool& bOutIsNewComponent)
 {
@@ -536,7 +483,10 @@ ComponentType* FChaosVDGeometryBuilder::GetAvailableInstancedStaticMeshComponent
 		if (UChaosVDInstancedStaticMeshComponent* AsInstancedMeshComponent = Component)
 		{
 			AsInstancedMeshComponent->bReverseCulling = EnumHasAnyFlags(MeshComponentAttributeFlags, EChaosVDMeshAttributesFlags::MirroredGeometry);
-			SetMeshComponentMaterial<ComponentType>(MeshComponentAttributeFlags, AsInstancedMeshComponent);
+			if (!AsInstancedMeshComponent->HasOverrideMaterials())
+			{
+				MeshComponentsWaitingForMaterial->EnqueueObject(AsInstancedMeshComponent);
+			}	
 		}
 
 		bOutIsNewComponent = true;
@@ -563,7 +513,10 @@ ComponentType* FChaosVDGeometryBuilder::GetAvailableMeshComponent(const TSharedP
 			return nullptr;
 		}
 
-		SetMeshComponentMaterial(MeshComponentAttributeFlags, MeshComponent);
+		if (!MeshComponent->HasOverrideMaterials())
+		{
+			MeshComponentsWaitingForMaterial->EnqueueObject(MeshComponent);
+		}
 
 		bOutIsNewComponent = true;
 	}
@@ -659,6 +612,8 @@ ComponentType* FChaosVDGeometryBuilder::GetMeshComponentForNewInstance(const TSh
 
 	if (IChaosVDGeometryComponent* CVDGeometryComponent = Cast<IChaosVDGeometryComponent>(Component))
 	{
+		CVDGeometryComponent->SetMeshComponentAttributeFlags(MeshAttributes);
+
 		if (!CVDGeometryComponent->IsMeshReady())
 		{
 			AddMeshComponentWaitingForGeometry(GeometryDataHandle->GetGeometryKey(), Component);
@@ -667,7 +622,6 @@ ComponentType* FChaosVDGeometryBuilder::GetMeshComponentForNewInstance(const TSh
 		if (bIsNew)
 		{
 			CVDGeometryComponent->OnComponentEmpty()->AddRaw(this, &FChaosVDGeometryBuilder::DestroyMeshComponent);			
-			CVDGeometryComponent->SetMeshComponentAttributeFlags(MeshAttributes);
 		}	
 	}
 
@@ -688,11 +642,7 @@ void FChaosVDGeometryBuilder::DispatchCreateAndCacheMeshForImplicitAsync(const u
 	}
 	
 	TSharedPtr<FGeometryGenerationTask> GenerationTask = MakeShared<FGeometryGenerationTask>(AsWeak(), MeshGenerator, GeometryKey, LODsToGenerateNum);
-	UE::Tasks::Launch(TEXT("GeometryGeneration"),
-		[GenerationTask]()
-		{
-			GenerationTask->GenerateGeometry<MeshType>();
-		});
+	GeometryTasksPendingLaunch->EnqueueObject(GenerationTask);
 }
 
 template <bool bIsInstanced, typename GeometryType>
@@ -739,25 +689,3 @@ const GeometryType* FChaosVDGeometryBuilder::GetGeometryBasedOnPackedType(const
 		return GetGeometry<false, GeometryType>(InImplicitObject, bIsScaled, Transform);
 	}
 }
-
-template <typename MeshType>
-void FGeometryGenerationTask::GenerateGeometry()
-{
-	if (const TSharedPtr<FChaosVDGeometryBuilder> BuilderPtr = Builder.Pin())
-	{
-		if constexpr (std::is_same_v<MeshType, UDynamicMesh>)
-		{
-			ensureMsgf(LODsToGenerateNum == 0, TEXT("LOD Generation is only suppoted with static meshes | [%d] LODs were requested for a dynamic mesh when 0 is expected"), LODsToGenerateNum);
-			BuilderPtr->CreateAndCacheDynamicMesh(GeometryKey, *MeshGenerator.Get());
-		}
-		else if constexpr (std::is_same_v<MeshType, UStaticMesh>)
-		{
-			BuilderPtr->CreateAndCacheStaticMesh(GeometryKey, *MeshGenerator.Get(), LODsToGenerateNum);
-		}
-
-		{
-			FWriteScopeLock WriteLock(BuilderPtr->GeometryCacheRWLock);
-			BuilderPtr->GeometryBeingGeneratedByKey.Remove(GeometryKey);
-		}
-	}
-}
diff --git a/Engine/Plugins/ChaosVD/Source/ChaosVD/Private/ChaosVDGeometryDataComponent.h b/Engine/Plugins/ChaosVD/Source/ChaosVD/Private/ChaosVDGeometryDataComponent.h
--- a/Engine/Plugins/ChaosVD/Source/ChaosVD/Private/ChaosVDGeometryDataComponent.h
+++ b/Engine/Plugins/ChaosVD/Source/ChaosVD/Private/ChaosVDGeometryDataComponent.h
@@ -374,6 +374,7 @@ public:
 	virtual void SetSelectedMeshInstance(const TWeakPtr<FChaosVDMeshDataInstanceHandle>& InMeshInstanceToSelect) PURE_VIRTUAL(IChaosVDGeometryOwnerInterface::SetSelectedMeshInstance);
 	virtual TWeakPtr<FChaosVDMeshDataInstanceHandle> GetSelectedMeshInstance() const PURE_VIRTUAL(IChaosVDGeometryOwnerInterface::GetSelectedMeshInstance, return nullptr;);
 	virtual TConstArrayView<TSharedPtr<FChaosVDMeshDataInstanceHandle>> GetMeshInstances() const PURE_VIRTUAL(IChaosVDGeometryOwnerInterface::GetMeshInstances, return TConstArrayView<TSharedPtr<FChaosVDMeshDataInstanceHandle>>(););
+	virtual void HandleNewGeometryLoaded(uint32 GeometryID, const Chaos::FConstImplicitObjectPtr& InGeometryData) PURE_VIRTUAL(IChaosVDGeometryOwnerInterface::HandleGeometryUpdated);
 };
 
 namespace Chaos::VisualDebugger
diff --git a/Engine/Plugins/ChaosVD/Source/ChaosVD/Private/ChaosVDMeshComponentPool.h b/Engine/Plugins/ChaosVD/Source/ChaosVD/Private/ChaosVDMeshComponentPool.h
--- a/Engine/Plugins/ChaosVD/Source/ChaosVD/Private/ChaosVDMeshComponentPool.h
+++ b/Engine/Plugins/ChaosVD/Source/ChaosVD/Private/ChaosVDMeshComponentPool.h
@@ -26,8 +26,8 @@ public:
 
 private:
 
-	template <typename TDesiredMeshComponent, typename TPooledComponent>
-	TDesiredMeshComponent* GetMeshComponentFromPool_Internal(TArray<TObjectPtr<TPooledComponent>>& InMeshComponentPool, UObject* Outer, FName Name);
+	template <typename DesiredMeshComponent, typename PooledComponent>
+	DesiredMeshComponent* GetMeshComponentFromPool_Internal(TArray<TObjectPtr<PooledComponent>>& InMeshComponentPool, UObject* Outer, FName Name);
 
 	void ResetMeshComponent(UMeshComponent* MeshComponent);
 
@@ -60,17 +60,17 @@ TMeshComponent* FChaosVDMeshComponentPool::AcquireMeshComponent(UObject* Outer,
 	}
 }
 
-template <typename TDesiredMeshComponent, typename TPooledComponent>
-TDesiredMeshComponent* FChaosVDMeshComponentPool::GetMeshComponentFromPool_Internal(TArray<TObjectPtr<TPooledComponent>>& InMeshComponentPool, UObject* Outer,  FName Name)
+template <typename DesiredMeshComponent, typename PooledComponent>
+DesiredMeshComponent* FChaosVDMeshComponentPool::GetMeshComponentFromPool_Internal(TArray<TObjectPtr<PooledComponent>>& InMeshComponentPool, UObject* Outer,  FName Name)
 {
-	static_assert(std::is_base_of_v<UStaticMeshComponent, TDesiredMeshComponent> || std::is_base_of_v<UInstancedStaticMeshComponent, TDesiredMeshComponent> || std::is_base_of_v<TDesiredMeshComponent, UDynamicMeshComponent>, "GetMeshComponentInternal Only supports DynamicMeshComponent, Static MeshComponent and Instanced Static Mesh Component");
+	static_assert(std::is_base_of_v<UStaticMeshComponent, DesiredMeshComponent> || std::is_base_of_v<UInstancedStaticMeshComponent, DesiredMeshComponent> || std::is_base_of_v<DesiredMeshComponent, UDynamicMeshComponent>, "GetMeshComponentInternal Only supports DynamicMeshComponent, Static MeshComponent and Instanced Static Mesh Component");
 
 	// We need to ensure unique names
 	const FString NewName = Name.ToString() + FGuid::NewGuid().ToString();
 	
 	if (bUseComponentsPool &&  InMeshComponentPool.Num() > 0)
 	{
-		if (TDesiredMeshComponent* Component = Cast<TDesiredMeshComponent>(InMeshComponentPool.Pop()))
+		if (DesiredMeshComponent* Component = Cast<DesiredMeshComponent>(InMeshComponentPool.Pop()))
 		{
 			Component->Rename(*NewName , Outer, REN_NonTransactional | REN_DoNotDirty | REN_ForceNoResetLoaders | REN_SkipGeneratedClasses | REN_DontCreateRedirectors);
 			return Component;
@@ -81,5 +81,8 @@ TDesiredMeshComponent* FChaosVDMeshComponentPool::GetMeshComponentFromPool_Inter
 		}
 	}
 
-	return NewObject<TDesiredMeshComponent>(Outer, *NewName);
+	DesiredMeshComponent* NewMeshComponent = NewObject<DesiredMeshComponent>(Outer, *NewName);
+	NewMeshComponent->EmptyOverrideMaterials();
+
+	return NewMeshComponent;
 }
diff --git a/Engine/Plugins/ChaosVD/Source/ChaosVD/Private/ChaosVDParticleActor.cpp b/Engine/Plugins/ChaosVD/Source/ChaosVD/Private/ChaosVDParticleActor.cpp
--- a/Engine/Plugins/ChaosVD/Source/ChaosVD/Private/ChaosVDParticleActor.cpp
+++ b/Engine/Plugins/ChaosVD/Source/ChaosVD/Private/ChaosVDParticleActor.cpp
@@ -133,6 +133,11 @@ void AChaosVDParticleActor::SetSelectedMeshInstance(const TWeakPtr<FChaosVDMeshD
 	}
 }
 
+void AChaosVDParticleActor::HandleNewGeometryLoaded(uint32 GeometryID, const Chaos::FConstImplicitObjectPtr& InGeometryData)
+{
+	UpdateGeometry(InGeometryData, EChaosVDActorGeometryUpdateFlags::ForceUpdate);
+}
+
 void AChaosVDParticleActor::HandleSelected()
 {
 	PushSelectionToProxies();
@@ -144,6 +149,12 @@ void AChaosVDParticleActor::HandleDeSelected()
 	PushSelectionToProxies();
 }
 
+bool AChaosVDParticleActor::Modify(bool bAlwaysMarkDirty)
+{
+	//CVD Actors are transient. Skipping super call to recude the cost of spawning them
+	return false;
+}
+
 void AChaosVDParticleActor::ProcessUpdatedAndRemovedHandles(TArray<TSharedPtr<FChaosVDExtractedGeometryDataHandle>>& OutExtractedGeometryDataHandles)
 {
 	for (TArray<TSharedPtr<FChaosVDMeshDataInstanceHandle>>::TIterator MeshDataHandleRemoveIterator = MeshDataHandles.CreateIterator(); MeshDataHandleRemoveIterator; ++MeshDataHandleRemoveIterator)
@@ -184,7 +195,6 @@ void AChaosVDParticleActor::ProcessUpdatedAndRemovedHandles(TArray<TSharedPtr<FC
 	}
 }
 
-
 void AChaosVDParticleActor::UpdateGeometry(const Chaos::FConstImplicitObjectPtr& InImplicitObject, EChaosVDActorGeometryUpdateFlags OptionsFlags)
 {
 	if (EnumHasAnyFlags(OptionsFlags, EChaosVDActorGeometryUpdateFlags::ForceUpdate))
@@ -314,29 +324,8 @@ void AChaosVDParticleActor::UpdateGeometry(uint32 NewGeometryHash, EChaosVDActor
 	}
 }
 
-void AChaosVDParticleActor::SetScene(TWeakPtr<FChaosVDScene> InScene)
-{
-	FChaosVDSceneObjectBase::SetScene(InScene);
-
-	if (const TSharedPtr<FChaosVDScene>& ScenePtr = SceneWeakPtr.Pin())
-	{
-		GeometryUpdatedDelegate = ScenePtr->OnNewGeometryAvailable().AddWeakLambda(this, [this](const Chaos::FConstImplicitObjectPtr& ImplicitObject, const uint32 ID)
-		{
-			if (ParticleDataPtr && ParticleDataPtr->GeometryHash == ID)
-			{
-				UpdateGeometry(ImplicitObject, EChaosVDActorGeometryUpdateFlags::ForceUpdate);
-			}
-		});
-	}
-}
-
 void AChaosVDParticleActor::Destroyed()
 {
-	if (const TSharedPtr<FChaosVDScene>& ScenePtr = SceneWeakPtr.Pin())
-	{
-		ScenePtr->OnNewGeometryAvailable().Remove(GeometryUpdatedDelegate);
-	}
-
 	VisitGeometryInstances([](const TSharedRef<FChaosVDMeshDataInstanceHandle>& MeshDataHandle)
 	{
 		if (IChaosVDGeometryComponent* AsGeometryComponent = Cast<IChaosVDGeometryComponent>(MeshDataHandle->GetMeshComponent()))
@@ -397,8 +386,15 @@ FBox AChaosVDParticleActor::GetComponentsBoundingBox(bool bNonColliding, bool bI
 			}
 		}
 
-		const FBoxSphereBounds SphereBounds= BoundsBuilder;
-		BoundingBox = SphereBounds.GetBox();
+		// Geometry might not be generated yet, so we need a placeholder box so the Focus on object feature works
+		if (!BoundsBuilder.IsValid())
+		{
+			return BoundingBox.ExpandBy(10.0f).MoveTo(ParticleDataPtr->ParticlePositionRotation.MX);
+		}
+		else
+		{
+			BoundingBox = FBoxSphereBounds(BoundsBuilder).GetBox();
+		}
 	}
 
 	return BoundingBox;
diff --git a/Engine/Plugins/ChaosVD/Source/ChaosVD/Private/ChaosVDParticleActor.h b/Engine/Plugins/ChaosVD/Source/ChaosVD/Private/ChaosVDParticleActor.h
--- a/Engine/Plugins/ChaosVD/Source/ChaosVD/Private/ChaosVDParticleActor.h
+++ b/Engine/Plugins/ChaosVD/Source/ChaosVD/Private/ChaosVDParticleActor.h
@@ -69,8 +69,6 @@ public:
 
 	void UpdateGeometry(uint32 NewGeometryHash, EChaosVDActorGeometryUpdateFlags OptionsFlags = EChaosVDActorGeometryUpdateFlags::None);
 
-	virtual void SetScene(TWeakPtr<FChaosVDScene> InScene) override;
-
 	virtual void Destroyed() override;
 
 	virtual const FChaosVDParticleDataWrapper* GetParticleData() override { return ParticleDataPtr.Get(); }
@@ -127,6 +125,7 @@ public:
 	virtual TConstArrayView<TSharedPtr<FChaosVDMeshDataInstanceHandle>> GetMeshInstances() const override { return MeshDataHandles; }
 	virtual void SetSelectedMeshInstance(const TWeakPtr<FChaosVDMeshDataInstanceHandle>& GeometryInstanceToSelect) override;
 	virtual TWeakPtr<FChaosVDMeshDataInstanceHandle> GetSelectedMeshInstance() const override { return CurrentSelectedGeometryInstance; }
+	virtual void HandleNewGeometryLoaded(uint32 GeometryID, const Chaos::FConstImplicitObjectPtr& InGeometryData) override;
 	// END IChaosVDGeometryOwner Interface
 	
 	// BEGIN IChaosVDSelectableObject Interface
@@ -134,6 +133,8 @@ public:
 	virtual void HandleDeSelected() override;
 	// END IChaosVDSelectableObject Interface
 
+	virtual bool Modify(bool bAlwaysMarkDirty) override;
+
 protected:
 
 	void ProcessUpdatedAndRemovedHandles(TArray<TSharedPtr<FChaosVDExtractedGeometryDataHandle>>& OutExtractedGeometryDataHandles);
diff --git a/Engine/Plugins/ChaosVD/Source/ChaosVD/Private/ChaosVDPlaybackViewportClient.cpp b/Engine/Plugins/ChaosVD/Source/ChaosVD/Private/ChaosVDPlaybackViewportClient.cpp
--- a/Engine/Plugins/ChaosVD/Source/ChaosVD/Private/ChaosVDPlaybackViewportClient.cpp
+++ b/Engine/Plugins/ChaosVD/Source/ChaosVD/Private/ChaosVDPlaybackViewportClient.cpp
@@ -178,23 +178,26 @@ void FChaosVDPlaybackViewportClient::TrackSelectedObject()
 		USelection* CurrentSelection = ModeTools->GetSelectedActors();
 
 		//TODO: Update this if we add multi selection support
-		if (const AActor* SelectedActor = CurrentSelection ? CurrentSelection->GetTop<AActor>() : nullptr)
+		if (AActor* SelectedActor = CurrentSelection ? CurrentSelection->GetTop<AActor>() : nullptr)
 		{
-			const FBox ActorBounds = SelectedActor->GetComponentsBoundingBox(false);
-			FocusViewportOnBox(ActorBounds.ExpandBy(TrackingViewDistance), true);		
+			FBox BoxToTrack = SelectedActor->GetComponentsBoundingBox(false);
+			FocusViewportOnBox(BoxToTrack.ExpandBy(TrackingViewDistance), true);
 		}
 	}
 }
 
-bool FChaosVDPlaybackViewportClient::InputKey(const FInputKeyEventArgs& EventArgs)
+void FChaosVDPlaybackViewportClient::UpdateMouseDelta()
 {
-	// Each time we requested a re-draw means we move something in the scene, so we need to re-create the cached hit proxy map.
-	if (bNeedsRedraw)
-	{
-		RequestInvalidateHitProxy(Viewport);
-	}
-	
-	return FEditorViewportClient::InputKey(EventArgs);
+	// Make sure we get the camera in the correct position before a mouse drag is handled
+	TrackSelectedObject();
+
+	FEditorViewportClient::UpdateMouseDelta();
+}
+
+void FChaosVDPlaybackViewportClient::HandleCVDSceneUpdated()
+{
+	TrackSelectedObject();
+	Invalidate();
 }
 
 void FChaosVDPlaybackViewportClient::ToggleObjectTrackingIfSelected()
@@ -225,8 +228,6 @@ void FChaosVDPlaybackViewportClient::Draw(const FSceneView* View, FPrimitiveDraw
 		const_cast<FSceneView*>(View)->bAllowTranslucentPrimitivesInHitProxy = true;
 	}
 
-	TrackSelectedObject();
-
 	const TSharedPtr<SChaosVDMainTab> MainTabToolkitHost = ModeTools.IsValid() ? StaticCastSharedPtr<SChaosVDMainTab>(ModeTools->GetToolkitHost()) : nullptr;
 	if (!MainTabToolkitHost.IsValid())
 	{
diff --git a/Engine/Plugins/ChaosVD/Source/ChaosVD/Private/ChaosVDPlaybackViewportClient.h b/Engine/Plugins/ChaosVD/Source/ChaosVD/Private/ChaosVDPlaybackViewportClient.h
--- a/Engine/Plugins/ChaosVD/Source/ChaosVD/Private/ChaosVDPlaybackViewportClient.h
+++ b/Engine/Plugins/ChaosVD/Source/ChaosVD/Private/ChaosVDPlaybackViewportClient.h
@@ -27,8 +27,6 @@ public:
 	virtual void Draw(const FSceneView* View, FPrimitiveDrawInterface* PDI) override;
 	virtual void DrawCanvas(FViewport& InViewport, FSceneView& View, FCanvas& Canvas) override;
 
-	virtual bool InputKey(const FInputKeyEventArgs& EventArgs) override;
-
 	void ToggleObjectTrackingIfSelected();
 	bool IsAutoTrackingSelectedObject() const { return bAutoTrackSelectedObject; };
 	void SetAutoTrackingViewDistance(float NewDistance);
@@ -36,9 +34,13 @@ public:
 
 	void GoToLocation(const FVector& InLocation);
 
-private:
-
 	void TrackSelectedObject();
+	
+	virtual void UpdateMouseDelta() override;
+
+	void HandleCVDSceneUpdated();
+
+private:
 
 	void HandleObjectFocused(UObject* FocusedObject);
 	void HandleActorMoving(AActor* MovedActor) const;
diff --git a/Engine/Plugins/ChaosVD/Source/ChaosVD/Private/ChaosVDScene.cpp b/Engine/Plugins/ChaosVD/Source/ChaosVD/Private/ChaosVDScene.cpp
--- a/Engine/Plugins/ChaosVD/Source/ChaosVD/Private/ChaosVDScene.cpp
+++ b/Engine/Plugins/ChaosVD/Source/ChaosVD/Private/ChaosVDScene.cpp
@@ -252,9 +252,22 @@ void FChaosVDScene::UpdateJointConstraintsData(const FChaosVDStepData& InRecorde
 	}
 }
 
-void FChaosVDScene::HandleNewGeometryData(const Chaos::FConstImplicitObjectPtr& GeometryData, const uint32 GeometryID) const
+void FChaosVDScene::HandleNewGeometryData(const Chaos::FConstImplicitObjectPtr& GeometryData, const uint32 GeometryID)
 {
-	NewGeometryAvailableDelegate.Broadcast(GeometryData, GeometryID);
+	if (TArray<IChaosVDGeometryOwnerInterface*>* ObjectsWaitingPtr = ObjectsWaitingForGeometry.Find(GeometryID))
+	{
+		TArray<IChaosVDGeometryOwnerInterface*>& ObjectsWaitingRef = *ObjectsWaitingPtr;
+		for (IChaosVDGeometryOwnerInterface* ObjectWaiting : ObjectsWaitingRef)
+		{
+			if (ObjectWaiting)
+			{
+				ObjectWaiting->HandleNewGeometryLoaded(GeometryID, GeometryData);
+			}
+		}
+
+		// Keep the array allocated in case another particle needs to go to the waiting list
+		ObjectsWaitingRef.Reset();
+	}
 }
 
 void FChaosVDScene::CreateSolverInfoActor(int32 SolverID)
@@ -386,6 +399,16 @@ Chaos::FConstImplicitObjectPtr FChaosVDScene::GetUpdatedGeometry(int32 GeometryI
 	return nullptr;
 }
 
+void FChaosVDScene::AddObjectWaitingForGeometry(uint32 GeometryID, IChaosVDGeometryOwnerInterface* ObjectWaitingForGeometry)
+{
+	if (!ObjectWaitingForGeometry)
+	{
+		return;
+	}
+
+	ObjectsWaitingForGeometry.FindOrAdd(GeometryID).Add(ObjectWaitingForGeometry);
+}
+
 AChaosVDParticleActor* FChaosVDScene::GetParticleActor(int32 SolverID, int32 ParticleID)
 {
 	if (AChaosVDSolverInfoActor** SolverDataInfo = SolverDataContainerBySolverID.Find(SolverID))
@@ -438,8 +461,9 @@ AChaosVDParticleActor* FChaosVDScene::SpawnParticleFromRecordedData(const TShare
 		NewActor->SetIsServerParticle(IsSolverForServer(InParticleData->SolverID));
 		NewActor->UpdateFromRecordedParticleData(InParticleData, InFrameData.SimulationTransform);
 
-		const bool bHasDebugName = !InParticleData->DebugName.IsEmpty();
-		NewActor->SetActorLabel(bHasDebugName ? InParticleData->DebugName : TEXT("Unnamed Particle - ID : ") + FString::FromInt(InParticleData->ParticleIndex));
+		// CVD's Outliner mode will update the label based on the particle data without needing to go trough all the code that Set Actor lable goes trough
+		// which can take +0.1 sec per actor
+		ParticleLabelUpdateDelegate.Broadcast(NewActor);
 
 		return NewActor;
 	}
diff --git a/Engine/Plugins/ChaosVD/Source/ChaosVD/Private/ChaosVDScene.h b/Engine/Plugins/ChaosVD/Source/ChaosVD/Private/ChaosVDScene.h
--- a/Engine/Plugins/ChaosVD/Source/ChaosVD/Private/ChaosVDScene.h
+++ b/Engine/Plugins/ChaosVD/Source/ChaosVD/Private/ChaosVDScene.h
@@ -11,6 +11,7 @@
 #include "UObject/ObjectMacros.h"
 #include "UObject/ObjectPtr.h"
 
+class IChaosVDGeometryOwnerInterface;
 class FChaosVDSelectionCustomization;
 class ITypedElementSelectionInterface;
 struct FTypedElementSelectionOptions;
@@ -33,7 +34,7 @@ struct FTypedElementHandle;
 typedef TMap<int32, AChaosVDSolverInfoActor*> FChaosVDSolverInfoByIDMap;
 
 DECLARE_MULTICAST_DELEGATE(FChaosVDSceneUpdatedDelegate)
-DECLARE_MULTICAST_DELEGATE_OneParam(FChaosVDActorActiveStateUpdateDelegate, AChaosVDParticleActor*)
+DECLARE_MULTICAST_DELEGATE_OneParam(FChaosVDActorUpdatedDelegate, AChaosVDParticleActor*)
 DECLARE_MULTICAST_DELEGATE_OneParam(FChaosVDOnObjectSelectedDelegate, UObject*)
 DECLARE_MULTICAST_DELEGATE_OneParam(FChaosVDSolverInfoActorCreatedDelegate, AChaosVDSolverInfoActor*)
 
@@ -65,7 +66,7 @@ public:
 	void UpdateJointConstraintsData(const FChaosVDStepData& InRecordedStepData, int32 SolverID);
 
 	// No need to deprecate the old version since it is not a public API nor inline 
-	void HandleNewGeometryData(const Chaos::FConstImplicitObjectPtr& Geometry, const uint32 GeometryID) const;
+	void HandleNewGeometryData(const Chaos::FConstImplicitObjectPtr& Geometry, const uint32 GeometryID);
 
 	void HandleEnterNewGameFrame(int32 FrameNumber, const TArray<int32, TInlineAllocator<16>>& AvailableSolversIds, const FChaosVDGameFrameData& InNewGameFrameData, TArray<int32, TInlineAllocator<16>>& OutRemovedSolversIds);
 	void HandleEnterNewSolverFrame(int32 FrameNumber, const FChaosVDSolverFrameData& InFrameData);
@@ -80,10 +81,10 @@ public:
 
 	const TSharedPtr<FChaosVDGeometryBuilder>& GetGeometryGenerator() { return  GeometryGenerator; }
 
-	FChaosVDGeometryDataLoaded& OnNewGeometryAvailable(){ return NewGeometryAvailableDelegate; }
-
 	// No need to deprecate the old version since it is not a public API nor inline 
 	Chaos::FConstImplicitObjectPtr GetUpdatedGeometry(int32 GeometryID) const;
+
+	void AddObjectWaitingForGeometry(uint32 GeometryID, IChaosVDGeometryOwnerInterface* ObjectWaitingForGeometry);
 	
 	/** Adds an object to the selection set if it was not selected already, making it selected in practice */
 	void SetSelectedObject(UObject* SelectedObject);
@@ -119,7 +120,8 @@ public:
 
 	UChaosVDSceneQueryDataComponent* GetSceneQueryDataContainerComponent() const;
 
-	FChaosVDActorActiveStateUpdateDelegate& OnActorActiveStateChanged() { return ParticleActorUpdateDelegate; }
+	FChaosVDActorUpdatedDelegate& OnActorActiveStateChanged() { return ParticleActorUpdateDelegate; }
+	FChaosVDActorUpdatedDelegate& OnActorLabelChanged() { return ParticleLabelUpdateDelegate; }
 
 	FChaosVDSolverInfoActorCreatedDelegate& OnSolverInfoActorCreated() { return SolverInfoActorCreatedDelegate; }
 
@@ -197,7 +199,8 @@ private:
 
 	bool bIsInitialized = false;
 
-	FChaosVDActorActiveStateUpdateDelegate ParticleActorUpdateDelegate;
+	FChaosVDActorUpdatedDelegate ParticleActorUpdateDelegate;
+	FChaosVDActorUpdatedDelegate ParticleLabelUpdateDelegate;
 
 	FDelegateHandle ActorDestroyedHandle;
 
@@ -205,5 +208,7 @@ private:
 
 	FChaosVDSolverVisibilityChangedDelegate SolverVisibilityChangedDelegate;
 
+	TMap<uint32, TArray<IChaosVDGeometryOwnerInterface*>> ObjectsWaitingForGeometry;
+
 	friend FChaosVDSelectionCustomization;
 };
diff --git a/Engine/Plugins/ChaosVD/Source/ChaosVD/Private/ChaosVDWorldOutlinerMode.cpp b/Engine/Plugins/ChaosVD/Source/ChaosVD/Private/ChaosVDWorldOutlinerMode.cpp
--- a/Engine/Plugins/ChaosVD/Source/ChaosVD/Private/ChaosVDWorldOutlinerMode.cpp
+++ b/Engine/Plugins/ChaosVD/Source/ChaosVD/Private/ChaosVDWorldOutlinerMode.cpp
@@ -5,6 +5,7 @@
 #include "ActorTreeItem.h"
 #include "ChaosVDModule.h"
 #include "ChaosVDParticleActor.h"
+#include "ChaosVDPlaybackController.h"
 #include "ChaosVDScene.h"
 #include "Elements/Framework/TypedElementSelectionSet.h"
 
@@ -89,6 +90,22 @@ void FChaosVDActorTreeItem::OnVisibilityChanged(const bool bNewVisibility)
 	}
 }
 
+void FChaosVDActorTreeItem::UpdateDisplayString()
+{
+	if (AChaosVDParticleActor* CVDActor = Cast<AChaosVDParticleActor>(Actor.Get()))
+	{
+		if (const FChaosVDParticleDataWrapper* ParticleData = CVDActor->GetParticleData())
+		{
+			const bool bHasDebugName = !ParticleData->DebugName.IsEmpty();
+			DisplayString = bHasDebugName ? ParticleData->DebugName : TEXT("Unnamed Particle - ID : ") + FString::FromInt(ParticleData->ParticleIndex);
+		}
+	}
+	else
+	{
+		FActorTreeItem::UpdateDisplayString();
+	}
+}
+
 TUniquePtr<FChaosVDOutlinerHierarchy> FChaosVDOutlinerHierarchy::Create(ISceneOutlinerMode* Mode, const TWeakObjectPtr<UWorld>& World)
 {
 	FChaosVDOutlinerHierarchy* Hierarchy = new FChaosVDOutlinerHierarchy(Mode, World);
@@ -104,9 +121,10 @@ FSceneOutlinerTreeItemPtr FChaosVDOutlinerHierarchy::CreateItemForActor(AActor*
 	return Mode->CreateItemFor<FChaosVDActorTreeItem>(InActor, bForce);
 }
 
-FChaosVDWorldOutlinerMode::FChaosVDWorldOutlinerMode(const FActorModeParams& InModeParams, TWeakPtr<FChaosVDScene> InScene)
+FChaosVDWorldOutlinerMode::FChaosVDWorldOutlinerMode(const FActorModeParams& InModeParams, TWeakPtr<FChaosVDScene> InScene, TWeakPtr<FChaosVDPlaybackController> InPlaybackController)
 	: FActorMode(InModeParams),
-	CVDScene(InScene)
+	CVDScene(InScene),
+	PlaybackController(InPlaybackController)
 {
 	TSharedPtr<FChaosVDScene> ScenePtr = CVDScene.Pin();
 	if (!ensure(ScenePtr.IsValid()))
@@ -115,19 +133,17 @@ FChaosVDWorldOutlinerMode::FChaosVDWorldOutlinerMode(const FActorModeParams& InM
 	}
 
 	ScenePtr->OnActorActiveStateChanged().AddRaw(this, &FChaosVDWorldOutlinerMode::HandleActorActiveStateChanged);
+	ScenePtr->OnActorLabelChanged().AddRaw(this, &FChaosVDWorldOutlinerMode::HandleActorLabelChanged);
 
 	RegisterSelectionSetObject(ScenePtr->GetElementSelectionSet());
-
-	ActorLabelChangedDelegateHandle = FCoreDelegates::OnActorLabelChanged.AddRaw(this, &FChaosVDWorldOutlinerMode::HandleActorLabelChanged);
 }
 
 FChaosVDWorldOutlinerMode::~FChaosVDWorldOutlinerMode()
 {
-	FCoreDelegates::OnActorLabelChanged.Remove(ActorLabelChangedDelegateHandle);
-
 	if (TSharedPtr<FChaosVDScene> ScenePtr = CVDScene.Pin())
 	{
 		ScenePtr->OnActorActiveStateChanged().RemoveAll(this);
+		ScenePtr->OnActorLabelChanged().RemoveAll(this);
 	}
 }
 
@@ -177,12 +193,28 @@ void FChaosVDWorldOutlinerMode::OnItemDoubleClick(FSceneOutlinerTreeItemPtr Item
 
 void FChaosVDWorldOutlinerMode::ProcessPendingHierarchyEvents()
 {
-	for (const TPair<FSceneOutlinerTreeItemID, FSceneOutlinerHierarchyChangedData>& PendingEvent : PendingOutlinerEventsMap)
+	const double StartTimeSeconds = FPlatformTime::Seconds();
+	double CurrentTimeSpentSeconds = 0.0;
+	int32 CurrentEventProcessedNum = 0;
+
+	constexpr double MaxUpdateBudgetSeconds = 0.001;
+
+	for (TMap<FSceneOutlinerTreeItemID, FSceneOutlinerHierarchyChangedData>::TIterator RemoveIterator = PendingOutlinerEventsMap.CreateIterator(); RemoveIterator; ++RemoveIterator)
 	{
-		Hierarchy->OnHierarchyChanged().Broadcast(PendingEvent.Value);
-	}
+		if (CurrentTimeSpentSeconds > MaxUpdateBudgetSeconds)
+		{
+			return;
+		}
+
+		// Only check the budget every 5 tasks as Getting the current time is a syscall and it is not free
+		if (CurrentEventProcessedNum % 5 == 0)
+		{
+			CurrentTimeSpentSeconds += FPlatformTime::Seconds() - StartTimeSeconds;
+		}
 
-	PendingOutlinerEventsMap.Reset();
+		Hierarchy->OnHierarchyChanged().Broadcast(RemoveIterator.Value());
+		RemoveIterator.RemoveCurrent();
+	}
 }
 
 bool FChaosVDWorldOutlinerMode::Tick(float DeltaTime)
@@ -212,6 +244,17 @@ bool FChaosVDWorldOutlinerMode::CanInteract(const ISceneOutlinerTreeItem& Item)
 	return true;
 }
 
+bool FChaosVDWorldOutlinerMode::CanPopulate() const
+{
+	if (TSharedPtr<FChaosVDPlaybackController> PlaybackControllerPtr = PlaybackController.Pin())
+	{
+		// Updating the scene outliner during playback it is very expensive and can tank framerate,
+		// as it need to re-build the hierarchy when things are added in ad removed. So if we are playing we want to pause any updates to the outliner
+		return !PlaybackControllerPtr->IsPlaying();
+	}
+	return true;
+}
+
 void FChaosVDWorldOutlinerMode::EnqueueAndCombineHierarchyEvent(const FSceneOutlinerTreeItemID& ItemID, const FSceneOutlinerHierarchyChangedData& EnventToProcess)
 {
 	if (FSceneOutlinerHierarchyChangedData* EventData = PendingOutlinerEventsMap.Find(ItemID))
@@ -223,7 +266,7 @@ void FChaosVDWorldOutlinerMode::EnqueueAndCombineHierarchyEvent(const FSceneOutl
 		PendingOutlinerEventsMap.Add(ItemID, EnventToProcess);
 	}	
 }
-void FChaosVDWorldOutlinerMode::HandleActorLabelChanged(AActor* ChangedActor)
+void FChaosVDWorldOutlinerMode::HandleActorLabelChanged(AChaosVDParticleActor* ChangedActor)
 {
 	if (!ensure(ChangedActor))
 	{
diff --git a/Engine/Plugins/ChaosVD/Source/ChaosVD/Private/ChaosVDWorldOutlinerMode.h b/Engine/Plugins/ChaosVD/Source/ChaosVD/Private/ChaosVDWorldOutlinerMode.h
--- a/Engine/Plugins/ChaosVD/Source/ChaosVD/Private/ChaosVDWorldOutlinerMode.h
+++ b/Engine/Plugins/ChaosVD/Source/ChaosVD/Private/ChaosVDWorldOutlinerMode.h
@@ -10,6 +10,7 @@
 #include "Containers/Map.h"
 #include "Templates/SharedPointer.h"
 
+class FChaosVDPlaybackController;
 class FChaosVDScene;
 
 /** Outliner Gutter used to override the behaviour of the visibility widget in the Scene outliner  */
@@ -35,12 +36,13 @@ public:
 	
 	FChaosVDActorTreeItem(AActor* InActor) : FActorTreeItem(InActor)
 	{
-		
-	};
+		UpdateDisplayString();
+	}
 
 	FChaosVDActorTreeItem(const FSceneOutlinerTreeItemType& TypeIn, AActor* InActor)
 		: FActorTreeItem(TypeIn, InActor)
 	{
+		UpdateDisplayString();
 	}
 
 	virtual bool GetVisibility() const override;
@@ -48,6 +50,8 @@ public:
 	virtual void OnVisibilityChanged(const bool bNewVisibility) override;
 
 	static const FSceneOutlinerTreeItemType Type;
+
+	virtual void UpdateDisplayString() override;
 };
 
 /** Actor Hierarchy used to override the what Tree items will be used for actors in CVD's Scene outliner  */
@@ -75,7 +79,7 @@ class FChaosVDWorldOutlinerMode : public FActorMode, public FChaosVDSceneSelecti
 {
 public:
 
-	FChaosVDWorldOutlinerMode(const FActorModeParams& InModeParams, TWeakPtr<FChaosVDScene> InScene);
+	FChaosVDWorldOutlinerMode(const FActorModeParams& InModeParams, TWeakPtr<FChaosVDScene> InScene, TWeakPtr<FChaosVDPlaybackController> InPlaybackController);
 
 	virtual ~FChaosVDWorldOutlinerMode() override;
 
@@ -95,20 +99,19 @@ public:
 
 	virtual TUniquePtr<ISceneOutlinerHierarchy> CreateHierarchy() override;
 
-	bool CanInteract(const ISceneOutlinerTreeItem& Item) const override;
+	virtual bool CanInteract(const ISceneOutlinerTreeItem& Item) const override;
+	virtual bool CanPopulate() const override;
 
 private:
 
 	void EnqueueAndCombineHierarchyEvent(const FSceneOutlinerTreeItemID& ItemID, const FSceneOutlinerHierarchyChangedData& EnventToProcess);
-	void HandleActorLabelChanged(AActor* ChangedActor);
+	void HandleActorLabelChanged(AChaosVDParticleActor* ChangedActor);
 	void HandleActorActiveStateChanged(AChaosVDParticleActor* ChangedActor);
 
 	virtual void HandlePostSelectionChange(const UTypedElementSelectionSet* ChangesSelectionSet) override;
 
-private:
-	FDelegateHandle ActorLabelChangedDelegateHandle;
-
 	TWeakPtr<FChaosVDScene> CVDScene;
+	TWeakPtr<FChaosVDPlaybackController> PlaybackController;
 
 	TMap<FSceneOutlinerTreeItemID, FSceneOutlinerHierarchyChangedData> PendingOutlinerEventsMap;
 };
diff --git a/Engine/Plugins/ChaosVD/Source/ChaosVD/Private/ChaosVDWorldOutlinerTab.cpp b/Engine/Plugins/ChaosVD/Source/ChaosVD/Private/ChaosVDWorldOutlinerTab.cpp
--- a/Engine/Plugins/ChaosVD/Source/ChaosVD/Private/ChaosVDWorldOutlinerTab.cpp
+++ b/Engine/Plugins/ChaosVD/Source/ChaosVD/Private/ChaosVDWorldOutlinerTab.cpp
@@ -2,6 +2,8 @@
 
 #include "ChaosVDWorldOutlinerTab.h"
 
+#include "ChaosVDEngine.h"
+#include "ChaosVDPlaybackController.h"
 #include "ChaosVDWorldOutlinerMode.h"
 #include "ChaosVDStyle.h"
 #include "Modules/ModuleManager.h"
@@ -9,11 +11,18 @@
 #include "SceneOutlinerPublicTypes.h"
 #include "Widgets/Docking/SDockTab.h"
 #include "Engine/World.h"
+#include "Widgets/SChaosVDMainTab.h"
 
 #define LOCTEXT_NAMESPACE "ChaosVisualDebugger"
 
 void FChaosVDWorldOutlinerTab::CreateWorldOutlinerWidget()
 {
+	TWeakPtr<FChaosVDPlaybackController> PlaybackController;
+	if (TSharedPtr<SChaosVDMainTab> MainTabPtr =  OwningTabWidget.Pin())
+	{
+		PlaybackController = MainTabPtr->GetChaosVDEngineInstance()->GetPlaybackController();
+	}
+
 	FSceneOutlinerInitializationOptions InitOptions;
 	InitOptions.bShowTransient = true;
 
@@ -26,7 +35,7 @@ void FChaosVDWorldOutlinerTab::CreateWorldOutlinerWidget()
 	InitOptions.ColumnMap.Add(FSceneOutlinerBuiltInColumnTypes::ActorInfo(), FSceneOutlinerColumnInfo(ESceneOutlinerColumnVisibility::Visible, 20, FCreateSceneOutlinerColumn(), true, TOptional<float>(), FSceneOutlinerBuiltInColumnTypes::ActorInfo_Localized()));
 	InitOptions.ColumnMap.Add(FChaosVDSceneOutlinerGutter::GetID(), FSceneOutlinerColumnInfo(ESceneOutlinerColumnVisibility::Visible, 0, FCreateSceneOutlinerColumn::CreateLambda([](ISceneOutliner& InSceneOutliner) { return MakeShareable(new FChaosVDSceneOutlinerGutter(InSceneOutliner)); })));
 		
-	FCreateSceneOutlinerMode ModeFactory = FCreateSceneOutlinerMode::CreateLambda([this](SSceneOutliner* Outliner)
+	FCreateSceneOutlinerMode ModeFactory = FCreateSceneOutlinerMode::CreateLambda([this, PlaybackController](SSceneOutliner* Outliner)
 	{
 		FActorModeParams ModeParams(Outliner);
 		ModeParams.SpecifiedWorldToDisplay = GetChaosVDWorld();
@@ -34,7 +43,7 @@ void FChaosVDWorldOutlinerTab::CreateWorldOutlinerWidget()
 		ModeParams.bCanInteractWithSelectableActorsOnly = false;
 
 		// The mode is deleted by the Outliner when it is destroyed 
-		return new FChaosVDWorldOutlinerMode(ModeParams, GetChaosVDScene());
+		return new FChaosVDWorldOutlinerMode(ModeParams, GetChaosVDScene(), PlaybackController);
 	});
 	InitOptions.ModeFactory = ModeFactory;
 	
diff --git a/Engine/Plugins/ChaosVD/Source/ChaosVD/Private/Components/ChaosVDInstancedStaticMeshComponent.cpp b/Engine/Plugins/ChaosVD/Source/ChaosVD/Private/Components/ChaosVDInstancedStaticMeshComponent.cpp
--- a/Engine/Plugins/ChaosVD/Source/ChaosVD/Private/Components/ChaosVDInstancedStaticMeshComponent.cpp
+++ b/Engine/Plugins/ChaosVD/Source/ChaosVD/Private/Components/ChaosVDInstancedStaticMeshComponent.cpp
@@ -134,6 +134,8 @@ void UChaosVDInstancedStaticMeshComponent::Reset()
 	ComponentEmptyDelegate = FChaosVDMeshComponentEmptyDelegate();
 	CurrentMaterial = nullptr;
 
+	EmptyOverrideMaterials();
+
 	CurrentInstanceHandles.Reset();
 	CurrentGeometryKey = 0;
 }
diff --git a/Engine/Plugins/ChaosVD/Source/ChaosVD/Private/Components/ChaosVDInstancedStaticMeshComponent.h b/Engine/Plugins/ChaosVD/Source/ChaosVD/Private/Components/ChaosVDInstancedStaticMeshComponent.h
--- a/Engine/Plugins/ChaosVD/Source/ChaosVD/Private/Components/ChaosVDInstancedStaticMeshComponent.h
+++ b/Engine/Plugins/ChaosVD/Source/ChaosVD/Private/Components/ChaosVDInstancedStaticMeshComponent.h
@@ -15,6 +15,7 @@ class UChaosVDInstancedStaticMeshComponent : public UInstancedStaticMeshComponen
 	UChaosVDInstancedStaticMeshComponent(const FObjectInitializer& ObjectInitializer)
 	: Super(ObjectInitializer)
 	{
+		SetRemoveSwap();
 		SetCanEverAffectNavigation(false);
 		bHasPerInstanceHitProxies = true;
 	}
diff --git a/Engine/Plugins/ChaosVD/Source/ChaosVD/Private/ObjectsWaitingGeometryList.cpp b/Engine/Plugins/ChaosVD/Source/ChaosVD/Private/ObjectsWaitingGeometryList.cpp
--- a/Engine/Plugins/ChaosVD/Source/ChaosVD/Private/ObjectsWaitingGeometryList.cpp
+++ b/Engine/Plugins/ChaosVD/Source/ChaosVD/Private/ObjectsWaitingGeometryList.cpp
@@ -1,9 +1,3 @@
 ﻿// Copyright Epic Games, Inc. All Rights Reserved.
 
 #include "ObjectsWaitingGeometryList.h"
-
-int32 FChaosVDWailingObjectListCVars::NumOfGTTaskToProcessPerTick = 200;
-FAutoConsoleVariableRef FChaosVDWailingObjectListCVars::CVarChaosVDGeometryToProcessPerTick(
-	TEXT("p.Chaos.VD.Tool.GeometryToProcessPerTick"),
-	FChaosVDWailingObjectListCVars::NumOfGTTaskToProcessPerTick,
-	TEXT("Number of generated geometry to process each tick when loading a teace file in the CVD tool."));
diff --git a/Engine/Plugins/ChaosVD/Source/ChaosVD/Private/ObjectsWaitingGeometryList.h b/Engine/Plugins/ChaosVD/Source/ChaosVD/Private/ObjectsWaitingGeometryList.h
--- a/Engine/Plugins/ChaosVD/Source/ChaosVD/Private/ObjectsWaitingGeometryList.h
+++ b/Engine/Plugins/ChaosVD/Source/ChaosVD/Private/ObjectsWaitingGeometryList.h
@@ -10,21 +10,15 @@
 
 class FTextFormat;
 
-struct FChaosVDWailingObjectListCVars
-{
-	static int32 NumOfGTTaskToProcessPerTick;
-	static FAutoConsoleVariableRef CVarChaosVDGeometryToProcessPerTick;
-};
-
 /** Object that handles objects waiting for geometry to be ready to perform a desired operation
  * The operation has to run on the Game Thread, and there is a limit on how many operations will be processed per tick
  */
-template<typename TObjectToProcess>
+template<typename ObjectType>
 class FObjectsWaitingGeometryList
 {
 
 public:
-	 FObjectsWaitingGeometryList(const TFunction<bool(uint32, TObjectToProcess)>& InObjectProcessorCallback, const FTextFormat& InProgressNotificationNameFormat, const TFunction<bool(uint32)>& InShouldProcessObjectsOverride) :
+	 FObjectsWaitingGeometryList(const TFunction<bool(uint32, ObjectType)>& InObjectProcessorCallback, const FTextFormat& InProgressNotificationNameFormat, const TFunction<bool(uint32)>& InShouldProcessObjectsOverride) :
 		  ObjectProcessorCallback(InObjectProcessorCallback)
 		, ShouldProcessObjectsForKeyOverride(InShouldProcessObjectsOverride)
 		, AsyncProgressNotification(InProgressNotificationNameFormat)
@@ -32,13 +26,13 @@ public:
 	}
 
 	/** Adds an object to process when the provided geometry key is ready */
-	void AddObject(uint32 GeometryKey, TObjectToProcess ObjectToProcess);
+	void AddObject(uint32 GeometryKey, ObjectType ObjectToProcess);
 
 	/** Removed an object to process */
-	void RemoveObject(uint32 GeometryKey, TObjectToProcess& OutToProcess);
+	void RemoveObject(uint32 GeometryKey, ObjectType& OutToProcess);
 
 	/** Evaluates the current geometry available, and process any objects on the waiting list for it */
-	bool ProcessWaitingObjects(int32& OutCurrentElementsProcessed);
+	bool ProcessWaitingObjects(float TimeBudgetSeconds);
 
 private:
 
@@ -51,23 +45,23 @@ private:
 		return true;
 	}
 
-	TFunction<bool(uint32, TObjectToProcess)> ObjectProcessorCallback;
+	TFunction<bool(uint32, ObjectType)> ObjectProcessorCallback;
 
 	TFunction<bool(uint32)> ShouldProcessObjectsForKeyOverride;
 
-	TMap<uint32, TArray<TObjectToProcess>> WaitingObjectsByGeometryKey;
+	TMap<uint32, TArray<ObjectType>> WaitingObjectsByGeometryKey;
 
 	FAsyncCompilationNotification AsyncProgressNotification;
 
 	int32 QueuedObjectsToProcessNum = 0;
 };
 
-template <typename TObjectToProcess>
-void FObjectsWaitingGeometryList<TObjectToProcess>::AddObject(uint32 GeometryKey, TObjectToProcess ObjectToProcess)
+template <typename ObjectType>
+void FObjectsWaitingGeometryList<ObjectType>::AddObject(uint32 GeometryKey, ObjectType ObjectToProcess)
 {
 	check(IsInGameThread());
 
-	if (TArray<TObjectToProcess>* ObjectsInQueueForKey = WaitingObjectsByGeometryKey.Find(GeometryKey))
+	if (TArray<ObjectType>* ObjectsInQueueForKey = WaitingObjectsByGeometryKey.Find(GeometryKey))
 	{
 		ObjectsInQueueForKey->Add(ObjectToProcess);
 	}
@@ -79,17 +73,17 @@ void FObjectsWaitingGeometryList<TObjectToProcess>::AddObject(uint32 GeometryKey
 	++QueuedObjectsToProcessNum;
 }
 
-template <typename TObjectToProcess>
-void FObjectsWaitingGeometryList<TObjectToProcess>::RemoveObject(uint32 GeometryKey, TObjectToProcess& OutToProcess)
+template <typename ObjectType>
+void FObjectsWaitingGeometryList<ObjectType>::RemoveObject(uint32 GeometryKey, ObjectType& OutToProcess)
 {
 	check(IsInGameThread());
 
-	if (TArray<TObjectToProcess>* QueueObjectsForKey = WaitingObjectsByGeometryKey.Find(GeometryKey))
+	if (TArray<ObjectType>* QueueObjectsForKey = WaitingObjectsByGeometryKey.Find(GeometryKey))
 	{
-		TArray<TObjectToProcess>& QueueObjectsForKeyRef = *QueueObjectsForKey;
-		for(typename TArray<TObjectToProcess>::TIterator ObjectToRemoveIterator = QueueObjectsForKeyRef.CreateIterator();  ObjectToRemoveIterator; ++ObjectToRemoveIterator)
+		TArray<ObjectType>& QueueObjectsForKeyRef = *QueueObjectsForKey;
+		for(typename TArray<ObjectType>::TIterator ObjectToRemoveIterator = QueueObjectsForKeyRef.CreateIterator();  ObjectToRemoveIterator; ++ObjectToRemoveIterator)
 		{
-			TObjectToProcess& ObjectToRemove = *ObjectToRemoveIterator;
+			ObjectType& ObjectToRemove = *ObjectToRemoveIterator;
 
 			if (ObjectToRemove == OutToProcess)
 			{
@@ -101,25 +95,38 @@ void FObjectsWaitingGeometryList<TObjectToProcess>::RemoveObject(uint32 Geometry
 	}
 }
 
-template <typename TObjectToProcess>
-bool FObjectsWaitingGeometryList<TObjectToProcess>::ProcessWaitingObjects(int32& OutCurrentElementsProcessed)
+template <typename ObjectType>
+bool FObjectsWaitingGeometryList<ObjectType>::ProcessWaitingObjects(float TimeBudgetSeconds)
 {
 	AsyncProgressNotification.Update(QueuedObjectsToProcessNum);
 
+	const double StartTimeSeconds = FPlatformTime::Seconds();
+	double CurrentTimeSpentSeconds = 0.0;
+	int32 CurrentObjectsProcessedNum = 0;
+
 	bool bCanContinueProcessing = true;
-	for (typename TMap<uint32, TArray<TObjectToProcess>>::TIterator QueuedObjectRemoveIterator = WaitingObjectsByGeometryKey.CreateIterator(); QueuedObjectRemoveIterator; ++QueuedObjectRemoveIterator)
+	for (typename TMap<uint32, TArray<ObjectType>>::TIterator QueuedObjectRemoveIterator = WaitingObjectsByGeometryKey.CreateIterator(); QueuedObjectRemoveIterator; ++QueuedObjectRemoveIterator)
 	{
 		const uint32 GeometryKey = QueuedObjectRemoveIterator.Key();
 
 		if (ShouldProcessObjectsForKey(GeometryKey))
 		{
-			for (typename TArray<TObjectToProcess>::TIterator ObjectRemoveIterator = QueuedObjectRemoveIterator.Value().CreateIterator(); ObjectRemoveIterator; ++ObjectRemoveIterator)
+			for (typename TArray<ObjectType>::TIterator ObjectRemoveIterator = QueuedObjectRemoveIterator.Value().CreateIterator(); ObjectRemoveIterator; ++ObjectRemoveIterator)
 			{
 				if (ObjectProcessorCallback && ObjectProcessorCallback(GeometryKey, *ObjectRemoveIterator))
 				{
-					OutCurrentElementsProcessed++;
+					CurrentObjectsProcessedNum++;
 					--QueuedObjectsToProcessNum;
-					bCanContinueProcessing = OutCurrentElementsProcessed < FChaosVDWailingObjectListCVars::NumOfGTTaskToProcessPerTick;
+					bCanContinueProcessing = CurrentTimeSpentSeconds < TimeBudgetSeconds;
+					
+					AsyncProgressNotification.Update(QueuedObjectsToProcessNum);
+
+					// Only check the budget every 5 tasks as Getting the current time is a syscall and it is not free
+					if (CurrentObjectsProcessedNum % 5 == 0)
+					{
+						CurrentTimeSpentSeconds += FPlatformTime::Seconds() - StartTimeSeconds;
+					}
+	
 					ObjectRemoveIterator.RemoveCurrent();
 				}
 			}
@@ -143,3 +150,77 @@ bool FObjectsWaitingGeometryList<TObjectToProcess>::ProcessWaitingObjects(int32&
 
 	return bCanContinueProcessing;
 }
+
+/** Object that allows adding objects that need processing to a queue, and provide a processor callback.
+ * When there are pending objects in the queue, a progress notification will be shown in the editor
+ */
+template<typename ObjectType>
+class FObjectsWaitingProcessingQueue
+{
+
+public:
+	FObjectsWaitingProcessingQueue(const TFunction<bool(ObjectType)>& InObjectProcessorCallback, const FTextFormat& InProgressNotificationNameFormat) :
+		 ObjectProcessorCallback(InObjectProcessorCallback)
+	   , AsyncProgressNotification(InProgressNotificationNameFormat)
+	{
+	}
+
+	~FObjectsWaitingProcessingQueue();
+
+	void EnqueueObject(ObjectType ObjectToProcess);
+
+	void ProcessWaitingTasks(float TimeBudgetSeconds);
+
+private:
+
+	TFunction<bool(ObjectType)> ObjectProcessorCallback;
+
+	TQueue<ObjectType> WaitingObjectQueue;
+
+	FAsyncCompilationNotification AsyncProgressNotification;
+
+	int32 QueuedObjectsToProcessNum = 0;
+};
+
+
+template <typename ObjectType>
+FObjectsWaitingProcessingQueue<ObjectType>::~FObjectsWaitingProcessingQueue()
+{
+	WaitingObjectQueue.Empty();
+	AsyncProgressNotification.Update(0);
+}
+
+template <typename ObjectType>
+void FObjectsWaitingProcessingQueue<ObjectType>::EnqueueObject(ObjectType ObjectToProcess)
+{
+	WaitingObjectQueue.Enqueue(ObjectToProcess);
+	QueuedObjectsToProcessNum++;
+}
+
+template <typename ObjectType>
+void FObjectsWaitingProcessingQueue<ObjectType>::ProcessWaitingTasks(float TimeBudgetSeconds)
+{
+	AsyncProgressNotification.Update(QueuedObjectsToProcessNum);
+
+	const double StartTimeSeconds = FPlatformTime::Seconds();
+	double CurrentTimeSpentSeconds = 0.0;
+	int32 CurrentTasksProcessedNum = 0;
+
+	ObjectType QueuedTask;
+	while (TimeBudgetSeconds > CurrentTimeSpentSeconds && WaitingObjectQueue.Dequeue(QueuedTask))
+	{
+		// Only check the budget every 5 tasks as Getting the current time is a syscall and it is not free
+		if (CurrentTasksProcessedNum % 5 == 0)
+		{
+			CurrentTimeSpentSeconds += FPlatformTime::Seconds() - StartTimeSeconds;
+		}
+
+		if (ObjectProcessorCallback(QueuedTask))
+		{
+			CurrentTasksProcessedNum++;
+			--QueuedObjectsToProcessNum;
+		}
+
+		AsyncProgressNotification.Update(QueuedObjectsToProcessNum);
+	}
+}
diff --git a/Engine/Plugins/ChaosVD/Source/ChaosVD/Private/Widgets/SChaosVDDetailsView.cpp b/Engine/Plugins/ChaosVD/Source/ChaosVD/Private/Widgets/SChaosVDDetailsView.cpp
--- a/Engine/Plugins/ChaosVD/Source/ChaosVD/Private/Widgets/SChaosVDDetailsView.cpp
+++ b/Engine/Plugins/ChaosVD/Source/ChaosVD/Private/Widgets/SChaosVDDetailsView.cpp
@@ -45,7 +45,7 @@ TSharedPtr<IDetailsView> SChaosVDDetailsView::CreateObjectDetailsView()
 	DetailsViewArgs.bUpdatesFromSelection = false;
 	DetailsViewArgs.bLockable = true;
 	DetailsViewArgs.bAllowFavoriteSystem = false;
-	DetailsViewArgs.NameAreaSettings = FDetailsViewArgs::ObjectsUseNameArea | FDetailsViewArgs::ComponentsAndActorsUseNameArea;
+	DetailsViewArgs.NameAreaSettings = FDetailsViewArgs::HideNameArea;
 	DetailsViewArgs.bCustomFilterAreaLocation = false;
 	DetailsViewArgs.bShowSectionSelector = false;
 	DetailsViewArgs.bShowScrollBar = false;
diff --git a/Engine/Plugins/ChaosVD/Source/ChaosVD/Private/Widgets/SChaosVDPlaybackViewport.cpp b/Engine/Plugins/ChaosVD/Source/ChaosVD/Private/Widgets/SChaosVDPlaybackViewport.cpp
--- a/Engine/Plugins/ChaosVD/Source/ChaosVD/Private/Widgets/SChaosVDPlaybackViewport.cpp
+++ b/Engine/Plugins/ChaosVD/Source/ChaosVD/Private/Widgets/SChaosVDPlaybackViewport.cpp
@@ -260,12 +260,12 @@ void SChaosVDPlaybackViewport::HandlePostSelectionChange(const UTypedElementSele
 
 void SChaosVDPlaybackViewport::OnPlaybackSceneUpdated()
 {
-	PlaybackViewportClient->bNeedsRedraw = true;
+	PlaybackViewportClient->HandleCVDSceneUpdated();
 }
 
 void SChaosVDPlaybackViewport::OnSolverVisibilityUpdated(int32 SolverID, bool bNewVisibility)
 {
-	PlaybackViewportClient->bNeedsRedraw = true;	
+	PlaybackViewportClient->HandleCVDSceneUpdated();	
 }
 
 void SChaosVDPlaybackViewport::BindToSceneUpdateEvents()
