[ChaosVD] Fix CIS warning in FChaosVDWorldOutlinerMode

Adding missing increement in a counter

diff --git a/Engine/Plugins/ChaosVD/Source/ChaosVD/Private/ChaosVDWorldOutlinerMode.cpp b/Engine/Plugins/ChaosVD/Source/ChaosVD/Private/ChaosVDWorldOutlinerMode.cpp
--- a/Engine/Plugins/ChaosVD/Source/ChaosVD/Private/ChaosVDWorldOutlinerMode.cpp
+++ b/Engine/Plugins/ChaosVD/Source/ChaosVD/Private/ChaosVDWorldOutlinerMode.cpp
@@ -197,7 +197,7 @@ void FChaosVDWorldOutlinerMode::ProcessPendingHierarchyEvents()
 	double CurrentTimeSpentSeconds = 0.0;
 	int32 CurrentEventProcessedNum = 0;
 
-	constexpr double MaxUpdateBudgetSeconds = 0.001;
+	constexpr double MaxUpdateBudgetSeconds = 0.002;
 
 	for (TMap<FSceneOutlinerTreeItemID, FSceneOutlinerHierarchyChangedData>::TIterator RemoveIterator = PendingOutlinerEventsMap.CreateIterator(); RemoveIterator; ++RemoveIterator)
 	{
@@ -214,6 +214,7 @@ void FChaosVDWorldOutlinerMode::ProcessPendingHierarchyEvents()
 
 		Hierarchy->OnHierarchyChanged().Broadcast(RemoveIterator.Value());
 		RemoveIterator.RemoveCurrent();
+		CurrentEventProcessedNum++;
 	}
 }
 
