[ChaosVD] Fix an issue introduced in my previous CL where colorization is not applied to heightfields

Heightfields are still represented with static mesh components, and it has its own system to handle materials. This CL removes the call to create/apply a material instance asynchronously on them for now.

diff --git a/Engine/Plugins/ChaosVD/Source/ChaosVD/Private/ChaosVDGeometryBuilder.h b/Engine/Plugins/ChaosVD/Source/ChaosVD/Private/ChaosVDGeometryBuilder.h
--- a/Engine/Plugins/ChaosVD/Source/ChaosVD/Private/ChaosVDGeometryBuilder.h
+++ b/Engine/Plugins/ChaosVD/Source/ChaosVD/Private/ChaosVDGeometryBuilder.h
@@ -513,11 +513,6 @@ ComponentType* FChaosVDGeometryBuilder::GetAvailableMeshComponent(const TSharedP
 			return nullptr;
 		}
 
-		if (!MeshComponent->HasOverrideMaterials())
-		{
-			MeshComponentsWaitingForMaterial->EnqueueObject(MeshComponent);
-		}
-
 		bOutIsNewComponent = true;
 	}
 
