[ChaosVD] Added size reporting stats of how much data was loaded for each type. This will be useful when we start working on reducing the mount of data we serialize

diff --git a/Engine/Plugins/ChaosVD/Source/ChaosVD/Private/Trace/ChaosVDTraceAnalyzer.cpp b/Engine/Plugins/ChaosVD/Source/ChaosVD/Private/Trace/ChaosVDTraceAnalyzer.cpp
--- a/Engine/Plugins/ChaosVD/Source/ChaosVD/Private/Trace/ChaosVDTraceAnalyzer.cpp
+++ b/Engine/Plugins/ChaosVD/Source/ChaosVD/Private/Trace/ChaosVDTraceAnalyzer.cpp
@@ -35,6 +35,8 @@ void FChaosVDTraceAnalyzer::OnAnalysisBegin(const FOnAnalysisContext& Context)
 
 void FChaosVDTraceAnalyzer::OnAnalysisEnd()
 {
+	ChaosVDTraceProvider->HandleAnalysisComplete();
+
 	OnAnalysisComplete().Broadcast();
 }
 
diff --git a/Engine/Plugins/ChaosVD/Source/ChaosVD/Private/Trace/ChaosVDTraceAnalyzer.h b/Engine/Plugins/ChaosVD/Source/ChaosVD/Private/Trace/ChaosVDTraceAnalyzer.h
--- a/Engine/Plugins/ChaosVD/Source/ChaosVD/Private/Trace/ChaosVDTraceAnalyzer.h
+++ b/Engine/Plugins/ChaosVD/Source/ChaosVD/Private/Trace/ChaosVDTraceAnalyzer.h
@@ -21,7 +21,7 @@ DECLARE_MULTICAST_DELEGATE(FChaosVDTraceAnalysisComplete)
 class FChaosVDTraceAnalyzer final : public UE::Trace::IAnalyzer
 {
 public:
-	FChaosVDTraceAnalyzer(TraceServices::IAnalysisSession& Session, const TSharedPtr<FChaosVDTraceProvider>& InChaosVDTraceProvider)
+	FChaosVDTraceAnalyzer(TraceServices::IAnalysisSession& Session, const TSharedRef<FChaosVDTraceProvider>& InChaosVDTraceProvider)
 	: Session(Session),
 	ChaosVDTraceProvider(InChaosVDTraceProvider)
 	{
@@ -54,7 +54,7 @@ private:
 
 	TraceServices::IAnalysisSession& Session;
 
-	TSharedPtr<FChaosVDTraceProvider> ChaosVDTraceProvider;
+	TSharedRef<FChaosVDTraceProvider> ChaosVDTraceProvider;
 	
 	FChaosVDTraceAnalysisComplete ChaosVDTraceAnalysisCompleteDelegate;
 };
diff --git a/Engine/Plugins/ChaosVD/Source/ChaosVD/Private/Trace/ChaosVDTraceModule.cpp b/Engine/Plugins/ChaosVD/Source/ChaosVD/Private/Trace/ChaosVDTraceModule.cpp
--- a/Engine/Plugins/ChaosVD/Source/ChaosVD/Private/Trace/ChaosVDTraceModule.cpp
+++ b/Engine/Plugins/ChaosVD/Source/ChaosVD/Private/Trace/ChaosVDTraceModule.cpp
@@ -16,7 +16,7 @@ void FChaosVDTraceModule::GetModuleInfo(TraceServices::FModuleInfo& OutModuleInf
 
 void FChaosVDTraceModule::OnAnalysisBegin(TraceServices::IAnalysisSession& InSession)
 {
-	const TSharedPtr<FChaosVDTraceProvider> Provider = MakeShared<FChaosVDTraceProvider>(InSession);
+	const TSharedRef<FChaosVDTraceProvider> Provider = MakeShared<FChaosVDTraceProvider>(InSession);
 
 	InSession.AddProvider(FChaosVDTraceProvider::ProviderName, Provider);
 	InSession.AddAnalyzer(new FChaosVDTraceAnalyzer(InSession, Provider));
diff --git a/Engine/Plugins/ChaosVD/Source/ChaosVD/Private/Trace/ChaosVDTraceProvider.cpp b/Engine/Plugins/ChaosVD/Source/ChaosVD/Private/Trace/ChaosVDTraceProvider.cpp
--- a/Engine/Plugins/ChaosVD/Source/ChaosVD/Private/Trace/ChaosVDTraceProvider.cpp
+++ b/Engine/Plugins/ChaosVD/Source/ChaosVD/Private/Trace/ChaosVDTraceProvider.cpp
@@ -19,7 +19,7 @@
 #include "Trace/DataProcessors/ChaosVDTraceImplicitObjectProcessor.h"
 #include "Trace/DataProcessors/ChaosVDTraceParticleDataProcessor.h"
 #include "Trace/DataProcessors/ChaosVDArchiveHeaderProcessor.h"
-#include "Trace/DataProcessors/IChaosVDDataProcessor.h"
+#include "Trace/DataProcessors/ChaosVDDataProcessorBase.h"
 
 FName FChaosVDTraceProvider::ProviderName("ChaosVDProvider");
 
@@ -153,12 +153,13 @@ bool FChaosVDTraceProvider::ProcessBinaryData(const int32 DataID)
 				RawData = &UnprocessedData->RawData;
 			}
 
-			if (TSharedPtr<IChaosVDDataProcessor>* DataProcessorPtrPtr = RegisteredDataProcessors.Find(UnprocessedData->TypeName))
+			if (TSharedPtr<FChaosVDDataProcessorBase>* DataProcessorPtrPtr = RegisteredDataProcessors.Find(UnprocessedData->TypeName))
 			{
-				if (TSharedPtr<IChaosVDDataProcessor> DataProcessorPtr = *DataProcessorPtrPtr)
+				if (TSharedPtr<FChaosVDDataProcessorBase> DataProcessorPtr = *DataProcessorPtrPtr)
 				{
 					if (ensure(DataProcessorPtr->ProcessRawData(*RawData)))
 					{
+						UnprocessedDataByID.Remove(DataID);
 						return true;
 					}
 					else
@@ -172,6 +173,8 @@ bool FChaosVDTraceProvider::ProcessBinaryData(const int32 DataID)
 				UE_LOG(LogChaosVDEditor, Warning, TEXT("[%s] Data processor for type [%s] not found"), ANSI_TO_TCHAR(__FUNCTION__), *UnprocessedData->TypeName);
 			}
 		}
+
+		UnprocessedDataByID.Remove(DataID);
 	}
 
 	return false;
@@ -182,11 +185,33 @@ TSharedPtr<FChaosVDRecording> FChaosVDTraceProvider::GetRecordingForSession() co
 	return InternalRecording;
 }
 
-void FChaosVDTraceProvider::RegisterDataProcessor(TSharedPtr<IChaosVDDataProcessor> InDataProcessor)
+void FChaosVDTraceProvider::RegisterDataProcessor(TSharedPtr<FChaosVDDataProcessorBase> InDataProcessor)
 {
 	RegisteredDataProcessors.Add(InDataProcessor->GetCompatibleTypeName(), InDataProcessor);
 }
 
+void FChaosVDTraceProvider::HandleAnalysisComplete()
+{
+	UnprocessedDataByID.Reset();
+
+	UE_LOG(LogChaosVDEditor, Log, TEXT("Trace Analysis complete for session [%s] | Calculating data loaded stats..."), Session.GetName());
+
+	static const FNumberFormattingOptions SizeFormattingOptions = FNumberFormattingOptions().SetMinimumFractionalDigits(2).SetMaximumFractionalDigits(2);
+
+	uint64 TotalBytes = 0;
+	for (const TPair<FStringView, TSharedPtr<FChaosVDDataProcessorBase>>& DataProcessor : RegisteredDataProcessors)
+	{
+		if (DataProcessor.Value)
+		{
+			uint64 ProcessedBytes = DataProcessor.Value->GetProcessedBytes();
+			TotalBytes += ProcessedBytes;
+			UE_LOG(LogChaosVDEditor, Log, TEXT("Data loaded for type [%s]  => [%s] "), DataProcessor.Key.IsEmpty() ? TEXT("Invalid") : DataProcessor.Key.GetData(), *FText::AsMemory(ProcessedBytes, &SizeFormattingOptions,nullptr, EMemoryUnitStandard::IEC).ToString());
+		}
+	}
+
+	UE_LOG(LogChaosVDEditor, Log, TEXT("Total size of loaded data => [%s]"), *FText::AsMemory(TotalBytes, &SizeFormattingOptions,nullptr, EMemoryUnitStandard::IEC).ToString());
+}
+
 void FChaosVDTraceProvider::RegisterDefaultDataProcessorsIfNeeded()
 {
 	if (bDefaultDataProcessorsRegistered)
diff --git a/Engine/Plugins/ChaosVD/Source/ChaosVD/Private/Trace/ChaosVDTraceProvider.h b/Engine/Plugins/ChaosVD/Source/ChaosVD/Private/Trace/ChaosVDTraceProvider.h
--- a/Engine/Plugins/ChaosVD/Source/ChaosVD/Private/Trace/ChaosVDTraceProvider.h
+++ b/Engine/Plugins/ChaosVD/Source/ChaosVD/Private/Trace/ChaosVDTraceProvider.h
@@ -12,6 +12,8 @@
 #include "Templates/SharedPointer.h"
 #include "TraceServices/Model/AnalysisSession.h"
 
+class FChaosVDDataProcessorBase;
+
 namespace Chaos::VisualDebugger
 {
 	class FChaosVDSerializableNameTable;
@@ -73,7 +75,9 @@ public:
 
 	TSharedPtr<FChaosVDRecording> GetRecordingForSession() const;
 
-	void RegisterDataProcessor(TSharedPtr<IChaosVDDataProcessor> InDataProcessor);
+	void RegisterDataProcessor(TSharedPtr<FChaosVDDataProcessorBase> InDataProcessor);
+
+	void HandleAnalysisComplete();
 
 private:
 
@@ -88,7 +92,7 @@ private:
 
 	TMap<int32, TSharedPtr<FChaosVDBinaryDataContainer>> UnprocessedDataByID;
 
-	TMap<FStringView, TSharedPtr<IChaosVDDataProcessor>> RegisteredDataProcessors;
+	TMap<FStringView, TSharedPtr<FChaosVDDataProcessorBase>> RegisteredDataProcessors;
 
 	TMap<int32, FChaosVDSolverFrameData> CurrentSolverFramesByID;
 
diff --git a/Engine/Plugins/ChaosVD/Source/ChaosVD/Private/Trace/DataProcessors/ChaosVDArchiveHeaderProcessor.cpp b/Engine/Plugins/ChaosVD/Source/ChaosVD/Private/Trace/DataProcessors/ChaosVDArchiveHeaderProcessor.cpp
--- a/Engine/Plugins/ChaosVD/Source/ChaosVD/Private/Trace/DataProcessors/ChaosVDArchiveHeaderProcessor.cpp
+++ b/Engine/Plugins/ChaosVD/Source/ChaosVD/Private/Trace/DataProcessors/ChaosVDArchiveHeaderProcessor.cpp
@@ -5,12 +5,14 @@
 #include "ChaosVDModule.h"
 
 
-FChaosVDArchiveHeaderProcessor::FChaosVDArchiveHeaderProcessor() : IChaosVDDataProcessor(Chaos::VisualDebugger::FChaosVDArchiveHeader::WrapperTypeName)
+FChaosVDArchiveHeaderProcessor::FChaosVDArchiveHeaderProcessor() : FChaosVDDataProcessorBase(Chaos::VisualDebugger::FChaosVDArchiveHeader::WrapperTypeName)
 {
 }
 
 bool FChaosVDArchiveHeaderProcessor::ProcessRawData(const TArray<uint8>& InData)
 {
+	FChaosVDDataProcessorBase::ProcessRawData(InData);
+
 	TSharedPtr<FChaosVDTraceProvider> ProviderSharedPtr = TraceProvider.Pin();
 	if (!ensure(ProviderSharedPtr.IsValid()))
 	{
diff --git a/Engine/Plugins/ChaosVD/Source/ChaosVD/Private/Trace/DataProcessors/ChaosVDArchiveHeaderProcessor.h b/Engine/Plugins/ChaosVD/Source/ChaosVD/Private/Trace/DataProcessors/ChaosVDArchiveHeaderProcessor.h
--- a/Engine/Plugins/ChaosVD/Source/ChaosVD/Private/Trace/DataProcessors/ChaosVDArchiveHeaderProcessor.h
+++ b/Engine/Plugins/ChaosVD/Source/ChaosVD/Private/Trace/DataProcessors/ChaosVDArchiveHeaderProcessor.h
@@ -1,12 +1,12 @@
 // Copyright Epic Games, Inc. All Rights Reserved.
 
 #pragma once
-#include "Trace/DataProcessors/IChaosVDDataProcessor.h"
+#include "Trace/DataProcessors/ChaosVDDataProcessorBase.h"
 
 /**
  * Data processor implementation that is able to deserialize traced Archive headers
  */
-class FChaosVDArchiveHeaderProcessor final : public IChaosVDDataProcessor
+class FChaosVDArchiveHeaderProcessor final : public FChaosVDDataProcessorBase
 {
 public:
 	explicit FChaosVDArchiveHeaderProcessor();
diff --git a/Engine/Plugins/ChaosVD/Source/ChaosVD/Private/Trace/DataProcessors/ChaosVDCharacterGroundConstraintDataProcessor.cpp b/Engine/Plugins/ChaosVD/Source/ChaosVD/Private/Trace/DataProcessors/ChaosVDCharacterGroundConstraintDataProcessor.cpp
--- a/Engine/Plugins/ChaosVD/Source/ChaosVD/Private/Trace/DataProcessors/ChaosVDCharacterGroundConstraintDataProcessor.cpp
+++ b/Engine/Plugins/ChaosVD/Source/ChaosVD/Private/Trace/DataProcessors/ChaosVDCharacterGroundConstraintDataProcessor.cpp
@@ -7,12 +7,14 @@
 
 
 FChaosVDCharacterGroundConstraintDataProcessor::FChaosVDCharacterGroundConstraintDataProcessor()
-	: IChaosVDDataProcessor(FChaosVDCharacterGroundConstraint::WrapperTypeName)
+	: FChaosVDDataProcessorBase(FChaosVDCharacterGroundConstraint::WrapperTypeName)
 {
 }
 
 bool FChaosVDCharacterGroundConstraintDataProcessor::ProcessRawData(const TArray<uint8>& InData)
 {
+	FChaosVDDataProcessorBase::ProcessRawData(InData);
+
 	TSharedPtr<FChaosVDTraceProvider> ProviderSharedPtr = TraceProvider.Pin();
 	if (!ensure(ProviderSharedPtr.IsValid()))
 	{
diff --git a/Engine/Plugins/ChaosVD/Source/ChaosVD/Private/Trace/DataProcessors/ChaosVDCharacterGroundConstraintDataProcessor.h b/Engine/Plugins/ChaosVD/Source/ChaosVD/Private/Trace/DataProcessors/ChaosVDCharacterGroundConstraintDataProcessor.h
--- a/Engine/Plugins/ChaosVD/Source/ChaosVD/Private/Trace/DataProcessors/ChaosVDCharacterGroundConstraintDataProcessor.h
+++ b/Engine/Plugins/ChaosVD/Source/ChaosVD/Private/Trace/DataProcessors/ChaosVDCharacterGroundConstraintDataProcessor.h
@@ -3,7 +3,7 @@
 #pragma once
 
 #pragma once
-#include "IChaosVDDataProcessor.h"
+#include "ChaosVDDataProcessorBase.h"
 
 struct FChaosVDSolverFrameData;
 struct FChaosVDConstraint;
@@ -11,7 +11,7 @@ struct FChaosVDConstraint;
 /**
  * Data processor implementation that is able to deserialize traced character ground constraints
  */
-class FChaosVDCharacterGroundConstraintDataProcessor final : public IChaosVDDataProcessor
+class FChaosVDCharacterGroundConstraintDataProcessor final : public FChaosVDDataProcessorBase
 {
 public:
 	explicit FChaosVDCharacterGroundConstraintDataProcessor();
diff --git a/Engine/Plugins/ChaosVD/Source/ChaosVD/Private/Trace/DataProcessors/ChaosVDConstraintDataProcessor.cpp b/Engine/Plugins/ChaosVD/Source/ChaosVD/Private/Trace/DataProcessors/ChaosVDConstraintDataProcessor.cpp
--- a/Engine/Plugins/ChaosVD/Source/ChaosVD/Private/Trace/DataProcessors/ChaosVDConstraintDataProcessor.cpp
+++ b/Engine/Plugins/ChaosVD/Source/ChaosVD/Private/Trace/DataProcessors/ChaosVDConstraintDataProcessor.cpp
@@ -7,12 +7,14 @@
 #include "DataWrappers/ChaosVDCollisionDataWrappers.h"
 #include "Trace/ChaosVDTraceProvider.h"
 
-FChaosVDConstraintDataProcessor::FChaosVDConstraintDataProcessor() : IChaosVDDataProcessor(FChaosVDConstraint::WrapperTypeName)
+FChaosVDConstraintDataProcessor::FChaosVDConstraintDataProcessor() : FChaosVDDataProcessorBase(FChaosVDConstraint::WrapperTypeName)
 {
 }
 
 bool FChaosVDConstraintDataProcessor::ProcessRawData(const TArray<uint8>& InData)
 {
+	FChaosVDDataProcessorBase::ProcessRawData(InData);
+
 	TSharedPtr<FChaosVDTraceProvider> ProviderSharedPtr = TraceProvider.Pin();
 	if (!ensure(ProviderSharedPtr.IsValid()))
 	{
diff --git a/Engine/Plugins/ChaosVD/Source/ChaosVD/Private/Trace/DataProcessors/ChaosVDConstraintDataProcessor.h b/Engine/Plugins/ChaosVD/Source/ChaosVD/Private/Trace/DataProcessors/ChaosVDConstraintDataProcessor.h
--- a/Engine/Plugins/ChaosVD/Source/ChaosVD/Private/Trace/DataProcessors/ChaosVDConstraintDataProcessor.h
+++ b/Engine/Plugins/ChaosVD/Source/ChaosVD/Private/Trace/DataProcessors/ChaosVDConstraintDataProcessor.h
@@ -1,7 +1,7 @@
 ﻿// Copyright Epic Games, Inc. All Rights Reserved.
 
 #pragma once
-#include "IChaosVDDataProcessor.h"
+#include "ChaosVDDataProcessorBase.h"
 
 struct FChaosVDSolverFrameData;
 struct FChaosVDConstraint;
@@ -9,7 +9,7 @@ struct FChaosVDConstraint;
 /**
  * Data processor implementation that is able to deserialize traced Constraints
  */
-class FChaosVDConstraintDataProcessor final : public IChaosVDDataProcessor
+class FChaosVDConstraintDataProcessor final : public FChaosVDDataProcessorBase
 {
 public:
 	explicit FChaosVDConstraintDataProcessor();
diff --git a/Engine/Plugins/ChaosVD/Source/ChaosVD/Private/Trace/DataProcessors/ChaosVDDataProcessorBase.cpp b/Engine/Plugins/ChaosVD/Source/ChaosVD/Private/Trace/DataProcessors/ChaosVDDataProcessorBase.cpp
new file mode 100644
--- /dev/null
+++ b/Engine/Plugins/ChaosVD/Source/ChaosVD/Private/Trace/DataProcessors/ChaosVDDataProcessorBase.cpp
@@ -0,0 +1,30 @@
+// Copyright Epic Games, Inc. All Rights Reserved.
+
+#pragma once
+
+#include "ChaosVDDataProcessorBase.h"
+
+FChaosVDDataProcessorBase::~FChaosVDDataProcessorBase()
+{
+}
+
+FStringView FChaosVDDataProcessorBase::GetCompatibleTypeName() const
+{
+	return CompatibleType;
+}
+
+bool FChaosVDDataProcessorBase::ProcessRawData(const TArray<uint8>& InData)
+{
+	ProcessedBytes += InData.Num();
+	return true;
+}
+
+uint64 FChaosVDDataProcessorBase::GetProcessedBytes() const
+{
+	return ProcessedBytes;
+}
+
+void FChaosVDDataProcessorBase::SetTraceProvider(const TSharedPtr<FChaosVDTraceProvider>& InProvider)
+{
+	TraceProvider = InProvider;
+}
diff --git a/Engine/Plugins/ChaosVD/Source/ChaosVD/Private/Trace/DataProcessors/ChaosVDDataProcessorBase.h b/Engine/Plugins/ChaosVD/Source/ChaosVD/Private/Trace/DataProcessors/ChaosVDDataProcessorBase.h
new file mode 100644
--- /dev/null
+++ b/Engine/Plugins/ChaosVD/Source/ChaosVD/Private/Trace/DataProcessors/ChaosVDDataProcessorBase.h
@@ -0,0 +1,90 @@
+﻿// Copyright Epic Games, Inc. All Rights Reserved.
+
+#pragma once
+
+#include "Chaos/ImplicitObject.h"
+#include "ChaosVisualDebugger/ChaosVDMemWriterReader.h"
+#include "ChaosVisualDebugger/ChaosVisualDebuggerTrace.h"
+#include "Containers/Array.h"
+#include "Containers/StringFwd.h"
+#include "DataWrappers/ChaosVDImplicitObjectDataWrapper.h"
+#include "HAL/Platform.h"
+#include "Templates/SharedPointer.h"
+#include "Trace/ChaosVDTraceProvider.h"
+
+#include <type_traits>
+
+struct FFortniteSeasonBranchObjectVersion;
+
+namespace Chaos::VisualDebugger
+{
+	template<typename TArchive>
+	void ApplyHeaderDataToArchive(TArchive& InOutArchive, const FChaosVDArchiveHeader& InRecordedHeader)
+	{
+		InOutArchive.SetCustomVersions(InRecordedHeader.CustomVersionContainer);
+		InOutArchive.SetEngineVer(InRecordedHeader.EngineVersion);
+		InOutArchive.SetShouldSkipUpdateCustomVersion(true);
+	}
+
+	template<typename TDataToSerialize>
+	bool ReadDataFromBuffer(const TArray<uint8>& InDataBuffer, TDataToSerialize& Data, const TSharedRef<FChaosVDTraceProvider>& DataProvider)
+	{
+		const TSharedPtr<FChaosVDRecording> RecordingInstance = DataProvider->GetRecordingForSession();
+		const TSharedPtr<FChaosVDSerializableNameTable> NameTableInstance = RecordingInstance ? RecordingInstance->GetNameTableInstance() : nullptr;
+
+		if (!ensure(NameTableInstance.IsValid()))
+		{
+			return false;
+		}
+		
+		FChaosVDMemoryReader MemReader(InDataBuffer, NameTableInstance.ToSharedRef());
+		const FChaosVDArchiveHeader& RecordedHeader =  RecordingInstance->GetHeaderData();
+		ApplyHeaderDataToArchive(MemReader, RecordedHeader);
+
+		bool bSuccess = false;
+
+		// We need to use FChaosArchive as proxy to properly read serialized Implicit objects
+		// Note: I don't expect we will need a proxy archive for other types, but if we end up in that situation, we should use to switch to use traits 
+		if constexpr (std::is_same_v<TDataToSerialize, FChaosVDImplicitObjectDataWrapper<Chaos::FImplicitObjectPtr, Chaos::FChaosArchive>>)
+		{
+			FChaosArchive Ar(MemReader);
+			bSuccess = Data.Serialize(Ar);
+		}
+		else
+		{
+			bSuccess = Data.Serialize(MemReader);
+		}
+
+		return bSuccess;
+	}
+}
+
+/** Abstract base class that should be used for any class that is able to process traced Chaos Visual Debugger binary data */
+class FChaosVDDataProcessorBase
+{
+public:
+
+	explicit FChaosVDDataProcessorBase(FStringView InCompatibleType) : TraceProvider(nullptr), CompatibleType(InCompatibleType), ProcessedBytes(0)
+	{
+	}
+
+	virtual ~FChaosVDDataProcessorBase() = 0;
+
+	/** Type name this data processor can interpret */
+	FStringView GetCompatibleTypeName() const;
+
+	/** Called with the raw serialized data to be processed */
+	virtual bool ProcessRawData(const TArray<uint8>& InData);
+
+	/** Returns the amount of data in bytes processed by this data processor at the moment of being called */
+	uint64 GetProcessedBytes() const;
+
+	/** Sets the Trace Provider that is storing the data being analyzed */
+	void SetTraceProvider(const TSharedPtr<FChaosVDTraceProvider>& InProvider);
+
+protected:
+	TWeakPtr<FChaosVDTraceProvider> TraceProvider;
+	FStringView CompatibleType;
+	uint64 ProcessedBytes;
+};
+
diff --git a/Engine/Plugins/ChaosVD/Source/ChaosVD/Private/Trace/DataProcessors/ChaosVDJointConstraintDataProcessor.cpp b/Engine/Plugins/ChaosVD/Source/ChaosVD/Private/Trace/DataProcessors/ChaosVDJointConstraintDataProcessor.cpp
--- a/Engine/Plugins/ChaosVD/Source/ChaosVD/Private/Trace/DataProcessors/ChaosVDJointConstraintDataProcessor.cpp
+++ b/Engine/Plugins/ChaosVD/Source/ChaosVD/Private/Trace/DataProcessors/ChaosVDJointConstraintDataProcessor.cpp
@@ -7,12 +7,14 @@
 
 
 FChaosVDJointConstraintDataProcessor::FChaosVDJointConstraintDataProcessor()
-	: IChaosVDDataProcessor(FChaosVDJointConstraint::WrapperTypeName)
+	: FChaosVDDataProcessorBase(FChaosVDJointConstraint::WrapperTypeName)
 {
 }
 
 bool FChaosVDJointConstraintDataProcessor::ProcessRawData(const TArray<uint8>& InData)
 {
+	FChaosVDDataProcessorBase::ProcessRawData(InData);
+
 	TSharedPtr<FChaosVDTraceProvider> ProviderSharedPtr = TraceProvider.Pin();
 	if (!ensure(ProviderSharedPtr.IsValid()))
 	{
diff --git a/Engine/Plugins/ChaosVD/Source/ChaosVD/Private/Trace/DataProcessors/ChaosVDJointConstraintDataProcessor.h b/Engine/Plugins/ChaosVD/Source/ChaosVD/Private/Trace/DataProcessors/ChaosVDJointConstraintDataProcessor.h
--- a/Engine/Plugins/ChaosVD/Source/ChaosVD/Private/Trace/DataProcessors/ChaosVDJointConstraintDataProcessor.h
+++ b/Engine/Plugins/ChaosVD/Source/ChaosVD/Private/Trace/DataProcessors/ChaosVDJointConstraintDataProcessor.h
@@ -3,7 +3,7 @@
 #pragma once
 
 #pragma once
-#include "IChaosVDDataProcessor.h"
+#include "ChaosVDDataProcessorBase.h"
 
 struct FChaosVDSolverFrameData;
 struct FChaosVDConstraint;
@@ -11,7 +11,7 @@ struct FChaosVDConstraint;
 /**
  * Data processor implementation that is able to deserialize traced joint constraints
  */
-class FChaosVDJointConstraintDataProcessor final : public IChaosVDDataProcessor
+class FChaosVDJointConstraintDataProcessor final : public FChaosVDDataProcessorBase
 {
 public:
 	explicit FChaosVDJointConstraintDataProcessor();
diff --git a/Engine/Plugins/ChaosVD/Source/ChaosVD/Private/Trace/DataProcessors/ChaosVDMidPhaseDataProcessor.cpp b/Engine/Plugins/ChaosVD/Source/ChaosVD/Private/Trace/DataProcessors/ChaosVDMidPhaseDataProcessor.cpp
--- a/Engine/Plugins/ChaosVD/Source/ChaosVD/Private/Trace/DataProcessors/ChaosVDMidPhaseDataProcessor.cpp
+++ b/Engine/Plugins/ChaosVD/Source/ChaosVD/Private/Trace/DataProcessors/ChaosVDMidPhaseDataProcessor.cpp
@@ -6,12 +6,14 @@
 #include "ChaosVisualDebugger/ChaosVDMemWriterReader.h"
 #include "Trace/ChaosVDTraceProvider.h"
 
-FChaosVDMidPhaseDataProcessor::FChaosVDMidPhaseDataProcessor() : IChaosVDDataProcessor(FChaosVDParticlePairMidPhase::WrapperTypeName)
+FChaosVDMidPhaseDataProcessor::FChaosVDMidPhaseDataProcessor() : FChaosVDDataProcessorBase(FChaosVDParticlePairMidPhase::WrapperTypeName)
 {
 }
 
 bool FChaosVDMidPhaseDataProcessor::ProcessRawData(const TArray<uint8>& InData)
 {
+	FChaosVDDataProcessorBase::ProcessRawData(InData);
+
 	TSharedPtr<FChaosVDTraceProvider> ProviderSharedPtr = TraceProvider.Pin();
 	if (!ensure(ProviderSharedPtr.IsValid()))
 	{
diff --git a/Engine/Plugins/ChaosVD/Source/ChaosVD/Private/Trace/DataProcessors/ChaosVDMidPhaseDataProcessor.h b/Engine/Plugins/ChaosVD/Source/ChaosVD/Private/Trace/DataProcessors/ChaosVDMidPhaseDataProcessor.h
--- a/Engine/Plugins/ChaosVD/Source/ChaosVD/Private/Trace/DataProcessors/ChaosVDMidPhaseDataProcessor.h
+++ b/Engine/Plugins/ChaosVD/Source/ChaosVD/Private/Trace/DataProcessors/ChaosVDMidPhaseDataProcessor.h
@@ -1,7 +1,7 @@
 ﻿// Copyright Epic Games, Inc. All Rights Reserved.
 
 #pragma once
-#include "IChaosVDDataProcessor.h"
+#include "ChaosVDDataProcessorBase.h"
 #include "Templates/SharedPointer.h"
 
 struct FChaosVDSolverFrameData;
@@ -10,7 +10,7 @@ struct FChaosVDParticlePairMidPhase;
 /**
  * Data processor implementation that is able to deserialize traced MidPhases 
  */
-class FChaosVDMidPhaseDataProcessor final : public IChaosVDDataProcessor
+class FChaosVDMidPhaseDataProcessor final : public FChaosVDDataProcessorBase
 {
 public:
 	explicit FChaosVDMidPhaseDataProcessor();
diff --git a/Engine/Plugins/ChaosVD/Source/ChaosVD/Private/Trace/DataProcessors/ChaosVDSceneQueryDataProcessor.cpp b/Engine/Plugins/ChaosVD/Source/ChaosVD/Private/Trace/DataProcessors/ChaosVDSceneQueryDataProcessor.cpp
--- a/Engine/Plugins/ChaosVD/Source/ChaosVD/Private/Trace/DataProcessors/ChaosVDSceneQueryDataProcessor.cpp
+++ b/Engine/Plugins/ChaosVD/Source/ChaosVD/Private/Trace/DataProcessors/ChaosVDSceneQueryDataProcessor.cpp
@@ -9,12 +9,14 @@
 #include "DataWrappers/ChaosVDQueryDataWrappers.h"
 #include "Trace/ChaosVDTraceProvider.h"
 
-FChaosVDSceneQueryDataProcessor::FChaosVDSceneQueryDataProcessor() : IChaosVDDataProcessor(FChaosVDQueryDataWrapper::WrapperTypeName)
+FChaosVDSceneQueryDataProcessor::FChaosVDSceneQueryDataProcessor() : FChaosVDDataProcessorBase(FChaosVDQueryDataWrapper::WrapperTypeName)
 {
 }
 
 bool FChaosVDSceneQueryDataProcessor::ProcessRawData(const TArray<uint8>& InData)
 {
+	FChaosVDDataProcessorBase::ProcessRawData(InData);
+
 	const TSharedPtr<FChaosVDTraceProvider> ProviderSharedPtr = TraceProvider.Pin();
 	if (!ensure(ProviderSharedPtr.IsValid()))
 	{
diff --git a/Engine/Plugins/ChaosVD/Source/ChaosVD/Private/Trace/DataProcessors/ChaosVDSceneQueryDataProcessor.h b/Engine/Plugins/ChaosVD/Source/ChaosVD/Private/Trace/DataProcessors/ChaosVDSceneQueryDataProcessor.h
--- a/Engine/Plugins/ChaosVD/Source/ChaosVD/Private/Trace/DataProcessors/ChaosVDSceneQueryDataProcessor.h
+++ b/Engine/Plugins/ChaosVD/Source/ChaosVD/Private/Trace/DataProcessors/ChaosVDSceneQueryDataProcessor.h
@@ -2,12 +2,12 @@
 
 #pragma once
 
-#include "Trace/DataProcessors/IChaosVDDataProcessor.h"
+#include "Trace/DataProcessors/ChaosVDDataProcessorBase.h"
 
 /**
  * Data processor implementation that is able to deserialize traced scene queries data
  */
-class FChaosVDSceneQueryDataProcessor final : public IChaosVDDataProcessor
+class FChaosVDSceneQueryDataProcessor final : public FChaosVDDataProcessorBase
 {
 public:
 	explicit FChaosVDSceneQueryDataProcessor();
diff --git a/Engine/Plugins/ChaosVD/Source/ChaosVD/Private/Trace/DataProcessors/ChaosVDSceneQueryVisitDataProcessor.cpp b/Engine/Plugins/ChaosVD/Source/ChaosVD/Private/Trace/DataProcessors/ChaosVDSceneQueryVisitDataProcessor.cpp
--- a/Engine/Plugins/ChaosVD/Source/ChaosVD/Private/Trace/DataProcessors/ChaosVDSceneQueryVisitDataProcessor.cpp
+++ b/Engine/Plugins/ChaosVD/Source/ChaosVD/Private/Trace/DataProcessors/ChaosVDSceneQueryVisitDataProcessor.cpp
@@ -8,12 +8,14 @@
 #include "DataWrappers/ChaosVDQueryDataWrappers.h"
 #include "Trace/ChaosVDTraceProvider.h"
 
-FChaosVDSceneQueryVisitDataProcessor::FChaosVDSceneQueryVisitDataProcessor() : IChaosVDDataProcessor(FChaosVDQueryVisitStep::WrapperTypeName)
+FChaosVDSceneQueryVisitDataProcessor::FChaosVDSceneQueryVisitDataProcessor() : FChaosVDDataProcessorBase(FChaosVDQueryVisitStep::WrapperTypeName)
 {
 }
 
 bool FChaosVDSceneQueryVisitDataProcessor::ProcessRawData(const TArray<uint8>& InData)
 {
+	FChaosVDDataProcessorBase::ProcessRawData(InData);
+
 	const TSharedPtr<FChaosVDTraceProvider> ProviderSharedPtr = TraceProvider.Pin();
 	if (!ensure(ProviderSharedPtr.IsValid()))
 	{
diff --git a/Engine/Plugins/ChaosVD/Source/ChaosVD/Private/Trace/DataProcessors/ChaosVDSceneQueryVisitDataProcessor.h b/Engine/Plugins/ChaosVD/Source/ChaosVD/Private/Trace/DataProcessors/ChaosVDSceneQueryVisitDataProcessor.h
--- a/Engine/Plugins/ChaosVD/Source/ChaosVD/Private/Trace/DataProcessors/ChaosVDSceneQueryVisitDataProcessor.h
+++ b/Engine/Plugins/ChaosVD/Source/ChaosVD/Private/Trace/DataProcessors/ChaosVDSceneQueryVisitDataProcessor.h
@@ -1,12 +1,12 @@
 ﻿// Copyright Epic Games, Inc. All Rights Reserved.
 
 #pragma once
-#include "IChaosVDDataProcessor.h"
+#include "ChaosVDDataProcessorBase.h"
 
 /**
  * Data processor implementation that is able to deserialize traced SQ Visits
  */
-class FChaosVDSceneQueryVisitDataProcessor final : public IChaosVDDataProcessor
+class FChaosVDSceneQueryVisitDataProcessor final : public FChaosVDDataProcessorBase
 {
 public:
 	explicit FChaosVDSceneQueryVisitDataProcessor();
diff --git a/Engine/Plugins/ChaosVD/Source/ChaosVD/Private/Trace/DataProcessors/ChaosVDSerializedNameEntryDataProcessor.cpp b/Engine/Plugins/ChaosVD/Source/ChaosVD/Private/Trace/DataProcessors/ChaosVDSerializedNameEntryDataProcessor.cpp
--- a/Engine/Plugins/ChaosVD/Source/ChaosVD/Private/Trace/DataProcessors/ChaosVDSerializedNameEntryDataProcessor.cpp
+++ b/Engine/Plugins/ChaosVD/Source/ChaosVD/Private/Trace/DataProcessors/ChaosVDSerializedNameEntryDataProcessor.cpp
@@ -8,12 +8,14 @@
 #include "Trace/ChaosVDTraceProvider.h"
 
 FChaosVDSerializedNameEntryDataProcessor::FChaosVDSerializedNameEntryDataProcessor()
-	: IChaosVDDataProcessor(Chaos::VisualDebugger::FChaosVDSerializedNameEntry::WrapperTypeName)
+	: FChaosVDDataProcessorBase(Chaos::VisualDebugger::FChaosVDSerializedNameEntry::WrapperTypeName)
 {
 }
 
 bool FChaosVDSerializedNameEntryDataProcessor::ProcessRawData(const TArray<uint8>& InData)
 {
+	FChaosVDDataProcessorBase::ProcessRawData(InData);
+	
 	const TSharedPtr<FChaosVDTraceProvider> ProviderSharedPtr = TraceProvider.Pin();
 	if (!ensure(ProviderSharedPtr.IsValid()))
 	{
diff --git a/Engine/Plugins/ChaosVD/Source/ChaosVD/Private/Trace/DataProcessors/ChaosVDSerializedNameEntryDataProcessor.h b/Engine/Plugins/ChaosVD/Source/ChaosVD/Private/Trace/DataProcessors/ChaosVDSerializedNameEntryDataProcessor.h
--- a/Engine/Plugins/ChaosVD/Source/ChaosVD/Private/Trace/DataProcessors/ChaosVDSerializedNameEntryDataProcessor.h
+++ b/Engine/Plugins/ChaosVD/Source/ChaosVD/Private/Trace/DataProcessors/ChaosVDSerializedNameEntryDataProcessor.h
@@ -2,12 +2,12 @@
 
 #pragma once
 
-#include "Trace/DataProcessors/IChaosVDDataProcessor.h"
+#include "Trace/DataProcessors/ChaosVDDataProcessorBase.h"
 
 /**
  * Data processor implementation that is able to deserialize traced Name Entries
  */
-class FChaosVDSerializedNameEntryDataProcessor final : public IChaosVDDataProcessor
+class FChaosVDSerializedNameEntryDataProcessor final : public FChaosVDDataProcessorBase
 {
 public:
 	explicit FChaosVDSerializedNameEntryDataProcessor();
diff --git a/Engine/Plugins/ChaosVD/Source/ChaosVD/Private/Trace/DataProcessors/ChaosVDTraceImplicitObjectProcessor.cpp b/Engine/Plugins/ChaosVD/Source/ChaosVD/Private/Trace/DataProcessors/ChaosVDTraceImplicitObjectProcessor.cpp
--- a/Engine/Plugins/ChaosVD/Source/ChaosVD/Private/Trace/DataProcessors/ChaosVDTraceImplicitObjectProcessor.cpp
+++ b/Engine/Plugins/ChaosVD/Source/ChaosVD/Private/Trace/DataProcessors/ChaosVDTraceImplicitObjectProcessor.cpp
@@ -11,12 +11,14 @@
 
 using FChaosVDImplicitObjectWrapper = FChaosVDImplicitObjectDataWrapper<Chaos::FImplicitObjectPtr, Chaos::FChaosArchive>;
 
-FChaosVDTraceImplicitObjectProcessor::FChaosVDTraceImplicitObjectProcessor(): IChaosVDDataProcessor(FChaosVDImplicitObjectWrapper::WrapperTypeName)
+FChaosVDTraceImplicitObjectProcessor::FChaosVDTraceImplicitObjectProcessor(): FChaosVDDataProcessorBase(FChaosVDImplicitObjectWrapper::WrapperTypeName)
 {
 }
 
 bool FChaosVDTraceImplicitObjectProcessor::ProcessRawData(const TArray<uint8>& InData)
 {
+	FChaosVDDataProcessorBase::ProcessRawData(InData);
+
 	TSharedPtr<FChaosVDTraceProvider> ProviderSharedPtr = TraceProvider.Pin();
 	if (!ensure(ProviderSharedPtr.IsValid()))
 	{
diff --git a/Engine/Plugins/ChaosVD/Source/ChaosVD/Private/Trace/DataProcessors/ChaosVDTraceImplicitObjectProcessor.h b/Engine/Plugins/ChaosVD/Source/ChaosVD/Private/Trace/DataProcessors/ChaosVDTraceImplicitObjectProcessor.h
--- a/Engine/Plugins/ChaosVD/Source/ChaosVD/Private/Trace/DataProcessors/ChaosVDTraceImplicitObjectProcessor.h
+++ b/Engine/Plugins/ChaosVD/Source/ChaosVD/Private/Trace/DataProcessors/ChaosVDTraceImplicitObjectProcessor.h
@@ -2,12 +2,12 @@
 
 #pragma once
 
-#include "IChaosVDDataProcessor.h"
+#include "ChaosVDDataProcessorBase.h"
 
 /**
  * Data processor implementation that is able to deserialize traced Implicit objects
  */
-class FChaosVDTraceImplicitObjectProcessor final : public IChaosVDDataProcessor
+class FChaosVDTraceImplicitObjectProcessor final : public FChaosVDDataProcessorBase
 {
 public:
 	explicit FChaosVDTraceImplicitObjectProcessor();
diff --git a/Engine/Plugins/ChaosVD/Source/ChaosVD/Private/Trace/DataProcessors/ChaosVDTraceParticleDataProcessor.cpp b/Engine/Plugins/ChaosVD/Source/ChaosVD/Private/Trace/DataProcessors/ChaosVDTraceParticleDataProcessor.cpp
--- a/Engine/Plugins/ChaosVD/Source/ChaosVD/Private/Trace/DataProcessors/ChaosVDTraceParticleDataProcessor.cpp
+++ b/Engine/Plugins/ChaosVD/Source/ChaosVD/Private/Trace/DataProcessors/ChaosVDTraceParticleDataProcessor.cpp
@@ -8,12 +8,14 @@
 #include "DataWrappers/ChaosVDParticleDataWrapper.h"
 #include "Trace/ChaosVDTraceProvider.h"
 
-FChaosVDTraceParticleDataProcessor::FChaosVDTraceParticleDataProcessor(): IChaosVDDataProcessor(FChaosVDParticleDataWrapper::WrapperTypeName)
+FChaosVDTraceParticleDataProcessor::FChaosVDTraceParticleDataProcessor(): FChaosVDDataProcessorBase(FChaosVDParticleDataWrapper::WrapperTypeName)
 {
 }
 
 bool FChaosVDTraceParticleDataProcessor::ProcessRawData(const TArray<uint8>& InData)
 {
+	FChaosVDDataProcessorBase::ProcessRawData(InData);
+
 	TSharedPtr<FChaosVDTraceProvider> ProviderSharedPtr = TraceProvider.Pin();
 	if (!ensure(ProviderSharedPtr.IsValid()))
 	{
diff --git a/Engine/Plugins/ChaosVD/Source/ChaosVD/Private/Trace/DataProcessors/ChaosVDTraceParticleDataProcessor.h b/Engine/Plugins/ChaosVD/Source/ChaosVD/Private/Trace/DataProcessors/ChaosVDTraceParticleDataProcessor.h
--- a/Engine/Plugins/ChaosVD/Source/ChaosVD/Private/Trace/DataProcessors/ChaosVDTraceParticleDataProcessor.h
+++ b/Engine/Plugins/ChaosVD/Source/ChaosVD/Private/Trace/DataProcessors/ChaosVDTraceParticleDataProcessor.h
@@ -2,12 +2,12 @@
 
 #pragma once
 
-#include "Trace/DataProcessors/IChaosVDDataProcessor.h"
+#include "Trace/DataProcessors/ChaosVDDataProcessorBase.h"
 
 /**
  * Data processor implementation that is able to deserialize traced Particles data
  */
-class FChaosVDTraceParticleDataProcessor final : public IChaosVDDataProcessor
+class FChaosVDTraceParticleDataProcessor final : public FChaosVDDataProcessorBase
 {
 public:
 	explicit FChaosVDTraceParticleDataProcessor();
diff --git a/Engine/Plugins/ChaosVD/Source/ChaosVD/Private/Trace/DataProcessors/IChaosVDDataProcessor.h b/Engine/Plugins/ChaosVD/Source/ChaosVD/Private/Trace/DataProcessors/IChaosVDDataProcessor.h
deleted file mode 100644
