[ChaosVD] In-Tool Performance improvements Part 2

- Eliminated the use of material instance dynamic. CVD only needs to change the color value read by the material depending on the settings and the recorded data for any given particle, so now we store the color data as primitive data, so the material can just read it from the primitive itself
- Re-implemented the mesh component system to be more uobject generic. Part of this work included improve how we generate unique names for the pooled components
- Now the step to process implicit object data into a format that can be use to create a static mesh is done in worked threads when the static meshes will actually be built
- Parallelized how we process the Tri mesh and Heightfield implicit object data. Now we process the vertex data in parallel when converting the data to a format we can use to create a static mesh
- Implemented de-initialization code for the geometry builder, and use it when we unload an load a recording.
- Added support to compare the recorded per shape instance data using memcmp instead of doing it per item
- Added a missing early out on the update visibility code path. Now if the visibility has not change, we don't do the update.
- Added a common interface for solver data component (this is not being used yet, it is part of work that will be submitted in other CL to provide a generic way to add more data like new constraint types)
- Removed support for dynamic meshes. It has been over a year we don't use dynamic meshes, and the code path as not been really tested since then. removing support now allows us to remove templatized code and make everything easier to read/maintain.

#virtualized

diff --git a/Engine/Plugins/ChaosVD/Source/ChaosVD/Private/Actors/ChaosVDGeometryContainer.cpp b/Engine/Plugins/ChaosVD/Source/ChaosVD/Private/Actors/ChaosVDGeometryContainer.cpp
new file mode 100644
--- /dev/null
+++ b/Engine/Plugins/ChaosVD/Source/ChaosVD/Private/Actors/ChaosVDGeometryContainer.cpp
@@ -0,0 +1,37 @@
+// Copyright Epic Games, Inc. All Rights Reserved.
+
+#include "ChaosVDGeometryContainer.h"
+#include "Components/MeshComponent.h"
+#include "Containers/Array.h"
+#include "Misc/MemStack.h"
+#include "Misc/ScopedSlowTask.h"
+
+AChaosVDGeometryContainer::AChaosVDGeometryContainer()
+{
+	PrimaryActorTick.bCanEverTick = false;
+}
+
+void AChaosVDGeometryContainer::CleanUp()
+{
+	// We can have ten of thousands of components, so the InlineComponents arrays will not be better than a normal array.
+	FMemMark Mark(FMemStack::Get());
+	TArray<UMeshComponent*, TMemStackAllocator<>> ComponentsToDestroy;
+
+	GetComponents<UMeshComponent>(ComponentsToDestroy);
+
+	constexpr float AmountOfWork = 1.0f;
+	const float PercentagePerElement = 1.0f / ComponentsToDestroy.Num();
+
+	FScopedSlowTask CleaningDataSlowTask(AmountOfWork, NSLOCTEXT("ChaosVisualDebugger", "CleaningupGeometryData", "Clearing Geoemtry Data ..."));
+	CleaningDataSlowTask.MakeDialog();
+
+	for (UMeshComponent* Component : ComponentsToDestroy)
+	{
+		if (Component)
+		{
+			Component->DestroyComponent();
+		}
+
+		CleaningDataSlowTask.EnterProgressFrame(PercentagePerElement);
+	}
+}
diff --git a/Engine/Plugins/ChaosVD/Source/ChaosVD/Private/Actors/ChaosVDGeometryContainer.h b/Engine/Plugins/ChaosVD/Source/ChaosVD/Private/Actors/ChaosVDGeometryContainer.h
new file mode 100644
--- /dev/null
+++ b/Engine/Plugins/ChaosVD/Source/ChaosVD/Private/Actors/ChaosVDGeometryContainer.h
@@ -0,0 +1,20 @@
+// Copyright Epic Games, Inc. All Rights Reserved.
+
+#pragma once
+
+#include "GameFramework/Actor.h"
+#include "UObject/ObjectPtr.h"
+
+#include "ChaosVDGeometryContainer.generated.h"
+
+/** Actor that contains Static Mesh Components used to visualize the geometry we generated from the recorded data */
+UCLASS()
+class CHAOSVD_API AChaosVDGeometryContainer : public AActor
+{
+	GENERATED_BODY()
+
+public:
+	AChaosVDGeometryContainer();
+
+	void CleanUp();
+};
\ No newline at end of file
diff --git a/Engine/Plugins/ChaosVD/Source/ChaosVD/Private/Actors/ChaosVDSceneQueryDataContainer.cpp b/Engine/Plugins/ChaosVD/Source/ChaosVD/Private/Actors/ChaosVDSceneQueryDataContainer.cpp
--- a/Engine/Plugins/ChaosVD/Source/ChaosVD/Private/Actors/ChaosVDSceneQueryDataContainer.cpp
+++ b/Engine/Plugins/ChaosVD/Source/ChaosVD/Private/Actors/ChaosVDSceneQueryDataContainer.cpp
@@ -10,3 +10,11 @@ AChaosVDSceneQueryDataContainer::AChaosVDSceneQueryDataContainer()
 
 	SceneQueryDataComponent = CreateDefaultSubobject<UChaosVDSceneQueryDataComponent>(TEXT("ChaosVDSceneQueryDataComponent"));
 }
+
+void AChaosVDSceneQueryDataContainer::CleanUp()
+{
+	if (SceneQueryDataComponent)
+	{
+		SceneQueryDataComponent->ClearData();
+	}
+}
diff --git a/Engine/Plugins/ChaosVD/Source/ChaosVD/Private/Actors/ChaosVDSceneQueryDataContainer.h b/Engine/Plugins/ChaosVD/Source/ChaosVD/Private/Actors/ChaosVDSceneQueryDataContainer.h
--- a/Engine/Plugins/ChaosVD/Source/ChaosVD/Private/Actors/ChaosVDSceneQueryDataContainer.h
+++ b/Engine/Plugins/ChaosVD/Source/ChaosVD/Private/Actors/ChaosVDSceneQueryDataContainer.h
@@ -21,6 +21,8 @@ public:
 
 	UChaosVDSceneQueryDataComponent* GetSceneQueryDataComponent() const { return SceneQueryDataComponent.Get(); }
 
+	void CleanUp();
+
 protected:
 	UPROPERTY()
 	TObjectPtr<UChaosVDSceneQueryDataComponent> SceneQueryDataComponent;
diff --git a/Engine/Plugins/ChaosVD/Source/ChaosVD/Private/ChaosVDGeometryBuilder.cpp b/Engine/Plugins/ChaosVD/Source/ChaosVD/Private/ChaosVDGeometryBuilder.cpp
--- a/Engine/Plugins/ChaosVD/Source/ChaosVD/Private/ChaosVDGeometryBuilder.cpp
+++ b/Engine/Plugins/ChaosVD/Source/ChaosVD/Private/ChaosVDGeometryBuilder.cpp
@@ -2,11 +2,11 @@
 
 #include "ChaosVDGeometryBuilder.h"
 
+#include "Chaos/HeightField.h"
 #include "ChaosVDConvexMeshGenerator.h"
 #include "ChaosVDGeometryDataComponent.h"
 #include "ChaosVDModule.h"
 #include "ChaosVDTriMeshGenerator.h"
-#include "Components/ChaosVDStaticMeshComponent.h"
 #include "DynamicMeshToMeshDescription.h"
 #include "Generators/CapsuleGenerator.h"
 #include "Generators/MinimalBoxMeshGenerator.h"
@@ -14,9 +14,9 @@
 #include "MeshConstraints.h"
 #include "MeshConstraintsUtil.h"
 #include "MeshSimplification.h"
+#include "Misc/ScopedSlowTask.h"
 #include "StaticMeshAttributes.h"
 #include "UDynamicMesh.h"
-#include "Chaos/HeightField.h"
 #include "UObject/UObjectGlobals.h"
 
 namespace Chaos::VisualDebugger
@@ -163,8 +163,18 @@ namespace Chaos::VisualDebugger
 	}
 }
 
+FChaosVDGeometryBuilder::~FChaosVDGeometryBuilder()
+{
+	DeInitialize();
+}
+
 void FChaosVDGeometryBuilder::Initialize(const TWeakPtr<FChaosVDScene>& ChaosVDScene)
 {
+	if (!ChaosVDScene.Pin())
+	{
+		return;
+	}
+
 	SceneWeakPtr = ChaosVDScene;
 
 	auto ProcessMeshComponent = [WeakThis = AsWeak()](uint32 GeometryKey, const TWeakObjectPtr<UMeshComponent> Object)
@@ -174,7 +184,7 @@ void FChaosVDGeometryBuilder::Initialize(const TWeakPtr<FChaosVDScene>& ChaosVDS
 		{
 			UE_LOG(LogChaosVDEditor, Verbose, TEXT(" [%s] Failed to update mesh for Handle | Geometry Key [%u] | Geometry Builder is invalid"), ANSI_TO_TCHAR(__FUNCTION__), GeometryKey);
 
-			// If the the builder is no longer valid, just consume the request
+			// If the builder is no longer valid, just consume the request
 			return true;
 		}
 
@@ -191,14 +201,14 @@ void FChaosVDGeometryBuilder::Initialize(const TWeakPtr<FChaosVDScene>& ChaosVDS
 		return false;
 	};
 	
-	auto CreateMaterialForMeshComponent = [WeakThis = AsWeak()](const TWeakObjectPtr<UMeshComponent> Object)
+	auto UpdateMeshMaterialForComponent = [WeakThis = AsWeak()](const TWeakObjectPtr<UMeshComponent> Object)
 	{
 		const TSharedPtr<FChaosVDGeometryBuilder> GeometryBuilder = WeakThis.Pin();
 		if (!GeometryBuilder)
 		{
 			UE_LOG(LogChaosVDEditor, Verbose, TEXT(" [%s] Failed to Create Material for Mesh | Geometry builder is no longer valid "), ANSI_TO_TCHAR(__FUNCTION__));
 
-			// If the the builder is no longer valid, just consume the request
+			// If the builder is no longer valid, just consume the request
 			return true;
 		}
 
@@ -212,29 +222,112 @@ void FChaosVDGeometryBuilder::Initialize(const TWeakPtr<FChaosVDScene>& ChaosVDS
 				return true;
 			}
 
-			GeometryBuilder->SetMeshComponentMaterial(CVDMeshComponent->GetMeshComponentAttributeFlags(), MeshComponent);
+			GeometryBuilder->SetMeshComponentMaterial(CVDMeshComponent);
 		}
 		return true;
 	};
 
-	auto LauchGeometryGenerationTaskDeferred = [WeakThis = AsWeak()](TSharedPtr<FGeometryGenerationTask> GeometryGenerationTask)
+	auto LaunchGeometryGenerationTaskDeferred = [](TSharedPtr<FChaosVDGeometryGenerationTask> GeometryGenerationTask)
 	{
-		UE::Tasks::Launch(TEXT("GeometryGeneration"),
-[GeometryGenerationTask]()
+		if(GeometryGenerationTask)
 		{
-			GeometryGenerationTask->GenerateGeometry();
-		});
+			GeometryGenerationTask->TaskHandle = UE::Tasks::Launch(TEXT("GeometryGeneration"),
+			[GeometryGenerationTask]()
+			{
+				if (GeometryGenerationTask)
+				{
+					if (GeometryGenerationTask->IsCanceled())
+					{
+						return;
+					}
+					GeometryGenerationTask->GenerateGeometry();
+				}
+			});
+		}
+	
 		return true;
 	};
 
 	MeshComponentsWaitingForGeometry = MakeUnique<FObjectsWaitingGeometryList<FMeshComponentWeakPtr>>(ProcessMeshComponent, NSLOCTEXT("ChaosVisualDebugger", "GeometryGenNotification","Mesh Components"), ShouldProcessObjectsForKey);
-	MeshComponentsWaitingForMaterial = MakeUnique<FObjectsWaitingProcessingQueue<FMeshComponentWeakPtr>>(CreateMaterialForMeshComponent, NSLOCTEXT("ChaosVisualDebugger", "GeometryMaterialNotification","Material instances"));
-	GeometryTasksPendingLaunch = MakeUnique<FObjectsWaitingProcessingQueue<TSharedPtr<FGeometryGenerationTask>>>(LauchGeometryGenerationTaskDeferred, NSLOCTEXT("ChaosVisualDebugger", "GeometryTaskLauchNotification","Static Meshes"));
+	MeshComponentsWaitingForMaterial = MakeUnique<FObjectsWaitingProcessingQueue<FMeshComponentWeakPtr>>(UpdateMeshMaterialForComponent, NSLOCTEXT("ChaosVisualDebugger", "GeometryMaterialNotification","Material instances"));
+	GeometryTasksPendingLaunch = MakeUnique<FObjectsWaitingProcessingQueue<TSharedPtr<FChaosVDGeometryGenerationTask>>>(LaunchGeometryGenerationTaskDeferred, NSLOCTEXT("ChaosVisualDebugger", "GeometryTaskLauchNotification","Static Meshes"));
 
 	GameThreadTickDelegate = FTSTicker::GetCoreTicker().AddTicker(FTickerDelegate::CreateRaw(this, &FChaosVDGeometryBuilder::GameThreadTick));
 
 	constexpr int32 MeshPendingDisposalContainerDefaultSize = 500; 
 	MeshComponentsPendingDisposal.Reserve(MeshPendingDisposalContainerDefaultSize);
+
+	bInitialized = true;
+}
+
+void FChaosVDGeometryBuilder::DeInitialize()
+{
+	if (!bInitialized)
+	{
+		return;
+	}
+
+	constexpr float MaxAmountOfWork = 1.0f;
+	const int32 WorkRemaining = GeometryBeingGeneratedByKey.Num() + StaticMeshCacheMap.Num();
+	const float PercentagePerElement = 1.0f / WorkRemaining;
+
+	FScopedSlowTask CleaningGeometrySlowTask(MaxAmountOfWork, NSLOCTEXT("ChaosVisualDebugger", "DeInitializeGeometrybuilderSlowTask", "Deinitialiing GeometryBuilder"));
+
+	FTSTicker::GetCoreTicker().RemoveTicker(GameThreadTickDelegate);
+
+	int32 TasksFailedToCancelNum = 0;
+
+	for (const TPair<uint32, TSharedPtr<FChaosVDGeometryGenerationTask>>& TaskWithID : GeometryBeingGeneratedByKey)
+	{
+		if (TaskWithID.Value)
+		{
+			TaskWithID.Value->CancelTask();
+		}
+
+		if (!TaskWithID.Value->TaskHandle.Wait(FTimespan::FromSeconds(10.0f)))
+		{
+			TasksFailedToCancelNum++;
+		}
+
+		CleaningGeometrySlowTask.EnterProgressFrame(PercentagePerElement);
+	}
+	
+	if (TasksFailedToCancelNum > 0)
+	{
+		UE_LOG(LogChaosVDEditor, Warning, TEXT("[%s] Failed to cancel [%d] tasks"), ANSI_TO_TCHAR(__FUNCTION__), TasksFailedToCancelNum);
+	}
+	
+	GeometryBeingGeneratedByKey.Reset();
+
+	FInstancedStaticMeshDelegates::OnInstanceIndexUpdated.RemoveAll(this);
+
+	for (const TPair<uint32, TObjectPtr<UStaticMesh>>& StaticMeshByKey : StaticMeshCacheMap)
+	{
+		if (StaticMeshByKey.Value)
+		{
+			StaticMeshByKey.Value->ClearFlags(RF_Standalone);
+			StaticMeshByKey.Value->MarkAsGarbage();
+		}
+		CleaningGeometrySlowTask.EnterProgressFrame(PercentagePerElement);
+	}
+
+	MeshComponentsPendingDisposal.Reset();
+	TranslucentMirroredInstancedMeshComponentByGeometryKey.Reset();
+	MirroredInstancedMeshComponentByGeometryKey.Reset();
+	TranslucentInstancedMeshComponentByGeometryKey.Reset();
+	InstancedMeshComponentByGeometryKey.Reset();
+	GeometryTasksPendingLaunch.Reset();
+	MeshComponentsWaitingForMaterial.Reset();
+	MeshComponentsWaitingForGeometry.Reset();
+	StaticMeshCacheMap.Reset();
+
+	bInitialized = false;
+}
+
+void FChaosVDGeometryBuilder::CreateMeshesFromImplicitObject(const Chaos::FImplicitObject* InImplicitObject, AActor* Owner, TArray<TSharedPtr<FChaosVDExtractedGeometryDataHandle>>& OutMeshDataHandles, const int32 DesiredLODCount, const Chaos::FRigidTransform3& InTransform, const int32 MeshIndex)
+{
+	// To start set the leaf and the root to the same ptr. If the object is an union, in the subsequent recursive call the leaf will be set correctly
+	CreateMeshesFromImplicit_Internal(InImplicitObject, InImplicitObject, Owner, OutMeshDataHandles, DesiredLODCount, InTransform, MeshIndex);
 }
 
 void FChaosVDGeometryBuilder::AddReferencedObjects(FReferenceCollector& Collector)
@@ -290,6 +383,55 @@ bool FChaosVDGeometryBuilder::HasNegativeScale(const Chaos::FRigidTransform3& In
 	return ScaleSignVector.X * ScaleSignVector.Y * ScaleSignVector.Z < 0;
 }
 
+void FChaosVDGeometryBuilder::CreateMeshesFromImplicit_Internal(const Chaos::FImplicitObject* InRootImplicitObject, const Chaos::FImplicitObject* InLeafImplicitObject, AActor* Owner, TArray<TSharedPtr<FChaosVDExtractedGeometryDataHandle>>& OutMeshDataHandles, const int32 DesiredLODCount, const Chaos::FRigidTransform3& InTransform, const int32 ShapeInstanceIndex)
+{
+	TRACE_CPUPROFILER_EVENT_SCOPE(FChaosVDGeometryBuilder::CreateMeshesFromImplicit_Internal);
+
+	using namespace Chaos;
+
+	const EImplicitObjectType InnerType = GetInnerType(InLeafImplicitObject->GetType());
+	
+	if (InnerType == ImplicitObjectType::Union || InnerType == ImplicitObjectType::UnionClustered)
+	{
+		if (const FImplicitObjectUnion* Union = InLeafImplicitObject->template AsA<FImplicitObjectUnion>())
+		{
+			const bool bIsRootUnion = InRootImplicitObject == InLeafImplicitObject;
+
+			for (int32 ObjectIndex = 0; ObjectIndex < Union->GetObjects().Num(); ++ObjectIndex)
+			{
+				const FImplicitObjectPtr& UnionImplicit = Union->GetObjects()[ObjectIndex];
+
+				// If this union it is not the root implicit object, then all its objects will share the same Instance index
+				int32 CurrentShapeInstanceIndex = bIsRootUnion ? ObjectIndex : ShapeInstanceIndex;
+
+				CreateMeshesFromImplicit_Internal(InRootImplicitObject, UnionImplicit.GetReference(), Owner, OutMeshDataHandles, DesiredLODCount, InTransform, CurrentShapeInstanceIndex);	
+			}
+		}
+
+		return;
+	}
+
+	if (InnerType == ImplicitObjectType::Transformed)
+	{
+		if (const TImplicitObjectTransformed<FReal, 3>* Transformed = InLeafImplicitObject->template GetObject<TImplicitObjectTransformed<FReal, 3>>())
+		{
+			// For transformed objects, the Instance index is the same so we pass it in without changing it
+			CreateMeshesFromImplicit_Internal(InRootImplicitObject, Transformed->GetTransformedObject(), Owner, OutMeshDataHandles, DesiredLODCount, Transformed->GetTransform(), ShapeInstanceIndex);
+		}
+		
+		return;
+	}
+
+	if (const TSharedPtr<FChaosVDExtractedGeometryDataHandle> MeshDataHandle = ExtractGeometryDataForImplicit(InLeafImplicitObject, InTransform))
+	{
+		MeshDataHandle->SetImplicitObject(InLeafImplicitObject);
+		MeshDataHandle->SetShapeInstanceIndex(ShapeInstanceIndex);
+		MeshDataHandle->SetRootImplicitObject(InRootImplicitObject);
+
+		OutMeshDataHandles.Add(MeshDataHandle);
+	}
+}
+
 bool FChaosVDGeometryBuilder::HasGeometryInCache(uint32 GeometryKey)
 {
 	FReadScopeLock ReadLock(GeometryCacheRWLock);
@@ -403,42 +545,45 @@ UStaticMesh* FChaosVDGeometryBuilder::CreateAndCacheStaticMesh(const uint32 Geom
 	return MainStaticMesh;
 }
 
-void FChaosVDGeometryBuilder::SetMeshComponentMaterial(EChaosVDMeshAttributesFlags MeshComponentAttributeFlags, UMeshComponent* MeshComponent)
+void FChaosVDGeometryBuilder::SetMeshComponentMaterial(IChaosVDGeometryComponent* GeometryComponent)
 {
-	UMaterialInstanceDynamic* Material = nullptr;
-
-	if (Cast<UChaosVDInstancedStaticMeshComponent>(MeshComponent))
-	{
-		Material = FChaosVDGeometryComponentUtils::CreateMaterialInstance(FChaosVDGeometryComponentUtils::GetMaterialTypeForComponent<UChaosVDInstancedStaticMeshComponent>(MeshComponentAttributeFlags));
-	}
-	else
+	if (!GeometryComponent)
 	{
-		Material = FChaosVDGeometryComponentUtils::CreateMaterialInstance(FChaosVDGeometryComponentUtils::GetMaterialTypeForComponent<UChaosVDStaticMeshComponent>(MeshComponentAttributeFlags));
+		return;
 	}
 
+	UMaterialInterface* Material = ComponentMeshPool.GetMaterialForType(GeometryComponent->GetMaterialType());
 	ensure(Material);
 
-	MeshComponent->SetMaterial(0, Material);
+	if (UMeshComponent* AsMeshComponent = Cast<UMeshComponent>(GeometryComponent))
+	{
+		AsMeshComponent->SetMaterial(0, Material);
+	}
 }
 
 void FChaosVDGeometryBuilder::DestroyMeshComponent(UMeshComponent* MeshComponent)
 {
-	if (Cast<UChaosVDInstancedStaticMeshComponent>(MeshComponent))
+	if (IChaosVDGeometryComponent* AsCVDGeometryComponent = Cast<IChaosVDGeometryComponent>(MeshComponent))
 	{
-		if (IChaosVDGeometryComponent* AsCVDGeometryComponent = Cast<IChaosVDGeometryComponent>(MeshComponent))
+		if (Cast<UChaosVDInstancedStaticMeshComponent>(MeshComponent))
 		{
-			TMap<uint32, UChaosVDInstancedStaticMeshComponent*>& InstancedMeshComponentCache = GetInstancedStaticMeshComponentCacheMap(AsCVDGeometryComponent->GetMeshComponentAttributeFlags());
+			EChaosVDMeshAttributesFlags MeshAttributes = AsCVDGeometryComponent->GetMeshComponentAttributeFlags();
+			TMap<uint32, UChaosVDInstancedStaticMeshComponent*>& InstancedMeshComponentCache = GetInstancedStaticMeshComponentCacheMap(MeshAttributes);
 			InstancedMeshComponentCache.Remove(AsCVDGeometryComponent->GetGeometryKey());
-
-			RemoveMeshComponentWaitingForGeometry(AsCVDGeometryComponent->GetGeometryKey(), MeshComponent);
-
-			AsCVDGeometryComponent->OnComponentEmpty()->RemoveAll(this);
 		}
+		
+		RemoveMeshComponentWaitingForGeometry(AsCVDGeometryComponent->GetGeometryKey(), MeshComponent);	
+        AsCVDGeometryComponent->OnComponentEmpty()->RemoveAll(this);
 	}
 
 	MeshComponentsPendingDisposal.Add(MeshComponent);
 }
 
+void FChaosVDGeometryBuilder::RequestMaterialInstance(UMeshComponent* MeshComponent)
+{
+	MeshComponentsWaitingForMaterial->EnqueueObject(MeshComponent);
+}
+
 TMap<uint32, UChaosVDInstancedStaticMeshComponent*>& FChaosVDGeometryBuilder::GetInstancedStaticMeshComponentCacheMap(EChaosVDMeshAttributesFlags MeshAttributeFlags)
 {
 	if (EnumHasAnyFlags(MeshAttributeFlags, EChaosVDMeshAttributesFlags::MirroredGeometry))
@@ -656,6 +801,36 @@ void FChaosVDGeometryBuilder::AdjustedTransformForImplicit(const Chaos::FImplici
 	}
 }
 
+TSharedPtr<FChaosVDExtractedGeometryDataHandle> FChaosVDGeometryBuilder::ExtractGeometryDataForImplicit(const Chaos::FImplicitObject* InImplicitObject, const Chaos::FRigidTransform3& InTransform)
+{
+	const uint32 ImplicitObjectHash = InImplicitObject->GetTypeHash();
+
+	Chaos::FRigidTransform3 ExtractedTransform = InTransform;
+	const bool bNeedsUnpack = ImplicitObjectNeedsUnpacking(InImplicitObject);
+	if (const Chaos::FImplicitObject* ImplicitObjectToProcess = bNeedsUnpack ? UnpackImplicitObject(InImplicitObject, ExtractedTransform) : InImplicitObject)
+	{
+		TSharedPtr<FChaosVDExtractedGeometryDataHandle> MeshDataHandle = MakeShared<FChaosVDExtractedGeometryDataHandle>();
+
+		const uint32 GeometryKey = ImplicitObjectToProcess->GetTypeHash();		
+		MeshDataHandle->SetGeometryKey(GeometryKey);
+
+		// For the Component data key, we need the hash of the implicit as it is (packed) because we will need to match it when looking for shape data
+		MeshDataHandle->SetDataComponentKey(bNeedsUnpack ? ImplicitObjectHash : GeometryKey);
+
+		AdjustedTransformForImplicit(ImplicitObjectToProcess, ExtractedTransform);
+		MeshDataHandle->SetTransform(ExtractedTransform);
+
+		if (!HasGeometryInCache(GeometryKey))
+		{
+			DispatchCreateAndCacheMeshForImplicitAsync(GeometryKey, ImplicitObjectToProcess);
+		}
+
+		return MeshDataHandle;
+	}
+
+	return nullptr;
+}
+
 bool FChaosVDGeometryBuilder::ImplicitObjectNeedsUnpacking(const Chaos::FImplicitObject* InImplicitObject) const
 {
 	using namespace Chaos;
@@ -749,14 +924,39 @@ void FChaosVDGeometryBuilder::HandleStaticMeshComponentInstanceIndexUpdated(UIns
 	}
 }
 
-void FGeometryGenerationTask::GenerateGeometry()
+void FChaosVDGeometryBuilder::DispatchCreateAndCacheMeshForImplicitAsync(const uint32 GeometryKey,  const Chaos::FImplicitObject* ImplicitObject, const int32 LODsToGenerateNum)
+{
+	ensure(IsInGameThread());
+
+	{
+		FReadScopeLock ReadLock(GeometryCacheRWLock);
+		if (GeometryBeingGeneratedByKey.Contains(GeometryKey))
+		{
+			return;
+		}
+	}
+
+	TSharedPtr<FChaosVDGeometryGenerationTask> GenerationTask = MakeShared<FChaosVDGeometryGenerationTask>(AsWeak(), GeometryKey, ImplicitObject, LODsToGenerateNum);
+
+	{
+		FWriteScopeLock WriteLock(GeometryCacheRWLock);
+		GeometryBeingGeneratedByKey.Add(GeometryKey, GenerationTask);
+	}
+
+	GeometryTasksPendingLaunch->EnqueueObject(GenerationTask);
+}
+
+void FChaosVDGeometryGenerationTask::GenerateGeometry()
 {
 	if (const TSharedPtr<FChaosVDGeometryBuilder> BuilderPtr = Builder.Pin())
 	{
-		BuilderPtr->CreateAndCacheStaticMesh(GeometryKey, *MeshGenerator.Get(), LODsToGenerateNum);
+		if (TSharedPtr<UE::Geometry::FMeshShapeGenerator> MeshGenerator = BuilderPtr->CreateMeshGeneratorForImplicitObject(ImplicitObject))
 		{
-			FWriteScopeLock WriteLock(BuilderPtr->GeometryCacheRWLock);
-			BuilderPtr->GeometryBeingGeneratedByKey.Remove(GeometryKey);
+			BuilderPtr->CreateAndCacheStaticMesh(GeometryKey, *MeshGenerator.Get(), LODsToGenerateNum);
+			{
+				FWriteScopeLock WriteLock(BuilderPtr->GeometryCacheRWLock);
+				BuilderPtr->GeometryBeingGeneratedByKey.Remove(GeometryKey);
+			}
 		}
 	}
 }
diff --git a/Engine/Plugins/ChaosVD/Source/ChaosVD/Private/ChaosVDGeometryBuilder.h b/Engine/Plugins/ChaosVD/Source/ChaosVD/Private/ChaosVDGeometryBuilder.h
--- a/Engine/Plugins/ChaosVD/Source/ChaosVD/Private/ChaosVDGeometryBuilder.h
+++ b/Engine/Plugins/ChaosVD/Source/ChaosVD/Private/ChaosVDGeometryBuilder.h
@@ -7,18 +7,17 @@
 #include "ChaosVDHeightfieldMeshGenerator.h"
 #include "ChaosVDMeshComponentPool.h"
 #include "ChaosVDModule.h"
-#include "ChaosVDParticleActor.h"
+#include "ChaosVDObjectPool.h"
 #include "ChaosVDScene.h"
 #include "ObjectsWaitingGeometryList.h"
 #include "Chaos/ImplicitObjectScaled.h"
 #include "Chaos/ImplicitObjectType.h"
 #include "Containers/Ticker.h"
-#include "Components/DynamicMeshComponent.h"
 #include "Components/InstancedStaticMeshComponent.h"
 #include "Engine/StaticMesh.h"
 
-#include "UDynamicMesh.h"
 #include "Components/ChaosVDInstancedStaticMeshComponent.h"
+#include "Components/ChaosVDStaticMeshComponent.h"
 #include "Materials/MaterialInstanceDynamic.h"
 #include "UObject/WeakObjectPtr.h"
 #include "Tasks/Task.h"
@@ -26,8 +25,7 @@
 #include "UObject/GCObject.h"
 #include "UObject/UObjectGlobals.h"
 
-class FGeometryGenerationTask;
-class UChaosVDStaticMeshComponent;
+class FChaosVDGeometryGenerationTask;
 class AActor;
 class UDynamicMesh;
 class UDynamicMeshComponent;
@@ -56,23 +54,10 @@ public:
 		FInstancedStaticMeshDelegates::OnInstanceIndexUpdated.AddRaw(this, &FChaosVDGeometryBuilder::HandleStaticMeshComponentInstanceIndexUpdated);
 	}
 
-	virtual ~FChaosVDGeometryBuilder() override
-	{
-		FTSTicker::GetCoreTicker().RemoveTicker(GameThreadTickDelegate);
-
-		FInstancedStaticMeshDelegates::OnInstanceIndexUpdated.RemoveAll(this);
-
-		for (const TPair<uint32, TObjectPtr<UStaticMesh>>& StaticMeshByKey : StaticMeshCacheMap)
-		{
-			if (StaticMeshByKey.Value)
-			{
-				StaticMeshByKey.Value->ClearFlags(RF_Standalone);
-				StaticMeshByKey.Value->MarkAsGarbage();
-			}
-		}
-	}
+	virtual ~FChaosVDGeometryBuilder() override;
 	
 	void Initialize(const TWeakPtr<FChaosVDScene>& ChaosVDScene);
+	void DeInitialize();
 
 	/** Creates Dynamic Mesh components for each object within the provided Implicit object
 	 *	@param InImplicitObject : Implicit object to process
@@ -82,14 +67,9 @@ public:
 	 *	@param MeshIndex Index of the current component being processed. This is useful when this method is called recursively
 	 *	@param InTransform to apply to the generated components/geometry
 	 */
-	
-	template<typename MeshType>
-	void CreateMeshesFromImplicitObject(const Chaos::FImplicitObject* InImplicitObject, AActor* Owner, TArray<TSharedPtr<FChaosVDExtractedGeometryDataHandle>>& OutMeshDataHandles, const int32 DesiredLODCount = 0, const Chaos::FRigidTransform3& InTransform = Chaos::FRigidTransform3(), const int32 MeshIndex = 0)
-	{
-		// To start set the leaf and the root to the same ptr. If the object is an union, in the subsequent recursive call the leaf will be set correctly
-		CreateMeshesFromImplicit_Internal<MeshType>(InImplicitObject, InImplicitObject, Owner, OutMeshDataHandles, DesiredLODCount, InTransform, MeshIndex);
-	}
-	
+
+	void CreateMeshesFromImplicitObject(const Chaos::FImplicitObject* InImplicitObject, AActor* Owner, TArray<TSharedPtr<FChaosVDExtractedGeometryDataHandle>>& OutMeshDataHandles, const int32 DesiredLODCount = 0, const Chaos::FRigidTransform3& InTransform = Chaos::FRigidTransform3(), const int32 MeshIndex = 0);
+
 	virtual void AddReferencedObjects(FReferenceCollector& Collector) override;
 	virtual FString GetReferencerName() const override
 	{
@@ -112,7 +92,6 @@ public:
 
 private:
 	
-	template<typename MeshType>
 	void CreateMeshesFromImplicit_Internal(const Chaos::FImplicitObject* InRootImplicitObject,const Chaos::FImplicitObject* InLeafImplicitObject, AActor* Owner, TArray<TSharedPtr<FChaosVDExtractedGeometryDataHandle>>& OutMeshDataHandles, const int32 DesiredLODCount = 0, const Chaos::FRigidTransform3& InTransform = Chaos::FRigidTransform3(), const int32 ShapeInstanceIndex = 0);
 	
 public:
@@ -144,7 +123,7 @@ private:
 	template <class ComponentType>
     bool InitializeMeshComponent(AActor* Owner, ComponentType* MeshComponent);
 
-	void SetMeshComponentMaterial(EChaosVDMeshAttributesFlags MeshComponentAttributeFlags, UMeshComponent* MeshComponent);
+	void SetMeshComponentMaterial(IChaosVDGeometryComponent* GeometryComponent);
 
 public:
 	/**
@@ -171,6 +150,8 @@ public:
 	 * */
 	void DestroyMeshComponent(UMeshComponent* MeshComponent);
 
+	void RequestMaterialInstance(UMeshComponent* MeshComponent);
+
 private:
 
 	/** Gets a ptr to a fully initialized Mesh component compatible with the provided geometry handle and mesh attribute flags, ready to accept a new mesh instance
@@ -230,7 +211,6 @@ private:
 	/** Extracts data from an implicit object in a format CVD can use, and starts the Mesh generation process if needed
 	 * @return Returns a handle to the generated data that can be used to access the generated mesh when ready
 	 */
-	template <typename MeshType>
 	TSharedPtr<FChaosVDExtractedGeometryDataHandle> ExtractGeometryDataForImplicit(const Chaos::FImplicitObject* InImplicitObject, const Chaos::FRigidTransform3& InTransform);
 
 	/** Returns true if the implicit object if of one of the types we need to unpack before generating a mesh for it */
@@ -239,11 +219,10 @@ private:
 	/**
 	 * Creates a Mesh from the provided Implicit object geometry data. This is a async operation, and the mesh will be assigned to the component once is ready
 	 * @param GeometryCacheKey Key to be used to find this geometry in the cache
-	 * @param MeshGenerator Generator class that has the data to generate the mesh
+	 * @param ImplicitObject Implicit object to use a sdata source to create the static mesh
 	 * @param LODsToGenerateNum Num of LODs to Generate. Not all mesh types support this
 	 * */
-	template<typename MeshType>
-	void DispatchCreateAndCacheMeshForImplicitAsync(const uint32 GeometryCacheKey, TSharedPtr<UE::Geometry::FMeshShapeGenerator> MeshGenerator, const int32 LODsToGenerateNum = 0);
+	void DispatchCreateAndCacheMeshForImplicitAsync(const uint32 GeometryCacheKey, const Chaos::FImplicitObject* ImplicitObject, const int32 LODsToGenerateNum = 0);
 
 	/* Process an FImplicitObject and returns de desired geometry type. Could be directly the shape or another version of the implicit */
 	template <bool bIsInstanced, typename GeometryType>
@@ -271,6 +250,8 @@ public:
 	 */
 	void RemoveMeshComponentWaitingForGeometry(uint32 GeometryKey, TWeakObjectPtr<UMeshComponent> MeshComponent) const;
 
+	FChaosVDMeshComponentPool& GetMeshComponentDataPool() { return ComponentMeshPool; }
+
 private:
 
 	/** Handles any changes to the indexes of created instanced mesh components we are managing, making corrections/updates as needed */
@@ -280,7 +261,7 @@ private:
 	TMap<uint32, TObjectPtr<UStaticMesh>> StaticMeshCacheMap;
 
 	/** Set of all geometry keys of the Meshes that are being generated but not ready yet */
-	TSet<uint32> GeometryBeingGeneratedByKey;
+	TMap<uint32, TSharedPtr<FChaosVDGeometryGenerationTask>> GeometryBeingGeneratedByKey;
 	
 	/** Used to lock Read or Writes to the Geometry cache and in flight job tracking containers */
 	FRWLock GeometryCacheRWLock;
@@ -293,7 +274,7 @@ private:
 
 	TUniquePtr<FObjectsWaitingProcessingQueue<FMeshComponentWeakPtr>> MeshComponentsWaitingForMaterial;
 
-	TUniquePtr<FObjectsWaitingProcessingQueue<TSharedPtr<FGeometryGenerationTask>>> GeometryTasksPendingLaunch;
+	TUniquePtr<FObjectsWaitingProcessingQueue<TSharedPtr<FChaosVDGeometryGenerationTask>>> GeometryTasksPendingLaunch;
 
 	/** Map containing already initialized Instanced static mesh components for any given geometry key */
 	TMap<uint32, UChaosVDInstancedStaticMeshComponent*> InstancedMeshComponentByGeometryKey;
@@ -315,122 +296,47 @@ private:
 
 	/** Weak Ptr to the CVD scene owning this geometry builder */
 	TWeakPtr<FChaosVDScene> SceneWeakPtr;
+
+	bool bInitialized = false;
 	
-	friend class FGeometryGenerationTask;
+	friend class FChaosVDGeometryGenerationTask;
 };
 
 /** Used to execute each individual Geometry Generation task using the data with which was constructed.
  * It allows the task to skip the actual generation attempt if the Geometry builder instance goes away which happens when the tool is closed
  */
-class FGeometryGenerationTask
+class FChaosVDGeometryGenerationTask
 {
 public:
-	FGeometryGenerationTask(const TWeakPtr<FChaosVDGeometryBuilder>& InBuilder, const TSharedPtr<UE::Geometry::FMeshShapeGenerator>& InGenerator, const uint32 GeometryKey,
+	FChaosVDGeometryGenerationTask(const TWeakPtr<FChaosVDGeometryBuilder>& InBuilder, const uint32 GeometryKey, const Chaos::FImplicitObject* ImplicitObject,
 		const int32 LODsToGenerateNum)
 		: Builder(InBuilder),
-		  MeshGenerator(InGenerator),
+		  ImplicitObject(ImplicitObject),
 		  GeometryKey(GeometryKey),
-		  LODsToGenerateNum(LODsToGenerateNum)
+		  LODsToGenerateNum(LODsToGenerateNum),
+		  bIsCanceled(false)
 	{
 	}
 
 	void GenerateGeometry();
 
-private:
-	TWeakPtr<FChaosVDGeometryBuilder> Builder;
-	TSharedPtr<UE::Geometry::FMeshShapeGenerator> MeshGenerator;
-	uint32 GeometryKey;
-	int32 LODsToGenerateNum;
-};
-
-template <typename MeshType>
-TSharedPtr<FChaosVDExtractedGeometryDataHandle> FChaosVDGeometryBuilder::ExtractGeometryDataForImplicit(const Chaos::FImplicitObject* InImplicitObject, const Chaos::FRigidTransform3& InTransform)
-{
-	const uint32 ImplicitObjectHash = InImplicitObject->GetTypeHash();
-
-	Chaos::FRigidTransform3 ExtractedTransform = InTransform;
-	const bool bNeedsUnpack = ImplicitObjectNeedsUnpacking(InImplicitObject);
-	if (const Chaos::FImplicitObject* ImplicitObjectToProcess = bNeedsUnpack ? UnpackImplicitObject(InImplicitObject, ExtractedTransform) : InImplicitObject)
+	uint32 GetGeometryKey() const
 	{
-		TSharedPtr<FChaosVDExtractedGeometryDataHandle> MeshDataHandle = MakeShared<FChaosVDExtractedGeometryDataHandle>();
-
-		const uint32 GeometryKey = ImplicitObjectToProcess->GetTypeHash();		
-		MeshDataHandle->SetGeometryKey(GeometryKey);
-	
-		// For the Component data key, we need the hash of the implicit as it is (packed) because we will need to match it when looking for shape data
-		MeshDataHandle->SetDataComponentKey(bNeedsUnpack ? ImplicitObjectHash : GeometryKey);
-
-		AdjustedTransformForImplicit(ImplicitObjectToProcess, ExtractedTransform);
-		MeshDataHandle->SetTransform(ExtractedTransform);
-
-		if (!HasGeometryInCache(GeometryKey))
-		{
-			if (TSharedPtr<UE::Geometry::FMeshShapeGenerator> MeshGenerator = CreateMeshGeneratorForImplicitObject(ImplicitObjectToProcess))
-			{
-				DispatchCreateAndCacheMeshForImplicitAsync<MeshType>(GeometryKey, MeshGenerator);
-			}
-			else
-			{
-				UE_LOG(LogChaosVDEditor, Error, TEXT("[%s] Failed create geometry data handle | Failed to generate a valid mesh generator from the implicit object. | Geometry Key [%u] "), ANSI_TO_TCHAR(__FUNCTION__), GeometryKey);
-				return nullptr;
-			}
-		}
-
-		return MeshDataHandle;
+		return GeometryKey;
 	}
 
-	return nullptr;
-}
-
-template <typename MeshType>
-void FChaosVDGeometryBuilder::CreateMeshesFromImplicit_Internal(const Chaos::FImplicitObject* InRootImplicitObject, const Chaos::FImplicitObject* InLeafImplicitObject, AActor* Owner, TArray<TSharedPtr<FChaosVDExtractedGeometryDataHandle>>& OutMeshDataHandles, const int32 DesiredLODCount, const Chaos::FRigidTransform3& InTransform, const int32 ShapeInstanceIndex)
-{
-	TRACE_CPUPROFILER_EVENT_SCOPE(FChaosVDGeometryBuilder::CreateMeshesFromImplicit_Internal);
-
-	using namespace Chaos;
-
-	const EImplicitObjectType InnerType = GetInnerType(InLeafImplicitObject->GetType());
-	
-	if (InnerType == ImplicitObjectType::Union || InnerType == ImplicitObjectType::UnionClustered)
-	{
-		if (const FImplicitObjectUnion* Union = InLeafImplicitObject->template AsA<FImplicitObjectUnion>())
-		{
-			const bool bIsRootUnion = InRootImplicitObject == InLeafImplicitObject;
+	bool IsCanceled() { return bIsCanceled; }
+	void CancelTask(){ bIsCanceled = true; }
 
-			for (int32 ObjectIndex = 0; ObjectIndex < Union->GetObjects().Num(); ++ObjectIndex)
-			{
-				const FImplicitObjectPtr& UnionImplicit = Union->GetObjects()[ObjectIndex];
+	UE::Tasks::FTask TaskHandle;
 
-				// If this union it is not the root implicit object, then all its objects will share the same Instance index
-				int32 CurrentShapeInstanceIndex = bIsRootUnion ? ObjectIndex : ShapeInstanceIndex;
-
-				CreateMeshesFromImplicit_Internal<MeshType>(InRootImplicitObject, UnionImplicit.GetReference(), Owner, OutMeshDataHandles, DesiredLODCount, InTransform, CurrentShapeInstanceIndex);	
-			}
-		}
-
-		return;
-	}
-
-	if (InnerType == ImplicitObjectType::Transformed)
-	{
-		if (const TImplicitObjectTransformed<FReal, 3>* Transformed = InLeafImplicitObject->template GetObject<TImplicitObjectTransformed<FReal, 3>>())
-		{
-			// For transformed objects, the Instance index is the same so we pass it in without changing it
-			CreateMeshesFromImplicit_Internal<MeshType>(InRootImplicitObject, Transformed->GetTransformedObject(), Owner, OutMeshDataHandles, DesiredLODCount, Transformed->GetTransform(), ShapeInstanceIndex);
-		}
-		
-		return;
-	}
-
-	if (const TSharedPtr<FChaosVDExtractedGeometryDataHandle> MeshDataHandle = ExtractGeometryDataForImplicit<MeshType>(InLeafImplicitObject, InTransform))
-	{
-		MeshDataHandle->SetImplicitObject(InLeafImplicitObject);
-		MeshDataHandle->SetShapeInstanceIndex(ShapeInstanceIndex);
-		MeshDataHandle->SetRootImplicitObject(InRootImplicitObject);
-
-		OutMeshDataHandles.Add(MeshDataHandle);
-	}
-}
+private:
+	TWeakPtr<FChaosVDGeometryBuilder> Builder;
+	const Chaos::FImplicitObject* ImplicitObject;
+	uint32 GeometryKey;
+	int32 LODsToGenerateNum;
+	std::atomic<bool> bIsCanceled;
+};
 
 template <typename ComponentType>
 bool FChaosVDGeometryBuilder::InitializeMeshComponent(AActor* Owner, ComponentType* MeshComponent)
@@ -444,7 +350,12 @@ bool FChaosVDGeometryBuilder::InitializeMeshComponent(AActor* Owner, ComponentTy
 	if (Owner)
 	{
 		Owner->AddOwnedComponent(MeshComponent);
-		MeshComponent->RegisterComponent();
+
+		if (ensure(!MeshComponent->IsRegistered()))
+		{
+			MeshComponent->RegisterComponent();
+		}
+		
 		MeshComponent->AttachToComponent(Owner->GetRootComponent(), FAttachmentTransformRules::SnapToTargetIncludingScale);
 
 		MeshComponent->bSelectable = true;
@@ -471,23 +382,8 @@ ComponentType* FChaosVDGeometryBuilder::GetAvailableInstancedStaticMeshComponent
 	}
 	else
 	{
-		// If no exiting component meets our requirements, get a new one form the pool
+		// If no existing component meets our requirements, get a new one form the pool
 		ComponentType* Component = ComponentMeshPool.AcquireMeshComponent<ComponentType>(MeshComponentsContainerActor, InExtractedGeometryDataHandle->GetTypeName());
-		if (!InitializeMeshComponent<ComponentType>(MeshComponentsContainerActor, Component))
-		{
-			return nullptr;
-		}
-
-		// If this is a Instanced Static Mesh Component, make sure we set the reverse culling flag correctly
-		// And the material
-		if (UChaosVDInstancedStaticMeshComponent* AsInstancedMeshComponent = Component)
-		{
-			AsInstancedMeshComponent->bReverseCulling = EnumHasAnyFlags(MeshComponentAttributeFlags, EChaosVDMeshAttributesFlags::MirroredGeometry);
-			if (!AsInstancedMeshComponent->HasOverrideMaterials())
-			{
-				MeshComponentsWaitingForMaterial->EnqueueObject(AsInstancedMeshComponent);
-			}	
-		}
 
 		bOutIsNewComponent = true;
 
@@ -508,18 +404,20 @@ ComponentType* FChaosVDGeometryBuilder::GetAvailableMeshComponent(const TSharedP
 	else
 	{
 		MeshComponent = ComponentMeshPool.AcquireMeshComponent<ComponentType>(MeshComponentsContainerActor, InExtractedGeometryDataHandle->GetTypeName());
+		bOutIsNewComponent = true;
+	}
+
+	if (bOutIsNewComponent)
+	{
 		if (!InitializeMeshComponent<ComponentType>(MeshComponentsContainerActor, MeshComponent))
 		{
 			return nullptr;
 		}
-
-		bOutIsNewComponent = true;
 	}
 
 	return MeshComponent;
 }
 
-
 template <typename ComponentType>
 TSharedPtr<FChaosVDMeshDataInstanceHandle> FChaosVDGeometryBuilder::CreateMeshDataInstance(const FChaosVDParticleDataWrapper& InOwningParticleData, const TSharedPtr<FChaosVDExtractedGeometryDataHandle>& InExtractedGeometryDataHandle)
 {
@@ -616,30 +514,15 @@ ComponentType* FChaosVDGeometryBuilder::GetMeshComponentForNewInstance(const TSh
 
 		if (bIsNew)
 		{
+			CVDGeometryComponent->SetGeometryBuilder(AsWeak());
+			CVDGeometryComponent->Initialize();
 			CVDGeometryComponent->OnComponentEmpty()->AddRaw(this, &FChaosVDGeometryBuilder::DestroyMeshComponent);			
-		}	
+		}
 	}
 
 	return Component;
 }
 
-template <typename MeshType>
-void FChaosVDGeometryBuilder::DispatchCreateAndCacheMeshForImplicitAsync(const uint32 GeometryKey, TSharedPtr<UE::Geometry::FMeshShapeGenerator> MeshGenerator, const int32 LODsToGenerateNum)
-{
-	{
-		FWriteScopeLock WriteLock(GeometryCacheRWLock);
-		if (GeometryBeingGeneratedByKey.Contains(GeometryKey))
-		{
-			return;
-		}
-
-		GeometryBeingGeneratedByKey.Add(GeometryKey);
-	}
-	
-	TSharedPtr<FGeometryGenerationTask> GenerationTask = MakeShared<FGeometryGenerationTask>(AsWeak(), MeshGenerator, GeometryKey, LODsToGenerateNum);
-	GeometryTasksPendingLaunch->EnqueueObject(GenerationTask);
-}
-
 template <bool bIsInstanced, typename GeometryType>
 const GeometryType* FChaosVDGeometryBuilder::GetGeometry(const Chaos::FImplicitObject* InImplicitObject, const bool bIsScaled, Chaos::FRigidTransform3& OutTransform) const
 {
diff --git a/Engine/Plugins/ChaosVD/Source/ChaosVD/Private/ChaosVDGeometryDataComponent.cpp b/Engine/Plugins/ChaosVD/Source/ChaosVD/Private/ChaosVDGeometryDataComponent.cpp
--- a/Engine/Plugins/ChaosVD/Source/ChaosVD/Private/ChaosVDGeometryDataComponent.cpp
+++ b/Engine/Plugins/ChaosVD/Source/ChaosVD/Private/ChaosVDGeometryDataComponent.cpp
@@ -162,12 +162,15 @@ void FChaosVDMeshDataInstanceHandle::SetIsSelected(bool bInIsSelected)
 
 void FChaosVDMeshDataInstanceHandle::SetVisibility(bool bInIsVisible)
 {
-	if (IChaosVDGeometryComponent* CVDGeometryComponent = Cast<IChaosVDGeometryComponent>(GetMeshComponent()))
+	if (InstanceState.bIsVisible != bInIsVisible)
 	{
-		CVDGeometryComponent->UpdateInstanceVisibility(AsShared(), bInIsVisible);
-	}
+		if (IChaosVDGeometryComponent* CVDGeometryComponent = Cast<IChaosVDGeometryComponent>(GetMeshComponent()))
+		{
+			CVDGeometryComponent->UpdateInstanceVisibility(AsShared(), bInIsVisible);
+		}
 
-	InstanceState.bIsVisible = bInIsVisible;
+		InstanceState.bIsVisible = bInIsVisible;
+	}
 }
 
 void FChaosVDMeshDataInstanceHandle::HandleInstanceIndexUpdated(TArrayView<const FInstancedStaticMeshDelegates::FInstanceIndexUpdateData> InIndexUpdates)
@@ -384,37 +387,19 @@ UMaterialInterface* FChaosVDGeometryComponentUtils::GetBaseMaterialForType(EChao
 
 	switch(Type)
 	{
-		case EChaosVDMaterialType::QueryOnlyMaterial:
+		case EChaosVDMaterialType::SMTranslucent:
 				return EditorSettings->QueryOnlyMeshesMaterial.Get();
-		case EChaosVDMaterialType::SimOnlyMaterial:
+		case EChaosVDMaterialType::SMOpaque:
 				return EditorSettings->SimOnlyMeshesMaterial.Get();
-		case EChaosVDMaterialType::Instanced:
+		case EChaosVDMaterialType::ISMCOpaque:
 				return EditorSettings->InstancedMeshesMaterial.Get();
-		case EChaosVDMaterialType::InstancedQueryOnly:
+		case EChaosVDMaterialType::ISMCTranslucent:
 				return EditorSettings->InstancedMeshesQueryOnlyMaterial.Get();
 		default:
 			return nullptr;
 	}	
 }
 
-UMaterialInstanceDynamic* FChaosVDGeometryComponentUtils::CreateMaterialInstance(UMaterialInterface* BaseMaterial)
-{
-	if (!BaseMaterial)
-	{
-		return nullptr;
-	}
-
-	UMaterialInstanceDynamic* DynamicMaterial = UMaterialInstanceDynamic::Create(BaseMaterial, nullptr);
-	DynamicMaterial->SetFlags(RF_Transient);
-	
-	return DynamicMaterial;
-}
-
-UMaterialInstanceDynamic* FChaosVDGeometryComponentUtils::CreateMaterialInstance(EChaosVDMaterialType Type)
-{
-	return CreateMaterialInstance(GetBaseMaterialForType(Type));
-}
-
 void Chaos::VisualDebugger::SelectParticleWithGeometryInstance(const TSharedRef<FChaosVDScene>& InScene, IChaosVDGeometryOwnerInterface* GeometryOwner, const TSharedPtr<FChaosVDMeshDataInstanceHandle>& InMeshDataHandle)
 {
 	InScene->SetSelectedObject(nullptr);
diff --git a/Engine/Plugins/ChaosVD/Source/ChaosVD/Private/ChaosVDGeometryDataComponent.h b/Engine/Plugins/ChaosVD/Source/ChaosVD/Private/ChaosVDGeometryDataComponent.h
--- a/Engine/Plugins/ChaosVD/Source/ChaosVD/Private/ChaosVDGeometryDataComponent.h
+++ b/Engine/Plugins/ChaosVD/Source/ChaosVD/Private/ChaosVDGeometryDataComponent.h
@@ -37,10 +37,10 @@ DECLARE_MULTICAST_DELEGATE_OneParam(FChaosVDMeshComponentEmptyDelegate, UMeshCom
 UENUM()
 enum class EChaosVDMaterialType
 {
-	SimOnlyMaterial,
-	QueryOnlyMaterial,
-	Instanced,
-	InstancedQueryOnly
+	SMOpaque,
+	SMTranslucent,
+	ISMCOpaque,
+	ISMCTranslucent
 };
 
 UENUM()
@@ -146,7 +146,7 @@ public:
 	void SetGeometryHandle(const TSharedPtr<FChaosVDExtractedGeometryDataHandle>& InHandle);
 
 	/** Returns the geometry handle used to create the mesh instance this handle represents */
-	TSharedPtr<FChaosVDExtractedGeometryDataHandle> GetGeometryHandle() const { return ExtractedGeometryHandle; }
+	const TSharedPtr<FChaosVDExtractedGeometryDataHandle>& GetGeometryHandle() const { return ExtractedGeometryHandle; }
 
 	/** Applies to the provided color to the mesh instance this handle represents */
 	void SetInstanceColor(const FLinearColor& NewColor);
@@ -168,6 +168,7 @@ public:
 
 	/** Sets a Ptr to the geometry builder used to generate and manage the geometry/mesh components this handle represents */
 	void SetGeometryBuilder(const TWeakPtr<FChaosVDGeometryBuilder>& InGeometryBuilder) { GeometryBuilderInstance = InGeometryBuilder; }
+	TWeakPtr<FChaosVDGeometryBuilder> GetGeometryBuilder() { return GeometryBuilderInstance; }
 
 	/** Marks this mesh instance as selected. Used to handle Selection in Editor */
 	void SetIsSelected(bool bInIsSelected);
@@ -309,8 +310,15 @@ public:
 	/** Returns the CVD Mesh Attribute flags this component is compatible with*/
 	virtual EChaosVDMeshAttributesFlags GetMeshComponentAttributeFlags() const PURE_VIRTUAL(IChaosVDGeometryDataComponent::GetMeshComponentAttributeFlags, return EChaosVDMeshAttributesFlags::None;);
 
-	/** Resets the state of this mesh component so it can be re-used later on */
+	/** Resets the state of this mesh component, so it can be re-used later on */
 	virtual void Reset() PURE_VIRTUAL(IChaosVDGeometryDataComponent::Reset);
+
+	virtual void Initialize() PURE_VIRTUAL(IChaosVDGeometryDataComponent::Initialize);
+	
+	/** Sets a Ptr to the geometry builder used to generate and manage the geometry/mesh components */
+	virtual void SetGeometryBuilder(TWeakPtr<FChaosVDGeometryBuilder> GeometryBuilder) PURE_VIRTUAL(IChaosVDGeometryDataComponent::SetGeometryBuilder);
+
+	virtual EChaosVDMaterialType GetMaterialType() const PURE_VIRTUAL(IChaosVDGeometryDataComponent::GetMaterialType,  return EChaosVDMaterialType::SMOpaque;);
 };
 
 class FChaosVDGeometryComponentUtils
@@ -325,15 +333,9 @@ public:
 	/** Calculates the correct visibility state based on the particle state, and applies it to the mesh instance the provided handle represents */
 	static void UpdateMeshVisibility(const TSharedPtr<FChaosVDMeshDataInstanceHandle>& InInstanceHandle, const FChaosVDParticleDataWrapper& InParticleData, bool bIsActive);
 
-	/** Returns an the material to use as a base to create material instances for the provided type */
+	/** Returns the material to use as a base to create material instances for the provided type */
 	static UMaterialInterface* GetBaseMaterialForType(EChaosVDMaterialType Type);
 
-	/** Creates a material instance using the provided material as a base */
-	static UMaterialInstanceDynamic* CreateMaterialInstance(UMaterialInterface* BaseMaterial);
-
-	/** Creates a material instance for the provided CVD material type */
-	static UMaterialInstanceDynamic* CreateMaterialInstance(EChaosVDMaterialType Type);
-
 	/** Returns the correct material type to use based on the provided Component type and Mesh Attributes */
 	template<typename TComponent>
 	static EChaosVDMaterialType GetMaterialTypeForComponent(EChaosVDMeshAttributesFlags MeshAttributes);
@@ -350,11 +352,11 @@ EChaosVDMaterialType FChaosVDGeometryComponentUtils::GetMaterialTypeForComponent
 	constexpr bool bIsInstancedMeshComponent = std::is_base_of_v<UInstancedStaticMeshComponent, TComponent>;
 	if (EnumHasAnyFlags(MeshAttributes, EChaosVDMeshAttributesFlags::TranslucentGeometry))
 	{
-		return bIsInstancedMeshComponent ? EChaosVDMaterialType::InstancedQueryOnly : EChaosVDMaterialType::QueryOnlyMaterial;
+		return bIsInstancedMeshComponent ? EChaosVDMaterialType::ISMCTranslucent : EChaosVDMaterialType::SMTranslucent;
 	}
 	else
 	{
-		return bIsInstancedMeshComponent ? EChaosVDMaterialType::Instanced : EChaosVDMaterialType::SimOnlyMaterial;
+		return bIsInstancedMeshComponent ? EChaosVDMaterialType::ISMCOpaque : EChaosVDMaterialType::SMOpaque;
 	}
 }
 
diff --git a/Engine/Plugins/ChaosVD/Source/ChaosVD/Private/ChaosVDHeightfieldMeshGenerator.cpp b/Engine/Plugins/ChaosVD/Source/ChaosVD/Private/ChaosVDHeightfieldMeshGenerator.cpp
--- a/Engine/Plugins/ChaosVD/Source/ChaosVD/Private/ChaosVDHeightfieldMeshGenerator.cpp
+++ b/Engine/Plugins/ChaosVD/Source/ChaosVD/Private/ChaosVDHeightfieldMeshGenerator.cpp
@@ -4,82 +4,96 @@
 
 #include "Chaos/HeightField.h"
 
-void FChaosVDHeightFieldMeshGenerator::AppendTriangle(int32& OutTriIndex, int& OutCurrentNormalIndex, int32 PolygonIndex, const Chaos::TVec2<Chaos::FReal>& InCellCoordinates, const UE::Geometry::FIndex3i& InTriangle, const Chaos::FHeightField& InHeightField)
+void FChaosVDHeightFieldMeshGenerator::AppendTriangle(const Chaos::TVec2<Chaos::FReal>& InCellCoordinates, const UE::Geometry::FIndex3i& InTriangle, const Chaos::FHeightField& InHeightField, int32 PolygonID, int32 TriangleIndex)
 {
-	int32 NormalIndices[3];
-    for (int32 LocalVertexIndex = 0; LocalVertexIndex < 3; ++LocalVertexIndex)
-    {
-        Normals[OutCurrentNormalIndex] = UE::Math::TVector<float>(InHeightField.GetNormalAt(InCellCoordinates));
-        NormalIndices[LocalVertexIndex] = OutCurrentNormalIndex;
-        OutCurrentNormalIndex++;
-    }
-
-	SetTriangle(OutTriIndex, InTriangle);
-	SetTrianglePolygon(OutTriIndex, PolygonIndex);
-	SetTriangleNormals(OutTriIndex, NormalIndices[0], NormalIndices[1], NormalIndices[2]);
-	OutTriIndex++;
+	int32 StartNormalIndex = TriangleIndex * 3;
+	FVector3f CellNormal(InHeightField.GetNormalAt(InCellCoordinates));
+
+	for (int32 LocalVertexIndex = 0; LocalVertexIndex < 3; ++LocalVertexIndex)
+	{
+		Normals[StartNormalIndex + LocalVertexIndex] = CellNormal;
+	}
+
+	SetTriangle(TriangleIndex, InTriangle);
+	SetTrianglePolygon(TriangleIndex, PolygonID);
+	SetTriangleNormals(TriangleIndex, StartNormalIndex, StartNormalIndex + 1, StartNormalIndex + 2);
 }
 
 void FChaosVDHeightFieldMeshGenerator::GenerateFromHeightField(const Chaos::FHeightField& InHeightField)
 {
 	using namespace UE::Geometry;
 
-	const int32 NumRows = InHeightField.GetNumRows();
-	const int32 NumCols = InHeightField.GetNumCols();
-	const int32 VertexCount = NumRows * NumCols;
+	const int32 VertexGridNumRows = InHeightField.GetNumRows();
+	const int32 VertexGridNumColumns = InHeightField.GetNumCols();
+	const int32 VertexCount = VertexGridNumRows * VertexGridNumColumns;
 
-	const int32 NumQuads = (NumRows - 1) * (NumCols - 1);
+	// If we consider each polygon/quad we create a cell of a new grid, the resulting grid will be one element smaller in both dimensions
+	const int32 PolygonsGridNumRows = VertexGridNumRows -1;
+	const int32 PolygonsGridNumColumns = VertexGridNumColumns -1;
+
+	const int32 NumQuads = PolygonsGridNumRows * PolygonsGridNumColumns;
 	const int32 NumTris = NumQuads * 2;
+
 	const int32 NumNormals = NumTris * 3;
 	constexpr int32 NumUVs = 0;
 	
 	SetBufferSizes(VertexCount, NumTris, NumUVs, NumNormals);
 
-	//Fill the vertex buffer with the height data
-	for (int32 Y = 0; Y < NumRows; Y++)
+	// TODO: This value is not tuned yet
+	constexpr int32 MaxRowsNumToProcessInSingleThread = 8;
+
+	EParallelForFlags VerticesProcessingLoopFlags = VertexGridNumRows > MaxRowsNumToProcessInSingleThread ? EParallelForFlags::ForceSingleThread : EParallelForFlags::None;
+
+	// Fill the vertex buffer with the height data
+	ParallelFor(VertexGridNumRows, [this, VertexGridNumColumns, &InHeightField](int32 RowIndex)
 	{
-		for (int32 X = 0; X < NumCols; X++)
+		for (int32 ColumnIndex = 0; ColumnIndex < VertexGridNumColumns; ++ColumnIndex)
 		{
-			const int32 SampleIndex = Y * NumCols + X;
+			const int32 SampleIndex = RowIndex * VertexGridNumColumns + ColumnIndex;
 			Vertices[SampleIndex] = InHeightField.GetPointScaled(SampleIndex);
 		}
-	}
-  
-	int32 CurrentNormalIndex = 0;
-	int32 CurrentTriangleIndex = 0;
-	int32 CurrentPolygonIndex = 0;
-	for (int32 Y = 0; Y < NumRows - 1; Y++)
+	}, VerticesProcessingLoopFlags);
+
+	EParallelForFlags PolygonsRowsProcessingLoopFlags = PolygonsGridNumRows > MaxRowsNumToProcessInSingleThread ? EParallelForFlags::ForceSingleThread : EParallelForFlags::None;
+
+	ParallelFor(PolygonsGridNumRows, [this, VertexGridNumColumns, PolygonsGridNumColumns, &InHeightField](int32 RowIndex)
 	{
-		for (int32 X = 0; X < NumCols - 1; X++)
+		for (int32 ColumIndex = 0; ColumIndex < PolygonsGridNumColumns; ++ColumIndex)
 		{
-			if (InHeightField.IsHole(X, Y))
+			if (InHeightField.IsHole(ColumIndex, RowIndex))
 			{
-				continue;
+				return;
 			}
 
-			Chaos::TVec2<Chaos::FReal> CellCoordinates(X,Y);
+			Chaos::TVec2<Chaos::FReal> CellCoordinates(ColumIndex,RowIndex);
 
 			// Vertices of each corner of the current cell
-			const int32 Vertex0 = Y * NumCols + X;
+			// Starting from the vertex index in the actual vertex grid instead of the polygon one
+			const int32 Vertex0 = RowIndex * VertexGridNumColumns + ColumIndex;
 			const int32 Vertex1 = Vertex0 + 1;
-			const int32 Vertex2 = Vertex0 + NumCols;
+			const int32 Vertex2 = Vertex0 + VertexGridNumColumns;
 			const int32 Vertex3 = Vertex2 + 1;
 
 			// Define the two triangles that form the current cell and add it to the generator
 			FIndex3i Triangle(Vertex0,Vertex3,Vertex1);
 			FIndex3i Triangle2(Vertex0,Vertex2,Vertex3);
 
-			AppendTriangle(CurrentTriangleIndex,CurrentNormalIndex, CurrentPolygonIndex, CellCoordinates, Triangle, InHeightField);
-			AppendTriangle(CurrentTriangleIndex,CurrentNormalIndex, CurrentPolygonIndex, CellCoordinates, Triangle2, InHeightField);
+			// Calculate which index in the polygon grid we are creating we are supposed to be in if it was represented in one-dimension array
+			// And then use that as Polygon id for both triangles we are creating to for the Quad polygon
+			const int32 PolygonID = RowIndex * PolygonsGridNumColumns + ColumIndex;
 
-			CurrentPolygonIndex++;
+			// The start triangle index will always be 2x the polygon index, because we have two triangles per polygon
+			constexpr int32 TrianglesPerCell = 2;
+			const int32 LocalTriangleStartIndex = PolygonID * TrianglesPerCell;
+
+			AppendTriangle(CellCoordinates, Triangle, InHeightField, PolygonID, LocalTriangleStartIndex);
+			AppendTriangle(CellCoordinates, Triangle2, InHeightField, PolygonID, LocalTriangleStartIndex + 1);
 		}
-	}
+	}, PolygonsRowsProcessingLoopFlags);
 
 	bIsGenerated = true;
 }
 
-
 UE::Geometry::FMeshShapeGenerator& FChaosVDHeightFieldMeshGenerator::Generate()
 {
 	ensureAlwaysMsgf(bIsGenerated, TEXT("You need to call FChaosVDHeightFieldMeshGenerator::GenerateFromHeightField before calling Generate"));
diff --git a/Engine/Plugins/ChaosVD/Source/ChaosVD/Private/ChaosVDHeightfieldMeshGenerator.h b/Engine/Plugins/ChaosVD/Source/ChaosVD/Private/ChaosVDHeightfieldMeshGenerator.h
--- a/Engine/Plugins/ChaosVD/Source/ChaosVD/Private/ChaosVDHeightfieldMeshGenerator.h
+++ b/Engine/Plugins/ChaosVD/Source/ChaosVD/Private/ChaosVDHeightfieldMeshGenerator.h
@@ -24,7 +24,7 @@ public:
 	virtual FMeshShapeGenerator& Generate() override;
 
 private:
-	void AppendTriangle(int32& OutTriIndex, int& OutCurrentNormalIndex, int32 PolygonIndex, const Chaos::TVec2<Chaos::FReal>& InCellCoordinates, const UE::Geometry::FIndex3i& InTriangle, const Chaos::FHeightField& InHeightField);
+	void AppendTriangle(const Chaos::TVec2<Chaos::FReal>& InCellCoordinates, const UE::Geometry::FIndex3i& InTriangle, const Chaos::FHeightField& InHeightField, int32 PolygonID, int32 TriangleIndex);
 
 	bool bIsGenerated = false;
 };
diff --git a/Engine/Plugins/ChaosVD/Source/ChaosVD/Private/ChaosVDMeshComponentPool.cpp b/Engine/Plugins/ChaosVD/Source/ChaosVD/Private/ChaosVDMeshComponentPool.cpp
--- a/Engine/Plugins/ChaosVD/Source/ChaosVD/Private/ChaosVDMeshComponentPool.cpp
+++ b/Engine/Plugins/ChaosVD/Source/ChaosVD/Private/ChaosVDMeshComponentPool.cpp
@@ -3,87 +3,58 @@
 #include "ChaosVDMeshComponentPool.h"
 
 #include "ChaosVDGeometryDataComponent.h"
-#include "ChaosVDModule.h"
+#include "Components/ChaosVDStaticMeshComponent.h"
 #include "UObject/Package.h"
 
-bool FChaosVDMeshComponentPool::bUseComponentsPool = true;
-FAutoConsoleVariableRef FChaosVDMeshComponentPool::CVarUseComponentsPool(TEXT("p.Chaos.VD.Tool.UseComponentsPool"), FChaosVDMeshComponentPool::bUseComponentsPool, TEXT("Set to false to disable the use of a pool system for Mesh Components."));
-
 void FChaosVDMeshComponentPool::AddReferencedObjects(FReferenceCollector& Collector)
 {
-	Collector.AddReferencedObjects(PooledDynamicMeshComponent);
-	Collector.AddReferencedObjects(PooledInstancedStaticMeshComponent);
-	Collector.AddReferencedObjects(PooledStaticMeshComponent);
+	Collector.AddReferencedObject(CachedISMCOpaqueMaterialBase);
+	Collector.AddReferencedObject(CachedISMCTranslucentBase);
+	Collector.AddReferencedObject(CachedStaticMeshComponentTranslucentMaterialBase);
+	Collector.AddReferencedObject(CachedStaticMeshComponentOpaqueMaterialBase);
 }
 
-void FChaosVDMeshComponentPool::DisposeMeshComponent(UMeshComponent* MeshComponent)
+FChaosVDMeshComponentPool::FChaosVDMeshComponentPool()
 {
-	if (!MeshComponent)
-	{
-		return;
-	}
-
-	if (!bUseComponentsPool)
-	{
-		MeshComponent->DestroyComponent();
-		return;
-	}
-
-	ResetMeshComponent(MeshComponent);
+	CachedISMCOpaqueMaterialBase = FChaosVDGeometryComponentUtils::GetBaseMaterialForType(EChaosVDMaterialType::ISMCOpaque);
+	CachedISMCTranslucentBase = FChaosVDGeometryComponentUtils::GetBaseMaterialForType(EChaosVDMaterialType::ISMCTranslucent);
+	CachedStaticMeshComponentTranslucentMaterialBase = FChaosVDGeometryComponentUtils::GetBaseMaterialForType(EChaosVDMaterialType::SMTranslucent);
+	CachedStaticMeshComponentOpaqueMaterialBase = FChaosVDGeometryComponentUtils::GetBaseMaterialForType(EChaosVDMaterialType::SMOpaque);
 
-	if (UInstancedStaticMeshComponent* AsISMC = Cast<UInstancedStaticMeshComponent>(MeshComponent))
-	{
-		PooledInstancedStaticMeshComponent.Add(AsISMC);
-	}
-	else if (UStaticMeshComponent* AsSMC = Cast<UStaticMeshComponent>(MeshComponent))
-	{
-		PooledStaticMeshComponent.Add(AsSMC);
-	}
-	else if (UDynamicMeshComponent* AsDMC = Cast<UDynamicMeshComponent>(MeshComponent))
-	{
-		PooledDynamicMeshComponent.Add(AsDMC);
-	}
-	else
-	{
-		UE_LOG(LogChaosVDEditor, Warning, TEXT("[%s] Unable top dispose Mesh component [%s] | Incompatible mesh type used [%s]"), ANSI_TO_TCHAR(__FUNCTION__), *GetNameSafe(MeshComponent), *GetNameSafe(MeshComponent->GetClass()));
-	}
-
-	const FString NewName = MeshComponent->GetName() + TEXT("_POOLED_") + FGuid::NewGuid().ToString();
-	MeshComponent->Rename(*NewName, GetTransientPackage(), REN_NonTransactional | REN_DoNotDirty | REN_ForceNoResetLoaders | REN_SkipGeneratedClasses | REN_DontCreateRedirectors);
+	StaticMeshComponentPool.SetPoolName(TEXT("Static Mesh Components Pool"));
+	InstancedStaticMeshComponentPool.SetPoolName(TEXT("Instanced Static Mesh Components Pool"));
 }
 
-void FChaosVDMeshComponentPool::ResetMeshComponent(UMeshComponent* MeshComponent)
+void FChaosVDMeshComponentPool::DisposeMeshComponent(UMeshComponent* MeshComponent)
 {
-	if (!MeshComponent)
-	{
-		return;
-	}
-
-	if (UDynamicMeshComponent* DynamicMeshComponent = Cast<UDynamicMeshComponent>(MeshComponent))
-	{
-		DynamicMeshComponent->SetDynamicMesh(nullptr);
-	}
-	else if (UInstancedStaticMeshComponent* InstancedStaticMeshComponent = Cast<UInstancedStaticMeshComponent>(MeshComponent))
+	if (UChaosVDStaticMeshComponent* AsStaticMeshComponent = Cast<UChaosVDStaticMeshComponent>(MeshComponent))
 	{
-		InstancedStaticMeshComponent->SetStaticMesh(nullptr);
-		InstancedStaticMeshComponent->ClearInstances();
+		StaticMeshComponentPool.DisposeObject(AsStaticMeshComponent);
 	}
-	else if (UStaticMeshComponent* StaticMeshComponent = Cast<UStaticMeshComponent>(MeshComponent))
+	else if (UChaosVDInstancedStaticMeshComponent* AsInstancedStaticMeshComponent = Cast<UChaosVDInstancedStaticMeshComponent>(MeshComponent))
 	{
-		StaticMeshComponent->SetStaticMesh(nullptr);
+		InstancedStaticMeshComponentPool.DisposeObject(AsInstancedStaticMeshComponent);
 	}
-	
-	MeshComponent->SetRelativeTransform(FTransform::Identity);
-	
-	if (IChaosVDGeometryComponent* DataComponent = Cast<IChaosVDGeometryComponent>(MeshComponent))
-	{
-		DataComponent->Reset();
-	}
-
-	MeshComponent->UnregisterComponent();
+}
 
-	if (AActor* Owner = MeshComponent->GetOwner())
-	{
-		Owner->RemoveOwnedComponent(MeshComponent);
-	}
+UMaterialInterface* FChaosVDMeshComponentPool::GetMaterialForType(EChaosVDMaterialType Type)
+{
+	switch(Type)
+	{
+	case EChaosVDMaterialType::SMTranslucent:
+		return CachedStaticMeshComponentTranslucentMaterialBase;
+	case EChaosVDMaterialType::SMOpaque:
+		return CachedStaticMeshComponentOpaqueMaterialBase;
+	case EChaosVDMaterialType::ISMCOpaque:
+		return CachedISMCOpaqueMaterialBase;
+	case EChaosVDMaterialType::ISMCTranslucent:
+		return CachedISMCTranslucentBase;
+	default:
+		{
+			ensure(false);
+			return nullptr;
+		}
+		
+	}	
 }
+
diff --git a/Engine/Plugins/ChaosVD/Source/ChaosVD/Private/ChaosVDMeshComponentPool.h b/Engine/Plugins/ChaosVD/Source/ChaosVD/Private/ChaosVDMeshComponentPool.h
--- a/Engine/Plugins/ChaosVD/Source/ChaosVD/Private/ChaosVDMeshComponentPool.h
+++ b/Engine/Plugins/ChaosVD/Source/ChaosVD/Private/ChaosVDMeshComponentPool.h
@@ -2,41 +2,40 @@
 
 #pragma once
 
-#include "ChaosVDModule.h"
+#include "ChaosVDObjectPool.h"
+#include "Components/ChaosVDInstancedStaticMeshComponent.h"
+#include "Components/ChaosVDStaticMeshComponent.h"
 #include "Containers/Array.h"
 
-#include "Components/DynamicMeshComponent.h"
-#include "Components/InstancedStaticMeshComponent.h"
-#include "Components/StaticMeshComponent.h"
-
 class FChaosVDMeshComponentPool : public FGCObject
 {
 public:
-	virtual void AddReferencedObjects(FReferenceCollector& Collector) override;
 	
+	FChaosVDMeshComponentPool();
+
 	virtual FString GetReferencerName() const override
 	{
 		return TEXT("FChaosVDMeshComponentPool");
 	}
+	
+	virtual void AddReferencedObjects(FReferenceCollector& Collector) override;
 
 	template<typename TMeshComponent>
 	TMeshComponent* AcquireMeshComponent(UObject* Outer, FName Name);
 
 	void DisposeMeshComponent(UMeshComponent* MeshComponent);
 
-private:
-
-	template <typename DesiredMeshComponent, typename PooledComponent>
-	DesiredMeshComponent* GetMeshComponentFromPool_Internal(TArray<TObjectPtr<PooledComponent>>& InMeshComponentPool, UObject* Outer, FName Name);
+	UMaterialInterface* GetMaterialForType(EChaosVDMaterialType Type);
 
-	void ResetMeshComponent(UMeshComponent* MeshComponent);
+private:
 
-	TArray<TObjectPtr<UStaticMeshComponent>> PooledStaticMeshComponent;
-	TArray<TObjectPtr<UInstancedStaticMeshComponent>> PooledInstancedStaticMeshComponent;
-	TArray<TObjectPtr<UDynamicMeshComponent>> PooledDynamicMeshComponent;
+	TChaosVDObjectPool<UChaosVDStaticMeshComponent> StaticMeshComponentPool;
+	TChaosVDObjectPool<UChaosVDInstancedStaticMeshComponent> InstancedStaticMeshComponentPool;
 
-	static bool bUseComponentsPool;
-	static FAutoConsoleVariableRef CVarUseComponentsPool;
+	TObjectPtr<UMaterialInterface> CachedISMCOpaqueMaterialBase;
+	TObjectPtr<UMaterialInterface> CachedISMCTranslucentBase;
+	TObjectPtr<UMaterialInterface> CachedStaticMeshComponentTranslucentMaterialBase;
+	TObjectPtr<UMaterialInterface> CachedStaticMeshComponentOpaqueMaterialBase;
 };
 
 template <typename TMeshComponent>
@@ -44,45 +43,14 @@ TMeshComponent* FChaosVDMeshComponentPool::AcquireMeshComponent(UObject* Outer,
 {
 	if constexpr (std::is_base_of_v<UStaticMeshComponent, TMeshComponent> && !std::is_base_of_v<UInstancedStaticMeshComponent, TMeshComponent>)
 	{
-		return GetMeshComponentFromPool_Internal<TMeshComponent>(PooledStaticMeshComponent, Outer, Name);
+		return StaticMeshComponentPool.AcquireObject(Outer, Name);
 	}
 	else if constexpr (std::is_base_of_v<UInstancedStaticMeshComponent, TMeshComponent>)
 	{
-		return GetMeshComponentFromPool_Internal<TMeshComponent>(PooledInstancedStaticMeshComponent, Outer, Name);
-	}
-	else if constexpr (std::is_base_of_v<UDynamicMeshComponent, TMeshComponent>)
-	{
-		return GetMeshComponentFromPool_Internal<TMeshComponent>(PooledDynamicMeshComponent, Outer, Name);
+		return InstancedStaticMeshComponentPool.AcquireObject(Outer, Name);
 	}
 	else
 	{
 		return nullptr;
 	}
 }
-
-template <typename DesiredMeshComponent, typename PooledComponent>
-DesiredMeshComponent* FChaosVDMeshComponentPool::GetMeshComponentFromPool_Internal(TArray<TObjectPtr<PooledComponent>>& InMeshComponentPool, UObject* Outer,  FName Name)
-{
-	static_assert(std::is_base_of_v<UStaticMeshComponent, DesiredMeshComponent> || std::is_base_of_v<UInstancedStaticMeshComponent, DesiredMeshComponent> || std::is_base_of_v<DesiredMeshComponent, UDynamicMeshComponent>, "GetMeshComponentInternal Only supports DynamicMeshComponent, Static MeshComponent and Instanced Static Mesh Component");
-
-	// We need to ensure unique names
-	const FString NewName = Name.ToString() + FGuid::NewGuid().ToString();
-	
-	if (bUseComponentsPool &&  InMeshComponentPool.Num() > 0)
-	{
-		if (DesiredMeshComponent* Component = Cast<DesiredMeshComponent>(InMeshComponentPool.Pop()))
-		{
-			Component->Rename(*NewName , Outer, REN_NonTransactional | REN_DoNotDirty | REN_ForceNoResetLoaders | REN_SkipGeneratedClasses | REN_DontCreateRedirectors);
-			return Component;
-		}
-		else
-		{
-			UE_LOG(LogChaosVDEditor, Warning, TEXT("[%s] Selected Pool has [%d] objects, but none where castable to the desired type"), ANSI_TO_TCHAR(__FUNCTION__), InMeshComponentPool.Num());
-		}
-	}
-
-	DesiredMeshComponent* NewMeshComponent = NewObject<DesiredMeshComponent>(Outer, *NewName);
-	NewMeshComponent->EmptyOverrideMaterials();
-
-	return NewMeshComponent;
-}
diff --git a/Engine/Plugins/ChaosVD/Source/ChaosVD/Private/ChaosVDObjectPool.cpp b/Engine/Plugins/ChaosVD/Source/ChaosVD/Private/ChaosVDObjectPool.cpp
new file mode 100644
--- /dev/null
+++ b/Engine/Plugins/ChaosVD/Source/ChaosVD/Private/ChaosVDObjectPool.cpp
@@ -0,0 +1,6 @@
+// Copyright Epic Games, Inc. All Rights Reserved.
+
+#include "ChaosVDObjectPool.h"
+
+bool FChaosVDObjectPoolCVars::bUseObjectPool = true;
+FAutoConsoleVariableRef FChaosVDObjectPoolCVars::CVarUseObjectPool(TEXT("p.Chaos.VD.Tool.UseObjectPool"),FChaosVDObjectPoolCVars::bUseObjectPool, TEXT("Set to false to disable the use of a pool system for uobjects."));
diff --git a/Engine/Plugins/ChaosVD/Source/ChaosVD/Private/ChaosVDObjectPool.h b/Engine/Plugins/ChaosVD/Source/ChaosVD/Private/ChaosVDObjectPool.h
new file mode 100644
--- /dev/null
+++ b/Engine/Plugins/ChaosVD/Source/ChaosVD/Private/ChaosVDObjectPool.h
@@ -0,0 +1,137 @@
+// Copyright Epic Games, Inc. All Rights Reserved.
+
+#pragma once
+
+#include "ChaosVDModule.h"
+#include "Containers/Array.h"
+#include "Components/ActorComponent.h"
+#include "GameFramework/Actor.h"
+#include "Interfaces/ChaosVDPooledObject.h"
+
+struct FChaosVDObjectPoolCVars
+{
+	static bool bUseObjectPool;
+	static FAutoConsoleVariableRef CVarUseObjectPool;
+};
+
+/** Basic Pool system for UObjects */
+template<typename ObjectType>
+class TChaosVDObjectPool : public FGCObject
+{
+public:
+	TChaosVDObjectPool() = default;
+	virtual ~TChaosVDObjectPool() override;
+
+	virtual void AddReferencedObjects(FReferenceCollector& Collector) override;
+
+	void SetPoolName(const FString& InName) { PoolName = InName; };
+	
+	virtual FString GetReferencerName() const override
+	{
+		return PoolName;
+	}
+
+	ObjectType* AcquireObject(UObject* Outer, FName Name);
+
+	void DisposeObject(UObject* Object);
+
+	TFunction<ObjectType*(UObject*,FName)> ObjectFactoryOverride;
+
+protected:
+	FString PoolName = TEXT("ChaosVDObjectPool");
+	int32 PoolHits = 0;
+	int32 PoolRequests = 0;
+	TArray<TObjectPtr<ObjectType>> PooledObjects;
+};
+
+template <typename ObjectType>
+TChaosVDObjectPool<ObjectType>::~TChaosVDObjectPool()
+{
+	float HitMissRatio = PoolRequests > 0 ? ((static_cast<float>(PoolHits) / static_cast<float>(PoolRequests)) * 100.0f) : 0;
+	UE_LOG(LogChaosVDEditor, Log, TEXT("Object pooling Statics for pool [%s] | Hits [%d] | Total Acquire requests [%d] | [%f] percent hit/miss ratio"), *PoolName, PoolHits,  PoolRequests, HitMissRatio);
+}
+
+template <typename ObjectType>
+void TChaosVDObjectPool<ObjectType>::AddReferencedObjects(FReferenceCollector& Collector)
+{
+	Collector.AddReferencedObjects(PooledObjects);
+}
+
+template <typename ObjectType>
+ObjectType* TChaosVDObjectPool<ObjectType>::AcquireObject(UObject* Outer, FName Name)
+{
+	PoolRequests++;
+
+	// If pooling is disabled, just fall through the code that will create the object
+	if (FChaosVDObjectPoolCVars::bUseObjectPool && PooledObjects.Num() > 0)
+	{
+		if (ObjectType* Object = PooledObjects.Pop())
+		{
+			const FName NewName = MakeUniqueObjectName(Outer, ObjectType::StaticClass(), Name);
+			Object->Rename(*NewName.ToString() , Outer, REN_NonTransactional | REN_DoNotDirty | REN_ForceNoResetLoaders | REN_SkipGeneratedClasses | REN_DontCreateRedirectors);
+		
+			if (IChaosVDPooledObject* AsPooledObject = Cast<IChaosVDPooledObject>(Object))
+			{
+				AsPooledObject->OnAcquired();
+			}
+
+			PoolHits++;
+			
+			return Object;
+		}
+	}
+
+	const FName NewName = MakeUniqueObjectName(Outer, ObjectType::StaticClass(), Name);
+
+	ObjectType* CreatedObject = nullptr;
+	if (ObjectFactoryOverride)
+	{
+		CreatedObject = ObjectFactoryOverride(Outer, NewName);
+	}
+	else
+	{
+		CreatedObject = NewObject<ObjectType>(Outer, NewName);
+	}
+
+	if (IChaosVDPooledObject* AsPooledObject = Cast<IChaosVDPooledObject>(CreatedObject))
+	{
+		AsPooledObject->OnAcquired();
+	}
+	
+	return CreatedObject;
+}
+
+template <typename ObjectType>
+void TChaosVDObjectPool<ObjectType>::DisposeObject(UObject* Object)
+{
+	// If pooling is disabled, just destroy the object
+	//TODO: Should we provide a way to override how these are destroyed?
+	if (!FChaosVDObjectPoolCVars::bUseObjectPool)
+	{
+		if (UActorComponent* AsActorComponent = Cast<UActorComponent>(Object))
+		{
+			AsActorComponent->DestroyComponent();
+		}
+		else if (AActor* AsActor = Cast<AActor>(Object))
+		{
+			AsActor->Destroy();
+		}
+		else
+		{
+			Object->ConditionalBeginDestroy();
+		}
+		
+		return;
+	}
+
+	if (IChaosVDPooledObject* AsPooledObject = Cast<IChaosVDPooledObject>(Object))
+	{
+		AsPooledObject->OnDisposed();
+	}
+
+	UPackage* TransientPackage = GetTransientPackage();
+	const FName NewName = MakeUniqueObjectName(TransientPackage, ObjectType::StaticClass());
+	Object->Rename(*NewName.ToString(), TransientPackage, REN_NonTransactional | REN_DoNotDirty | REN_ForceNoResetLoaders | REN_SkipGeneratedClasses | REN_DontCreateRedirectors);
+
+	PooledObjects.Emplace(Cast<ObjectType>(Object));
+}
diff --git a/Engine/Plugins/ChaosVD/Source/ChaosVD/Private/ChaosVDParticleActor.cpp b/Engine/Plugins/ChaosVD/Source/ChaosVD/Private/ChaosVDParticleActor.cpp
--- a/Engine/Plugins/ChaosVD/Source/ChaosVD/Private/ChaosVDParticleActor.cpp
+++ b/Engine/Plugins/ChaosVD/Source/ChaosVD/Private/ChaosVDParticleActor.cpp
@@ -223,7 +223,7 @@ void AChaosVDParticleActor::UpdateGeometry(const Chaos::FConstImplicitObjectPtr&
 		return;
 	}
 
-	const TSharedPtr<FChaosVDGeometryBuilder> GeometryGenerator = ScenePtr->GetGeometryGenerator();
+	const TSharedPtr<FChaosVDGeometryBuilder> GeometryGenerator = ScenePtr->GetGeometryGenerator().Pin();
 	if (!GeometryGenerator.IsValid())
 	{
 		return;
@@ -254,7 +254,7 @@ void AChaosVDParticleActor::UpdateGeometry(const Chaos::FConstImplicitObjectPtr&
 	constexpr int32 LODsToGenerateNum = 3;
 	constexpr int32 LODsToGenerateNumForInstancedStaticMesh = 0;
 
-	GeometryGenerator->CreateMeshesFromImplicitObject<UStaticMesh>(InImplicitObject, this, OutExtractedGeometryDataHandles, bHasToUseStaticMeshComponent ? LODsToGenerateNum : LODsToGenerateNumForInstancedStaticMesh);
+	GeometryGenerator->CreateMeshesFromImplicitObject(InImplicitObject, this, OutExtractedGeometryDataHandles, bHasToUseStaticMeshComponent ? LODsToGenerateNum : LODsToGenerateNumForInstancedStaticMesh);
 
 	// This should not happen in theory, but there might be some valid situations where it does. Adding an ensure to catch them and then evaluate if it is really an issue (if it is not I will remove the ensure later on). 
 	if (!ensure(ObjectsToGenerateNum == OutExtractedGeometryDataHandles.Num()))
diff --git a/Engine/Plugins/ChaosVD/Source/ChaosVD/Private/ChaosVDPlaybackController.cpp b/Engine/Plugins/ChaosVD/Source/ChaosVD/Private/ChaosVDPlaybackController.cpp
--- a/Engine/Plugins/ChaosVD/Source/ChaosVD/Private/ChaosVDPlaybackController.cpp
+++ b/Engine/Plugins/ChaosVD/Source/ChaosVD/Private/ChaosVDPlaybackController.cpp
@@ -100,6 +100,14 @@ void FChaosVDPlaybackController::UnloadCurrentRecording(EChaosVDUnloadRecordingF
 	RecordingLastSeenTimeUpdatedAsCycle = 0;
 
 	TrackInfoUpdateGTQueue.Empty();
+
+	if (const TSharedPtr<FChaosVDScene> SceneToControlSharedPtr = SceneToControl.Pin())
+	{
+		if (SceneToControlSharedPtr->IsInitialized())
+		{
+			SceneToControlSharedPtr->CleanUpScene(EChaosVDSceneCleanUpOptions::ReInitializeGeometryBuilder | EChaosVDSceneCleanUpOptions::CollectGarbage);
+		}	
+	}
 	
 	if (LoadedRecording.IsValid())
 	{
@@ -111,14 +119,6 @@ void FChaosVDPlaybackController::UnloadCurrentRecording(EChaosVDUnloadRecordingF
 
 	HandleCurrentRecordingUpdated();
 
-	if (const TSharedPtr<FChaosVDScene> SceneToControlSharedPtr = SceneToControl.Pin())
-	{
-		if (SceneToControlSharedPtr->IsInitialized())
-		{
-			SceneToControlSharedPtr->CleanUpScene();
-		}	
-	}
-
 	if (EnumHasAnyFlags(UnloadOptions, EChaosVDUnloadRecordingFlags::BroadcastChanges))
 	{
 		bHasPendingGTUpdateBroadcast = true;
diff --git a/Engine/Plugins/ChaosVD/Source/ChaosVD/Private/ChaosVDScene.cpp b/Engine/Plugins/ChaosVD/Source/ChaosVD/Private/ChaosVDScene.cpp
--- a/Engine/Plugins/ChaosVD/Source/ChaosVD/Private/ChaosVDScene.cpp
+++ b/Engine/Plugins/ChaosVD/Source/ChaosVD/Private/ChaosVDScene.cpp
@@ -33,6 +33,7 @@
 #include "Settings/ChaosVDCoreSettings.h"
 #include "UObject/Package.h"
 #include "WorldPersistentFolders.h"
+#include "Actors/ChaosVDGeometryContainer.h"
 
 #define LOCTEXT_NAMESPACE "ChaosVisualDebugger"
 
@@ -47,6 +48,15 @@ namespace ChaosVDSceneUIOptions
 	constexpr bool bAllowInPIE = false;
 }
 
+namespace Chaos::VisualDebugger::Cvars
+{
+	static bool bReInitializeGeometryBuilderOnCleanup = true;
+	static FAutoConsoleVariableRef CVarChaosVDReInitializeGeometryBuilderOnCleanup(
+		TEXT("p.Chaos.VD.Tool.ReInitializeGeometryBuilderOnCleanup"),
+		bReInitializeGeometryBuilderOnCleanup,
+		TEXT("If true, any static mesh component and static mesh component created will be destroyed when a new CVD recording is loaded"));
+}
+
 void FChaosVDScene::Initialize()
 {
 	if (!ensure(!bIsInitialized))
@@ -57,10 +67,6 @@ void FChaosVDScene::Initialize()
 	InitializeSelectionSets();
 	
 	PhysicsVDWorld = CreatePhysicsVDWorld();
-
-	GeometryGenerator = MakeShared<FChaosVDGeometryBuilder>();
-
-	GeometryGenerator->Initialize(AsWeak());
 	
 	StreamableManager = MakeShared<FStreamableManager>();
 
@@ -74,9 +80,25 @@ void FChaosVDScene::Initialize()
 		StreamableManager->RequestSyncLoad(Settings->InstancedMeshesQueryOnlyMaterial.ToSoftObjectPath());
 	}
 
+	GeometryGenerator = MakeShared<FChaosVDGeometryBuilder>();
+
+	GeometryGenerator->Initialize(AsWeak());
+
 	bIsInitialized = true;
 }
 
+
+void FChaosVDScene::PerformGarbageCollection()
+{
+	FScopedSlowTask CollectingGarbageSlowTask(1, LOCTEXT("CollectingGarbageDataMessage", "Collecting Garbage ..."));
+	CollectingGarbageSlowTask.MakeDialog();
+
+	CollectGarbage(GARBAGE_COLLECTION_KEEPFLAGS);
+
+	CollectingGarbageSlowTask.EnterProgressFrame();
+}
+
+
 void FChaosVDScene::DeInitialize()
 {
 	constexpr float AmountOfWork = 1.0f;
@@ -105,14 +127,7 @@ void FChaosVDScene::DeInitialize()
 		PhysicsVDWorld = nullptr;
 	}
 
-	{
-		FScopedSlowTask CollectingGarbageSlowTask(1, LOCTEXT("CollectingGarbageDataMessage", "Collecting Garbage ..."));
-		CollectingGarbageSlowTask.MakeDialog();
-
-		CollectGarbage(GARBAGE_COLLECTION_KEEPFLAGS);
-
-		CollectingGarbageSlowTask.EnterProgressFrame();
-	}
+	PerformGarbageCollection();
 
 	bIsInitialized = false;
 }
@@ -360,26 +375,49 @@ void FChaosVDScene::HandleEnterNewSolverFrame(int32 FrameNumber, const FChaosVDS
 	}
 }
 
-void FChaosVDScene::CleanUpScene()
+void FChaosVDScene::CleanUpScene(EChaosVDSceneCleanUpOptions Options)
 {
-	constexpr float AmountOfWork = 1.0f;
-	const float PercentagePerElement = 1.0f / SolverDataContainerBySolverID.Num();
+	{
+		constexpr float AmountOfWork = 1.0f;
+		const float PercentagePerElement = 1.0f / SolverDataContainerBySolverID.Num();
 
-	FScopedSlowTask CleaningSceneSlowTask(AmountOfWork, LOCTEXT("CleaningupSceneSolverMessage", "Clearing Solver Data ..."));
-	CleaningSceneSlowTask.MakeDialog();
+		FScopedSlowTask CleaningSceneSlowTask(AmountOfWork, LOCTEXT("CleaningupSceneSolverMessage", "Clearing Solver Data ..."));
+		CleaningSceneSlowTask.MakeDialog();
 
-	ClearSelectionAndNotify();
+		ClearSelectionAndNotify();
 
-	if (PhysicsVDWorld)
+		if (PhysicsVDWorld)
+		{
+			for (const TPair<int32, AChaosVDSolverInfoActor*>& SolverDataInfoWithID : SolverDataContainerBySolverID)
+			{
+				PhysicsVDWorld->DestroyActor(SolverDataInfoWithID.Value);
+				CleaningSceneSlowTask.EnterProgressFrame(PercentagePerElement);
+			}
+		}
+
+		SolverDataContainerBySolverID.Reset();
+	}
+
+	if (Chaos::VisualDebugger::Cvars:: bReInitializeGeometryBuilderOnCleanup && EnumHasAnyFlags(Options, EChaosVDSceneCleanUpOptions::ReInitializeGeometryBuilder))
 	{
-		for (const TPair<int32, AChaosVDSolverInfoActor*>& SolverDataInfoWithID : SolverDataContainerBySolverID)
+		if (AChaosVDGeometryContainer* AsGeometryContainer = Cast<AChaosVDGeometryContainer>(MeshComponentContainerActor))
 		{
-			PhysicsVDWorld->DestroyActor(SolverDataInfoWithID.Value);
-			CleaningSceneSlowTask.EnterProgressFrame(PercentagePerElement);
+			AsGeometryContainer->CleanUp();
 		}
+
+		GeometryGenerator->DeInitialize();
+		GeometryGenerator.Reset();
+	
+		GeometryGenerator = MakeShared<FChaosVDGeometryBuilder>();
+		GeometryGenerator->Initialize(AsWeak());
 	}
+	
+	SceneQueriesContainer->CleanUp();
 
-	SolverDataContainerBySolverID.Reset();
+	if (EnumHasAnyFlags(Options, EChaosVDSceneCleanUpOptions::CollectGarbage))
+	{
+		PerformGarbageCollection();
+	}
 }
 
 Chaos::FConstImplicitObjectPtr FChaosVDScene::GetUpdatedGeometry(int32 GeometryID) const
@@ -526,7 +564,7 @@ AActor* FChaosVDScene::CreateMeshComponentsContainer(UWorld* TargetWorld)
 {
 	const FName GeometryFolderPath("ChaosVisualDebugger/GeneratedMeshComponents");
 
-	MeshComponentContainerActor = TargetWorld->SpawnActor<AActor>();
+	MeshComponentContainerActor = TargetWorld->SpawnActor<AChaosVDGeometryContainer>();
 	MeshComponentContainerActor->SetFolderPath(GeometryFolderPath);
 
 	return MeshComponentContainerActor;
diff --git a/Engine/Plugins/ChaosVD/Source/ChaosVD/Private/ChaosVDScene.h b/Engine/Plugins/ChaosVD/Source/ChaosVD/Private/ChaosVDScene.h
--- a/Engine/Plugins/ChaosVD/Source/ChaosVD/Private/ChaosVDScene.h
+++ b/Engine/Plugins/ChaosVD/Source/ChaosVD/Private/ChaosVDScene.h
@@ -11,6 +11,7 @@
 #include "UObject/ObjectMacros.h"
 #include "UObject/ObjectPtr.h"
 
+class AChaosVDGeometryContainer;
 class IChaosVDGeometryOwnerInterface;
 class FChaosVDSelectionCustomization;
 class ITypedElementSelectionInterface;
@@ -40,6 +41,16 @@ DECLARE_MULTICAST_DELEGATE_OneParam(FChaosVDSolverInfoActorCreatedDelegate, ACha
 
 DECLARE_MULTICAST_DELEGATE_TwoParams(FChaosVDSolverVisibilityChangedDelegate, int32 SolverID, bool bNewVisibility)
 
+UENUM()
+enum class EChaosVDSceneCleanUpOptions
+{
+	None = 0,
+	ReInitializeGeometryBuilder = 1 << 0,
+	CollectGarbage = 1 << 1
+};
+
+ENUM_CLASS_FLAGS(EChaosVDSceneCleanUpOptions)
+
 /** Recreates a UWorld from a recorded Chaos VD Frame */
 class FChaosVDScene : public FGCObject , public TSharedFromThis<FChaosVDScene>
 {
@@ -72,14 +83,14 @@ public:
 	void HandleEnterNewSolverFrame(int32 FrameNumber, const FChaosVDSolverFrameData& InFrameData);
 
 	/** Deletes all actors of the Scene and underlying UWorld */
-	void CleanUpScene();
+	void CleanUpScene(EChaosVDSceneCleanUpOptions Options = EChaosVDSceneCleanUpOptions::None);
 
 	/** Returns the a ptr to the UWorld used to represent the current recorded frame data */
 	UWorld* GetUnderlyingWorld() const { return PhysicsVDWorld; };
 
 	bool IsInitialized() const { return  bIsInitialized; }
 
-	const TSharedPtr<FChaosVDGeometryBuilder>& GetGeometryGenerator() { return  GeometryGenerator; }
+	TWeakPtr<FChaosVDGeometryBuilder> GetGeometryGenerator() { return  GeometryGenerator; }
 
 	// No need to deprecate the old version since it is not a public API nor inline 
 	Chaos::FConstImplicitObjectPtr GetUpdatedGeometry(int32 GeometryID) const;
@@ -134,6 +145,8 @@ public:
 
 private:
 
+	void PerformGarbageCollection();
+
 	/** Creates an ChaosVDParticle actor for the Provided recorded Particle Data */
 	AChaosVDParticleActor* SpawnParticleFromRecordedData(const TSharedPtr<FChaosVDParticleDataWrapper>& InParticleData, const FChaosVDSolverFrameData& InFrameData);
 
diff --git a/Engine/Plugins/ChaosVD/Source/ChaosVD/Private/ChaosVDTriMeshGenerator.cpp b/Engine/Plugins/ChaosVD/Source/ChaosVD/Private/ChaosVDTriMeshGenerator.cpp
--- a/Engine/Plugins/ChaosVD/Source/ChaosVD/Private/ChaosVDTriMeshGenerator.cpp
+++ b/Engine/Plugins/ChaosVD/Source/ChaosVD/Private/ChaosVDTriMeshGenerator.cpp
@@ -10,13 +10,14 @@ void FChaosVDTriMeshGenerator::GenerateFromTriMesh(const Chaos::FTriangleMeshImp
 
 	const int32 NormalsNum = NumTriangles * 3;
 	constexpr int32 UVsNum = 0;
-	SetBufferSizes(NumVertices, NumTriangles , UVsNum, NormalsNum );
+	SetBufferSizes(NumVertices, NumTriangles , UVsNum, NormalsNum);
 
+	EParallelForFlags Flags = NumVertices > MaxElementsNumToProcessInSingleThread ? EParallelForFlags::ForceSingleThread : EParallelForFlags::None;
 	// Fill the vertex buffer with the transformed vertices of the TriMesh Shape
-	for (int32 i = 0; i < NumVertices; i++)
+	ParallelFor(NumVertices, [this, &InTriMesh](int32 VertexIndex)
 	{
-		Vertices[i] = FVector3d(UE::Math::TVector<double>(InTriMesh.Particles().GetX(i)));
-	}
+		Vertices[VertexIndex] = FVector3d(UE::Math::TVector<double>(InTriMesh.Particles().GetX(VertexIndex)));
+	}, Flags);
 
 	const Chaos::FTrimeshIndexBuffer& IdxBuffer = InTriMesh.Elements();
 	if (IdxBuffer.RequiresLargeIndices())
diff --git a/Engine/Plugins/ChaosVD/Source/ChaosVD/Private/ChaosVDTriMeshGenerator.h b/Engine/Plugins/ChaosVD/Source/ChaosVD/Private/ChaosVDTriMeshGenerator.h
--- a/Engine/Plugins/ChaosVD/Source/ChaosVD/Private/ChaosVDTriMeshGenerator.h
+++ b/Engine/Plugins/ChaosVD/Source/ChaosVD/Private/ChaosVDTriMeshGenerator.h
@@ -24,29 +24,33 @@ private:
 	template<typename BufferIndexType>
 	void ProcessTriangles(const TArray<BufferIndexType>& InTriangles, const int32 NumTriangles, const Chaos::FTriangleMeshImplicitObject& InTriMesh);
 
-	bool bIsGenerated = false;;
+	bool bIsGenerated = false;
+
+	/** Max number of elements after which we will process elements in parallel using worker threads
+	 * Note: This value is not tuned yet
+	 */
+	static constexpr int32 MaxElementsNumToProcessInSingleThread = 64;
 };
 
 template <typename BufferIndexType>
 void FChaosVDTriMeshGenerator::ProcessTriangles(const TArray<BufferIndexType>& InTriangles, const int32 NumTriangles, const Chaos::FTriangleMeshImplicitObject& InTriMesh)
 {
-	int CurrentNormalIndex = 0;
-	for (int32 TriangleIndex = 0; TriangleIndex < NumTriangles; ++TriangleIndex)
+	EParallelForFlags Flags = NumTriangles > MaxElementsNumToProcessInSingleThread ? EParallelForFlags::ForceSingleThread : EParallelForFlags::None;
+	ParallelFor(NumTriangles, [this, &InTriangles, &InTriMesh](int32 TriangleIndex)
 	{
 		UE::Geometry::FIndex3i Triangle(InTriangles[TriangleIndex][0],InTriangles[TriangleIndex][1],InTriangles[TriangleIndex][2]);
 
+		FVector3f FaceNormal(InTriMesh.GetFaceNormal(TriangleIndex));
+		int32 StartNormalIndexNumber = TriangleIndex * 3;
 		// Create a Normal entry per triangle vertex.
-		int NormalIndices[3];
 		for (int32 LocalVertexIndex = 0; LocalVertexIndex < 3; ++LocalVertexIndex)
 		{
 			// Use the normal of the face for all it's vertices
-			Normals[CurrentNormalIndex] = UE::Math::TVector<float>(InTriMesh.GetFaceNormal(TriangleIndex));
-			NormalIndices[LocalVertexIndex] = CurrentNormalIndex;
-			CurrentNormalIndex++;
+			Normals[StartNormalIndexNumber + LocalVertexIndex] = FaceNormal;
 		}
 
-		SetTriangle(TriangleIndex, MoveTemp(Triangle));
 		SetTrianglePolygon(TriangleIndex, TriangleIndex);
-		SetTriangleNormals(TriangleIndex, NormalIndices[0], NormalIndices[1], NormalIndices[2]);
-	}
+		SetTriangleNormals(TriangleIndex, StartNormalIndexNumber, StartNormalIndexNumber + 1, StartNormalIndexNumber + 2);
+		SetTriangle(TriangleIndex, MoveTemp(Triangle));
+	}, Flags);
 }
diff --git a/Engine/Plugins/ChaosVD/Source/ChaosVD/Private/Components/ChaosVDInstancedStaticMeshComponent.cpp b/Engine/Plugins/ChaosVD/Source/ChaosVD/Private/Components/ChaosVDInstancedStaticMeshComponent.cpp
--- a/Engine/Plugins/ChaosVD/Source/ChaosVD/Private/Components/ChaosVDInstancedStaticMeshComponent.cpp
+++ b/Engine/Plugins/ChaosVD/Source/ChaosVD/Private/Components/ChaosVDInstancedStaticMeshComponent.cpp
@@ -2,6 +2,7 @@
 
 #include "Components/ChaosVDInstancedStaticMeshComponent.h"
 
+#include "ChaosVDGeometryBuilder.h"
 #include "ChaosVDModule.h"
 #include "Materials/MaterialInterface.h"
 #include "Materials/MaterialInstanceDynamic.h"
@@ -134,12 +135,63 @@ void UChaosVDInstancedStaticMeshComponent::Reset()
 	ComponentEmptyDelegate = FChaosVDMeshComponentEmptyDelegate();
 	CurrentMaterial = nullptr;
 
+	SetStaticMesh(nullptr);
+
 	EmptyOverrideMaterials();
 
 	CurrentInstanceHandles.Reset();
 	CurrentGeometryKey = 0;
 }
 
+void UChaosVDInstancedStaticMeshComponent::Initialize()
+{
+	// We need to set the reverse culling flag correctly
+	// And the material
+
+	bReverseCulling = EnumHasAnyFlags(MeshComponentAttributeFlags, EChaosVDMeshAttributesFlags::MirroredGeometry);
+
+	TSharedPtr<FChaosVDGeometryBuilder> GeometryBuilder = GeometryBuilderWeakPtr.Pin();
+	if (ensure(GeometryBuilder))
+	{
+		if (!HasOverrideMaterials())
+		{
+			GeometryBuilder->RequestMaterialInstance(this);
+		}
+	}
+}
+
+void UChaosVDInstancedStaticMeshComponent::SetGeometryBuilder(TWeakPtr<FChaosVDGeometryBuilder> GeometryBuilder)
+{
+	GeometryBuilderWeakPtr = GeometryBuilder;
+}
+
+EChaosVDMaterialType UChaosVDInstancedStaticMeshComponent::GetMaterialType() const
+{
+	if (EnumHasAnyFlags(MeshComponentAttributeFlags, EChaosVDMeshAttributesFlags::TranslucentGeometry))
+	{
+		return  EChaosVDMaterialType::ISMCTranslucent;
+	}
+
+	return EChaosVDMaterialType::ISMCOpaque;
+}
+
+void UChaosVDInstancedStaticMeshComponent::OnDisposed()
+{
+	Reset();
+	
+	SetRelativeTransform(FTransform::Identity);
+
+	if (IsRegistered())
+	{
+		UnregisterComponent();
+	}
+
+	if (AActor* Owner =GetOwner())
+	{
+		Owner->RemoveOwnedComponent(this);
+	}
+}
+
 bool UChaosVDInstancedStaticMeshComponent::UpdateGeometryKey(const uint32 NewHandleGeometryKey)
 {
 	if (CurrentGeometryKey != 0 && CurrentGeometryKey != NewHandleGeometryKey)
diff --git a/Engine/Plugins/ChaosVD/Source/ChaosVD/Private/Components/ChaosVDInstancedStaticMeshComponent.h b/Engine/Plugins/ChaosVD/Source/ChaosVD/Private/Components/ChaosVDInstancedStaticMeshComponent.h
--- a/Engine/Plugins/ChaosVD/Source/ChaosVD/Private/Components/ChaosVDInstancedStaticMeshComponent.h
+++ b/Engine/Plugins/ChaosVD/Source/ChaosVD/Private/Components/ChaosVDInstancedStaticMeshComponent.h
@@ -4,14 +4,15 @@
 
 #include "Components/InstancedStaticMeshComponent.h"
 #include "ChaosVDGeometryDataComponent.h"
+#include "Interfaces/ChaosVDPooledObject.h"
 #include "ChaosVDInstancedStaticMeshComponent.generated.h"
 
 /** CVD version of a Instance Static Mesh Component that holds additional CVD data */
 UCLASS(HideCategories=("Transform"), MinimalAPI)
-class UChaosVDInstancedStaticMeshComponent : public UInstancedStaticMeshComponent, public IChaosVDGeometryComponent
+class UChaosVDInstancedStaticMeshComponent : public UInstancedStaticMeshComponent, public IChaosVDGeometryComponent, public IChaosVDPooledObject
 {
 	GENERATED_BODY()
-
+public:
 	UChaosVDInstancedStaticMeshComponent(const FObjectInitializer& ObjectInitializer)
 	: Super(ObjectInitializer)
 	{
@@ -42,12 +43,18 @@ class UChaosVDInstancedStaticMeshComponent : public UInstancedStaticMeshComponen
 	
 	virtual void UpdateInstanceWorldTransform(const TSharedPtr<FChaosVDMeshDataInstanceHandle>& InInstanceHandle, const FTransform& InTransform) override;
 
-	virtual void SetMeshComponentAttributeFlags(EChaosVDMeshAttributesFlags Flags) override { MeshComponentAttributeFlags = static_cast<uint8>(Flags); }
-	virtual EChaosVDMeshAttributesFlags GetMeshComponentAttributeFlags() const override {return static_cast<EChaosVDMeshAttributesFlags>(MeshComponentAttributeFlags); };
+	virtual void SetMeshComponentAttributeFlags(EChaosVDMeshAttributesFlags Flags) override { MeshComponentAttributeFlags = Flags; }
+	virtual EChaosVDMeshAttributesFlags GetMeshComponentAttributeFlags() const override { return MeshComponentAttributeFlags; };
 
 	virtual TArrayView<TSharedPtr<FChaosVDMeshDataInstanceHandle>> GetMeshDataInstanceHandles() override;
 
 	virtual void Reset() override;
+
+	virtual void Initialize() override;
+
+	virtual void SetGeometryBuilder(TWeakPtr<FChaosVDGeometryBuilder> GeometryBuilder) override;
+	
+	virtual EChaosVDMaterialType GetMaterialType() const;
 	// END IChaosVDGeometryDataComponent Interface
 
 	virtual TSharedPtr<FChaosVDMeshDataInstanceHandle> AddMeshInstance(const FTransform InstanceTransform, bool bIsWorldSpace, const TSharedPtr<FChaosVDExtractedGeometryDataHandle>& InGeometryHandle, int32 ParticleID, int32 SolverID) override;
@@ -58,11 +65,16 @@ class UChaosVDInstancedStaticMeshComponent : public UInstancedStaticMeshComponen
 
 	virtual bool IsNavigationRelevant() const override;
 
+	// BEGIN IChaosVDPooledObject Interface
+	virtual void OnAcquired() override{};
+	virtual void OnDisposed() override;
+	// END IChaosVDPooledObject Interface
+
 protected:
 	
 	bool UpdateGeometryKey(uint32 NewHandleGeometryKey);
 
-	uint8 MeshComponentAttributeFlags = 0;
+	EChaosVDMeshAttributesFlags MeshComponentAttributeFlags = EChaosVDMeshAttributesFlags::None;
 
 	uint32 CurrentGeometryKey = 0;
 
@@ -78,4 +90,6 @@ protected:
 	FChaosVDMeshComponentEmptyDelegate ComponentEmptyDelegate;
 
 	TArray<TSharedPtr<FChaosVDMeshDataInstanceHandle>> CurrentInstanceHandles;
+	
+	TWeakPtr<FChaosVDGeometryBuilder> GeometryBuilderWeakPtr;
 };
diff --git a/Engine/Plugins/ChaosVD/Source/ChaosVD/Private/Components/ChaosVDSceneQueryDataComponent.cpp b/Engine/Plugins/ChaosVD/Source/ChaosVD/Private/Components/ChaosVDSceneQueryDataComponent.cpp
--- a/Engine/Plugins/ChaosVD/Source/ChaosVD/Private/Components/ChaosVDSceneQueryDataComponent.cpp
+++ b/Engine/Plugins/ChaosVD/Source/ChaosVD/Private/Components/ChaosVDSceneQueryDataComponent.cpp
@@ -119,3 +119,12 @@ bool UChaosVDSceneQueryDataComponent::IsQuerySelected(int32 QueryID) const
 
 	return false;
 }
+
+void UChaosVDSceneQueryDataComponent::ClearData()
+{
+	RecordedQueriesByType.Reset();
+	RecordedQueriesByID.Reset();
+	RecordedQueries.Reset();
+	
+	CurrentSQSelectionHandle = FChaosVDSceneQuerySelectionHandle();
+}
diff --git a/Engine/Plugins/ChaosVD/Source/ChaosVD/Private/Components/ChaosVDSceneQueryDataComponent.h b/Engine/Plugins/ChaosVD/Source/ChaosVD/Private/Components/ChaosVDSceneQueryDataComponent.h
--- a/Engine/Plugins/ChaosVD/Source/ChaosVD/Private/Components/ChaosVDSceneQueryDataComponent.h
+++ b/Engine/Plugins/ChaosVD/Source/ChaosVD/Private/Components/ChaosVDSceneQueryDataComponent.h
@@ -2,8 +2,9 @@
 
 #pragma once
 
-#include "Components/ActorComponent.h"
 #include "DataWrappers/ChaosVDQueryDataWrappers.h"
+#include "SolverDataComponent.h"
+
 #include "ChaosVDSceneQueryDataComponent.generated.h"
 
 struct FChaosVDGameFrameData;
@@ -42,7 +43,7 @@ DECLARE_MULTICAST_DELEGATE_OneParam(FChaosVDSQSelectionChangedDelegate, const FC
 
 /** Actor Component that contains all the scene queries recorded at the current loaded frame */
 UCLASS()
-class CHAOSVD_API UChaosVDSceneQueryDataComponent : public UActorComponent
+class CHAOSVD_API UChaosVDSceneQueryDataComponent : public USolverDataComponent
 {
 	GENERATED_BODY()
 
@@ -64,6 +65,8 @@ public:
 
 	FChaosVDSQSelectionChangedDelegate& GetOnSelectionChangeDelegate() { return SelectionChangeDelegate; }
 
+	virtual void ClearData() override;
+
 protected:
 	
 	TMap<EChaosVDSceneQueryType, TArray<TSharedPtr<FChaosVDQueryDataWrapper>>> RecordedQueriesByType;
diff --git a/Engine/Plugins/ChaosVD/Source/ChaosVD/Private/Components/ChaosVDSolverCharacterGroundConstraintDataComponent.h b/Engine/Plugins/ChaosVD/Source/ChaosVD/Private/Components/ChaosVDSolverCharacterGroundConstraintDataComponent.h
--- a/Engine/Plugins/ChaosVD/Source/ChaosVD/Private/Components/ChaosVDSolverCharacterGroundConstraintDataComponent.h
+++ b/Engine/Plugins/ChaosVD/Source/ChaosVD/Private/Components/ChaosVDSolverCharacterGroundConstraintDataComponent.h
@@ -2,7 +2,7 @@
 
 #pragma once
 
-#include "Components/ActorComponent.h"
+#include "SolverDataComponent.h"
 #include "ChaosVDSolverCharacterGroundConstraintDataComponent.generated.h"
 
 struct FChaosVDCharacterGroundConstraint;
@@ -28,8 +28,6 @@ struct FChaosVDCharacterGroundConstraintSelectionHandle
 	void SetIsSelected(bool bNewSelected);
 	bool IsSelected() const;
 
-	bool IsValid();
-
 	TWeakPtr<FChaosVDCharacterGroundConstraint> GetData() const { return ConstraintData; }
 
 private:
@@ -39,7 +37,7 @@ private:
 DECLARE_MULTICAST_DELEGATE_OneParam(FChaosVDCharacterGroundSelectionChangedDelegate, const FChaosVDCharacterGroundConstraintSelectionHandle& SelectionHandle)
 
 UCLASS()
-class CHAOSVD_API UChaosVDSolverCharacterGroundConstraintDataComponent : public UActorComponent
+class CHAOSVD_API UChaosVDSolverCharacterGroundConstraintDataComponent : public USolverDataComponent
 {
 	GENERATED_BODY()
 
@@ -61,11 +59,11 @@ public:
 
 	const FChaosVDCharacterGroundConstraintSelectionHandle& GetCurrentSelectionHandle() const { return CurrentConstraintSelectionHandle; }
 
+	virtual void ClearData() override;
+
 protected:
 
 	TSharedPtr<FChaosVDCharacterGroundConstraint> GetConstraintByIndex(int32 ConstraintIndex);
-	
-	void ClearData();
 
 	FChaosVDCharacterGroundDataArray AllConstraints;
 	
diff --git a/Engine/Plugins/ChaosVD/Source/ChaosVD/Private/Components/ChaosVDSolverCollisionDataComponent.cpp b/Engine/Plugins/ChaosVD/Source/ChaosVD/Private/Components/ChaosVDSolverCollisionDataComponent.cpp
--- a/Engine/Plugins/ChaosVD/Source/ChaosVD/Private/Components/ChaosVDSolverCollisionDataComponent.cpp
+++ b/Engine/Plugins/ChaosVD/Source/ChaosVD/Private/Components/ChaosVDSolverCollisionDataComponent.cpp
@@ -12,7 +12,7 @@ UChaosVDSolverCollisionDataComponent::UChaosVDSolverCollisionDataComponent()
 
 void UChaosVDSolverCollisionDataComponent::UpdateCollisionData(const TArray<TSharedPtr<FChaosVDParticlePairMidPhase>>& InMidPhaseData)
 {
-	ClearCollisionData();
+	ClearData();
 
 	AllMidPhases.Reserve(InMidPhaseData.Num());
 	AllMidPhases = InMidPhaseData;
@@ -45,7 +45,7 @@ const TArray<FChaosVDConstraint*>* UChaosVDSolverCollisionDataComponent::GetCons
 	return Chaos::VisualDebugger::Utils::GetDataFromParticlePairMaps<FChaosVDConstraintByParticleMap, FChaosVDConstraint*>(ConstraintsByParticleID0, ConstraintsByParticleID1, ParticleID, Options);
 }
 
-void UChaosVDSolverCollisionDataComponent::ClearCollisionData()
+void UChaosVDSolverCollisionDataComponent::ClearData()
 {
 	AllMidPhases.Reset();
 	MidPhasesByParticleID0.Reset();
diff --git a/Engine/Plugins/ChaosVD/Source/ChaosVD/Private/Components/ChaosVDSolverCollisionDataComponent.h b/Engine/Plugins/ChaosVD/Source/ChaosVD/Private/Components/ChaosVDSolverCollisionDataComponent.h
--- a/Engine/Plugins/ChaosVD/Source/ChaosVD/Private/Components/ChaosVDSolverCollisionDataComponent.h
+++ b/Engine/Plugins/ChaosVD/Source/ChaosVD/Private/Components/ChaosVDSolverCollisionDataComponent.h
@@ -3,12 +3,12 @@
 #pragma once
 
 #include "ChaosVDConstraintDataHelpers.h"
-#include "Components/ActorComponent.h"
 
 #include "Containers/Array.h"
 #include "Containers/Map.h"
 #include "DataWrappers/ChaosVDCollisionDataWrappers.h"
 #include "GameFramework/Info.h"
+#include "SolverDataComponent.h"
 #include "Templates/SharedPointer.h"
 
 #include "ChaosVDSolverCollisionDataComponent.generated.h"
@@ -19,7 +19,7 @@ typedef TMap<int32, TArray<TSharedPtr<FChaosVDParticlePairMidPhase>>> FChaosVDMi
 typedef TMap<int32, TArray<FChaosVDConstraint*>> FChaosVDConstraintByParticleMap;
 
 UCLASS()
-class UChaosVDSolverCollisionDataComponent : public UActorComponent
+class UChaosVDSolverCollisionDataComponent : public USolverDataComponent
 {
 	GENERATED_BODY()
 
@@ -32,9 +32,9 @@ public:
 	const TArray<TSharedPtr<FChaosVDParticlePairMidPhase>>* GetMidPhasesForParticle(int32 ParticleID, EChaosVDParticlePairSlot Options) const;
 	const TArray<FChaosVDConstraint*>* GetConstraintsForParticle(int32 ParticleID, EChaosVDParticlePairSlot Options) const;
 
-protected:
+	virtual void ClearData() override;
 
-	void ClearCollisionData();
+protected:
 
 	TArray<TSharedPtr<FChaosVDParticlePairMidPhase>> AllMidPhases;
 	FChaosVDMidPhaseByParticleMap MidPhasesByParticleID0;
diff --git a/Engine/Plugins/ChaosVD/Source/ChaosVD/Private/Components/ChaosVDSolverJointConstraintDataComponent.h b/Engine/Plugins/ChaosVD/Source/ChaosVD/Private/Components/ChaosVDSolverJointConstraintDataComponent.h
--- a/Engine/Plugins/ChaosVD/Source/ChaosVD/Private/Components/ChaosVDSolverJointConstraintDataComponent.h
+++ b/Engine/Plugins/ChaosVD/Source/ChaosVD/Private/Components/ChaosVDSolverJointConstraintDataComponent.h
@@ -2,7 +2,7 @@
 
 #pragma once
 
-#include "Components/ActorComponent.h"
+#include "SolverDataComponent.h"
 #include "ChaosVDSolverJointConstraintDataComponent.generated.h"
 
 struct FChaosVDJointConstraint;
@@ -28,8 +28,6 @@ struct FChaosVDJointConstraintSelectionHandle
 	void SetIsSelected(bool bNewSelected);
 	bool IsSelected() const;
 
-	bool IsValid();
-
 	TWeakPtr<FChaosVDJointConstraint> GetData() const { return JointData; } 
 
 private:
@@ -39,7 +37,7 @@ private:
 DECLARE_MULTICAST_DELEGATE_OneParam(FChaosVDJointSlectionChangedDelegate, const FChaosVDJointConstraintSelectionHandle& SelectionHandle)
 
 UCLASS()
-class CHAOSVD_API UChaosVDSolverJointConstraintDataComponent : public UActorComponent
+class CHAOSVD_API UChaosVDSolverJointConstraintDataComponent : public USolverDataComponent
 {
 	GENERATED_BODY()
 
@@ -61,11 +59,11 @@ public:
 
 	const FChaosVDJointConstraintSelectionHandle& GetCurrentSelectionHandle() const { return CurrentJointSelectionHandle; }
 
+	virtual void ClearData() override;
+
 protected:
 
 	TSharedPtr<FChaosVDJointConstraint> GetJointConstraintByIndex(int32 ConstraintIndex);
-	
-	void ClearData();
 
 	FChaosVDJointDataArray AllJointConstraints;
 	
diff --git a/Engine/Plugins/ChaosVD/Source/ChaosVD/Private/Components/ChaosVDStaticMeshComponent.cpp b/Engine/Plugins/ChaosVD/Source/ChaosVD/Private/Components/ChaosVDStaticMeshComponent.cpp
--- a/Engine/Plugins/ChaosVD/Source/ChaosVD/Private/Components/ChaosVDStaticMeshComponent.cpp
+++ b/Engine/Plugins/ChaosVD/Source/ChaosVD/Private/Components/ChaosVDStaticMeshComponent.cpp
@@ -2,10 +2,12 @@
 
 #include "Components/ChaosVDStaticMeshComponent.h"
 
+#include "ChaosVDGeometryBuilder.h"
 #include "ChaosVDModule.h"
-#include "MaterialDomain.h"
 #include "Materials/Material.h"
 #include "Materials/MaterialInstanceDynamic.h"
+#include "Widgets/SChaosVDEnumFlagsMenu.h"
+
 
 uint32 UChaosVDStaticMeshComponent::GetGeometryKey() const
 {
@@ -54,20 +56,26 @@ void UChaosVDStaticMeshComponent::UpdateInstanceColor(const TSharedPtr<FChaosVDM
 		UE_LOG(LogChaosVDEditor, Error, TEXT("[%s] Attempted to update a mesh instance using a handle from another component. No instances were updated | Handle Component [%s] | Current Component [%s]"), ANSI_TO_TCHAR(__FUNCTION__), *GetNameSafe(InInstanceHandle->GetMeshComponent()), *GetNameSafe(this));
 		return;
 	}
-
+	
 	const bool bIsSolidColor = FMath::IsNearlyEqual(NewColor.A, 1.0f);
 
-	if (UMaterialInstanceDynamic* MaterialToApply = GetCachedMaterialInstance(bIsSolidColor ? EChaosVDMaterialType::SimOnlyMaterial : EChaosVDMaterialType::QueryOnlyMaterial))
-	{
-		MaterialToApply->SetVectorParameterValue(TEXT("BaseColor"), NewColor);
-		InInstanceHandle->GetMeshComponent()->SetMaterial(0, MaterialToApply);
-	}
-	else
+	bool bHasOpaqueMaterial = !EnumHasAnyFlags(MeshComponentAttributeFlags, EChaosVDMeshAttributesFlags::TranslucentGeometry);
+
+	constexpr int32 ColorPrimitiveDataIndex = 0;
+	SetCustomPrimitiveDataVector4(ColorPrimitiveDataIndex, NewColor);
+
+	// If we want to show a color with transparency, we might need to change our material
+	if (bHasOpaqueMaterial != bIsSolidColor)
 	{
-		ensure(false);
-		UE_LOG(LogChaosVDEditor, Error, TEXT("[%s] Failed to get Query Only material for, applying the default mesh to all geometry"), ANSI_TO_TCHAR(__FUNCTION__));
+		TSharedPtr<FChaosVDGeometryBuilder> GeometryBuilder = GeometryBuilderWeakPtr.Pin();
+		if (ensure(GeometryBuilder))
+		{
+			EmptyOverrideMaterials();
+
+			Chaos::VisualDebugger::Utils::EnumAddToggleFlag(MeshComponentAttributeFlags, EChaosVDMeshAttributesFlags::TranslucentGeometry);
 			
-		InInstanceHandle->GetMeshComponent()->SetMaterial(0, UMaterial::GetDefaultMaterial(EMaterialDomain::MD_Surface));
+			GeometryBuilder->RequestMaterialInstance(this);
+		}
 	}
 }
 
@@ -94,6 +102,8 @@ void UChaosVDStaticMeshComponent::Reset()
 	MeshReadyDelegate = FChaosVDMeshReadyDelegate();
 	ComponentEmptyDelegate = FChaosVDMeshComponentEmptyDelegate();
 
+	SetStaticMesh(nullptr);
+
 	CurrentMeshDataHandle = nullptr;
 	CurrentGeometryKey = 0;
 }
@@ -142,6 +152,38 @@ void UChaosVDStaticMeshComponent::RemoveMeshInstance(TSharedPtr<FChaosVDMeshData
 	ComponentEmptyDelegate.Broadcast(this);
 }
 
+void UChaosVDStaticMeshComponent::SetGeometryBuilder(TWeakPtr<FChaosVDGeometryBuilder> GeometryBuilder)
+{
+	GeometryBuilderWeakPtr = GeometryBuilder;
+}
+
+EChaosVDMaterialType UChaosVDStaticMeshComponent::GetMaterialType() const
+{
+	if (EnumHasAnyFlags(MeshComponentAttributeFlags, EChaosVDMeshAttributesFlags::TranslucentGeometry))
+	{
+		return  EChaosVDMaterialType::SMTranslucent;
+	}
+
+	return EChaosVDMaterialType::SMOpaque;
+}
+
+void UChaosVDStaticMeshComponent::OnDisposed()
+{
+	Reset();
+
+	SetRelativeTransform(FTransform::Identity);
+
+	if (IsRegistered())
+	{
+		UnregisterComponent();
+	}
+
+	if (AActor* Owner =GetOwner())
+	{
+		Owner->RemoveOwnedComponent(this);
+	}
+}
+
 bool UChaosVDStaticMeshComponent::UpdateGeometryKey(uint32 NewHandleGeometryKey)
 {
 	if (CurrentGeometryKey != 0 && CurrentGeometryKey != NewHandleGeometryKey)
@@ -159,7 +201,6 @@ bool UChaosVDStaticMeshComponent::UpdateGeometryKey(uint32 NewHandleGeometryKey)
 	return true;
 }
 
-
 TSharedPtr<FChaosVDMeshDataInstanceHandle> UChaosVDStaticMeshComponent::GetMeshDataInstanceHandle(int32 InstanceIndex) const
 {
 	return CurrentMeshDataHandle;
@@ -170,22 +211,11 @@ TArrayView<TSharedPtr<FChaosVDMeshDataInstanceHandle>> UChaosVDStaticMeshCompone
 	return TArrayView<TSharedPtr<FChaosVDMeshDataInstanceHandle>>(&CurrentMeshDataHandle, 1);
 }
 
-UMaterialInstanceDynamic* UChaosVDStaticMeshComponent::GetCachedMaterialInstance(EChaosVDMaterialType Type)
+void UChaosVDStaticMeshComponent::Initialize()
 {
-	if (const TObjectPtr<UMaterialInstanceDynamic>* MaterialInstance = CachedMaterialInstancesByID.Find(Type))
-	{
-		return MaterialInstance->Get();
-	}
-	else
+	TSharedPtr<FChaosVDGeometryBuilder> GeometryBuilder = GeometryBuilderWeakPtr.Pin();
+	if (ensure(GeometryBuilder))
 	{
-		if (UMaterialInterface* MaterialToCreate = FChaosVDGeometryComponentUtils::GetBaseMaterialForType(Type))
-		{
-			UMaterialInstanceDynamic* DynamicMaterial = UMaterialInstanceDynamic::Create(MaterialToCreate, nullptr);
-			CachedMaterialInstancesByID.Add(Type, TObjectPtr<UMaterialInstanceDynamic>(DynamicMaterial));
-
-			return DynamicMaterial;
-		}
+		GeometryBuilder->RequestMaterialInstance(this);
 	}
-
-	return nullptr;
 }
diff --git a/Engine/Plugins/ChaosVD/Source/ChaosVD/Private/Components/ChaosVDStaticMeshComponent.h b/Engine/Plugins/ChaosVD/Source/ChaosVD/Private/Components/ChaosVDStaticMeshComponent.h
--- a/Engine/Plugins/ChaosVD/Source/ChaosVD/Private/Components/ChaosVDStaticMeshComponent.h
+++ b/Engine/Plugins/ChaosVD/Source/ChaosVD/Private/Components/ChaosVDStaticMeshComponent.h
@@ -4,11 +4,12 @@
 
 #include "Components/StaticMeshComponent.h"
 #include "ChaosVDGeometryDataComponent.h"
+#include "Interfaces/ChaosVDPooledObject.h"
 #include "ChaosVDStaticMeshComponent.generated.h"
 
 /** CVD version of a Static Mesh Component that holds additional CVD data */
 UCLASS(HideCategories=("Transform"), MinimalAPI)
-class UChaosVDStaticMeshComponent : public UStaticMeshComponent, public IChaosVDGeometryComponent
+class UChaosVDStaticMeshComponent : public UStaticMeshComponent, public IChaosVDGeometryComponent, public IChaosVDPooledObject
 {
 	GENERATED_BODY()
 
@@ -36,31 +37,39 @@ public:
 	virtual void SetIsSelected(const TSharedPtr<FChaosVDMeshDataInstanceHandle>& InInstanceHandle, bool bIsSelected) override;
 
 	virtual bool ShouldRenderSelected() const override;
-	
+
 	virtual void UpdateInstanceColor(const TSharedPtr<FChaosVDMeshDataInstanceHandle>& InInstanceHandle, FLinearColor NewColor) override;
 	virtual void UpdateInstanceWorldTransform(const TSharedPtr<FChaosVDMeshDataInstanceHandle>& InInstanceHandle, const FTransform& InTransform) override;
 
-	virtual void SetMeshComponentAttributeFlags(EChaosVDMeshAttributesFlags Flags) override { MeshComponentAttributeFlags = static_cast<uint8>(Flags); }
-	virtual EChaosVDMeshAttributesFlags GetMeshComponentAttributeFlags() const override {return static_cast<EChaosVDMeshAttributesFlags>(MeshComponentAttributeFlags); };
+	virtual void SetMeshComponentAttributeFlags(EChaosVDMeshAttributesFlags Flags) override { MeshComponentAttributeFlags = Flags; };
+	virtual EChaosVDMeshAttributesFlags GetMeshComponentAttributeFlags() const override { return MeshComponentAttributeFlags; };
 
 	virtual TSharedPtr<FChaosVDMeshDataInstanceHandle> GetMeshDataInstanceHandle(int32 InstanceIndex) const override;
 	virtual TArrayView<TSharedPtr<FChaosVDMeshDataInstanceHandle>> GetMeshDataInstanceHandles() override;
 
+	virtual void Initialize() override;
 	virtual void Reset() override;
 
 	virtual TSharedPtr<FChaosVDMeshDataInstanceHandle> AddMeshInstance(const FTransform InstanceTransform, bool bIsWorldSpace, const TSharedPtr<FChaosVDExtractedGeometryDataHandle>& InGeometryHandle, int32 ParticleID, int32 SolverID) override;
 	virtual void AddMeshInstanceForHandle(TSharedPtr<FChaosVDMeshDataInstanceHandle> MeshDataHandle, const FTransform InstanceTransform, bool bIsWorldSpace, const TSharedPtr<FChaosVDExtractedGeometryDataHandle>& InGeometryHandle, int32 ParticleID, int32 SolverID) override;
 	virtual void RemoveMeshInstance(TSharedPtr<FChaosVDMeshDataInstanceHandle> InHandleToRemove) override;
+
+	virtual void SetGeometryBuilder(TWeakPtr<FChaosVDGeometryBuilder> GeometryBuilder) override;
+
+	virtual EChaosVDMaterialType GetMaterialType() const override;
+
 	// END IChaosVDGeometryDataComponent Interface
 
+	// BEGIN IChaosVDPooledObject Interface
+	virtual void OnAcquired() override {};
+	virtual void OnDisposed() override;
+	// END IChaosVDPooledObject Interface
+
 protected:
 
 	bool UpdateGeometryKey(uint32 NewHandleGeometryKey);
 
-	/** Returns an existing material instance used by this mesh instances, or creates a new one for the provided type */
-	UMaterialInstanceDynamic* GetCachedMaterialInstance(EChaosVDMaterialType Type);
-
-	uint8 MeshComponentAttributeFlags = 0;
+	EChaosVDMeshAttributesFlags MeshComponentAttributeFlags = EChaosVDMeshAttributesFlags::None;
 	uint8 CurrentGeometryKey = 0;
 	bool bIsMeshReady = false;
 	bool bIsOwningParticleSelected = false;
@@ -70,7 +79,6 @@ protected:
 	TSharedPtr<FChaosVDMeshDataInstanceHandle> CurrentMeshDataHandle = nullptr;
 
 	TSharedPtr<FChaosVDExtractedGeometryDataHandle> CurrentGeometryHandle = nullptr;
-
-	UPROPERTY(Transient)
-	TMap<EChaosVDMaterialType, TObjectPtr<UMaterialInstanceDynamic>> CachedMaterialInstancesByID;
+	
+	TWeakPtr<FChaosVDGeometryBuilder> GeometryBuilderWeakPtr;
 };
diff --git a/Engine/Plugins/ChaosVD/Source/ChaosVD/Private/Components/SolverDataComponent.h b/Engine/Plugins/ChaosVD/Source/ChaosVD/Private/Components/SolverDataComponent.h
new file mode 100644
--- /dev/null
+++ b/Engine/Plugins/ChaosVD/Source/ChaosVD/Private/Components/SolverDataComponent.h
@@ -0,0 +1,17 @@
+// Copyright Epic Games, Inc. All Rights Reserved.
+
+#pragma once
+
+#include "Components/ActorComponent.h"
+#include "SolverDataComponent.generated.h"
+
+/**
+ * Base class for all components that stores recorded solver data
+ */
+UCLASS(Abstract)
+class USolverDataComponent : public UActorComponent
+{
+	GENERATED_BODY()
+
+	virtual void ClearData() PURE_VIRTUAL(ISolverDataComponent::ClearData);
+};
diff --git a/Engine/Plugins/ChaosVD/Source/ChaosVD/Private/Interfaces/ChaosVDPooledObject.h b/Engine/Plugins/ChaosVD/Source/ChaosVD/Private/Interfaces/ChaosVDPooledObject.h
new file mode 100644
--- /dev/null
+++ b/Engine/Plugins/ChaosVD/Source/ChaosVD/Private/Interfaces/ChaosVDPooledObject.h
@@ -0,0 +1,24 @@
+// Copyright Epic Games, Inc. All Rights Reserved.
+
+#pragma once
+
+#include "UObject/Interface.h"
+#include "ChaosVDPooledObject.generated.h"
+
+UINTERFACE()
+class UChaosVDPooledObject : public UInterface
+{
+	GENERATED_BODY()
+};
+
+/**
+ * Interface used by any object that can be pooled
+ */
+class CHAOSVD_API IChaosVDPooledObject
+{
+	GENERATED_BODY()
+
+public:
+	virtual void OnAcquired() PURE_VIRTUAL(IChaosVDPooledObject::OnAcquired)
+	virtual void OnDisposed() PURE_VIRTUAL(IChaosVDPooledObject::OnDisposed)
+};
diff --git a/Engine/Plugins/ChaosVD/Source/ChaosVD/Private/ObjectsWaitingGeometryList.h b/Engine/Plugins/ChaosVD/Source/ChaosVD/Private/ObjectsWaitingGeometryList.h
--- a/Engine/Plugins/ChaosVD/Source/ChaosVD/Private/ObjectsWaitingGeometryList.h
+++ b/Engine/Plugins/ChaosVD/Source/ChaosVD/Private/ObjectsWaitingGeometryList.h
@@ -25,6 +25,13 @@ public:
 	{
 	}
 
+	~FObjectsWaitingGeometryList()
+	 {
+	 	// Clean up the notification progress bar
+	 	WaitingObjectsByGeometryKey.Empty();
+	 	AsyncProgressNotification.Update(0);
+	 }
+
 	/** Adds an object to process when the provided geometry key is ready */
 	void AddObject(uint32 GeometryKey, ObjectType ObjectToProcess);
 
diff --git a/Engine/Plugins/ChaosVD/Source/ChaosVD/Private/Visualizers/ChaosVDSceneQueryDataComponentVisualizer.cpp b/Engine/Plugins/ChaosVD/Source/ChaosVD/Private/Visualizers/ChaosVDSceneQueryDataComponentVisualizer.cpp
--- a/Engine/Plugins/ChaosVD/Source/ChaosVD/Private/Visualizers/ChaosVDSceneQueryDataComponentVisualizer.cpp
+++ b/Engine/Plugins/ChaosVD/Source/ChaosVD/Private/Visualizers/ChaosVDSceneQueryDataComponentVisualizer.cpp
@@ -87,7 +87,7 @@ void FChaosVDSceneQueryDataComponentVisualizer::DrawVisualization(const UActorCo
 		return;
 	}
 
-	const TSharedPtr<FChaosVDGeometryBuilder> GeometryGenerator = CVDScene->GetGeometryGenerator();
+	const TSharedPtr<FChaosVDGeometryBuilder> GeometryGenerator = CVDScene->GetGeometryGenerator().Pin();
 	if (!GeometryGenerator)
 	{
 		return;
diff --git a/Engine/Source/Runtime/Experimental/ChaosVisualDebugger/Public/DataWrappers/ChaosVDCollisionDataWrappers.h b/Engine/Source/Runtime/Experimental/ChaosVisualDebugger/Public/DataWrappers/ChaosVDCollisionDataWrappers.h
--- a/Engine/Source/Runtime/Experimental/ChaosVisualDebugger/Public/DataWrappers/ChaosVDCollisionDataWrappers.h
+++ b/Engine/Source/Runtime/Experimental/ChaosVisualDebugger/Public/DataWrappers/ChaosVDCollisionDataWrappers.h
@@ -371,6 +371,24 @@ enum class EChaosVDMidPhaseFlags : uint8
 	IsModified = 1 << 4,
 };
 
+UENUM()
+enum class EChaosVDMidPhaseType : int8
+{
+	// A general purpose midphase that handle BVHs, Meshes, 
+	// Unions of Unions, etc in the geometry hierarchy.
+	Generic,
+
+	// A midphase optimized for particle pairs with a small
+	// number of shapes. Pre-expands the set of potentially
+	// colliding shape pairs.
+	ShapePair,
+
+	// A midphase used to collide particles as sphere approximations
+	SphereApproximation,
+
+	Unknown
+};
+
 USTRUCT()
 struct CHAOSVDRUNTIME_API FChaosVDParticlePairMidPhase
 {
@@ -381,6 +399,9 @@ struct CHAOSVDRUNTIME_API FChaosVDParticlePairMidPhase
 	UPROPERTY()
 	int32 SolverID = INDEX_NONE;
 
+	UPROPERTY(VisibleAnywhere, Category=General)
+	EChaosVDMidPhaseType MidPhaseType;
+
 	UPROPERTY(VisibleAnywhere, Category=Flags)
 	uint8 bIsActive:1 = false;
 	UPROPERTY(VisibleAnywhere, Category=Flags)
@@ -454,6 +475,12 @@ struct CHAOSVDRUNTIME_API FChaosVDCollisionFilterData
 	bool Serialize(FArchive& Ar);
 };
 
+template <>
+struct TTypeTraits<FChaosVDCollisionFilterData> : public TTypeTraitsBase <FChaosVDCollisionFilterData>
+{
+	enum { IsBytewiseComparable = true };
+};
+
 template<>
 struct TStructOpsTypeTraits<FChaosVDCollisionFilterData> : public TStructOpsTypeTraitsBase2<FChaosVDCollisionFilterData>
 {
@@ -509,6 +536,12 @@ struct CHAOSVDRUNTIME_API FChaosVDShapeCollisionData
 	bool operator==(const FChaosVDShapeCollisionData& Other) const;
 };
 
+template <>
+struct TTypeTraits<FChaosVDShapeCollisionData> : public TTypeTraitsBase <FChaosVDShapeCollisionData>
+{
+	enum { IsBytewiseComparable = true };
+};
+
 template<>
 struct TStructOpsTypeTraits<FChaosVDShapeCollisionData> : public TStructOpsTypeTraitsBase2<FChaosVDShapeCollisionData>
 {
