[ChaosVD] Fix CIS issue in ChaosVDObjectPool

diff --git a/Engine/Plugins/ChaosVD/Source/ChaosVD/Private/ChaosVDObjectPool.h b/Engine/Plugins/ChaosVD/Source/ChaosVD/Private/ChaosVDObjectPool.h
--- a/Engine/Plugins/ChaosVD/Source/ChaosVD/Private/ChaosVDObjectPool.h
+++ b/Engine/Plugins/ChaosVD/Source/ChaosVD/Private/ChaosVDObjectPool.h
@@ -7,6 +7,8 @@
 #include "Components/ActorComponent.h"
 #include "GameFramework/Actor.h"
 #include "Interfaces/ChaosVDPooledObject.h"
+#include "UObject/Object.h"
+#include "UObject/Package.h"
 
 struct FChaosVDObjectPoolCVars
 {
