[ChaosVD] Fix for CIS issue on FChaosVDDataProcessorBase

There was a stranded pragma once in the cpp file, as result of a copy/paste error during the refactor when I moved this code done in a previous CL, it seems it was not caught by the preglight

diff --git a/Engine/Plugins/ChaosVD/Source/ChaosVD/Private/Trace/DataProcessors/ChaosVDDataProcessorBase.cpp b/Engine/Plugins/ChaosVD/Source/ChaosVD/Private/Trace/DataProcessors/ChaosVDDataProcessorBase.cpp
--- a/Engine/Plugins/ChaosVD/Source/ChaosVD/Private/Trace/DataProcessors/ChaosVDDataProcessorBase.cpp
+++ b/Engine/Plugins/ChaosVD/Source/ChaosVD/Private/Trace/DataProcessors/ChaosVDDataProcessorBase.cpp
@@ -1,7 +1,5 @@
 // Copyright Epic Games, Inc. All Rights Reserved.
 
-#pragma once
-
 #include "ChaosVDDataProcessorBase.h"
 
 FChaosVDDataProcessorBase::~FChaosVDDataProcessorBase()
