[ChaosVD] Removing VisibleAnywhere form the mesh component property in FChaosVDMeshDataInstanceState

This was causing several seconds hitches when selecting a particle actor as it would embed the details panel from the ISMC. This property does not need to be visible

diff --git a/Engine/Plugins/ChaosVD/Source/ChaosVD/Private/ChaosVDGeometryDataComponent.h b/Engine/Plugins/ChaosVD/Source/ChaosVD/Private/ChaosVDGeometryDataComponent.h
--- a/Engine/Plugins/ChaosVD/Source/ChaosVD/Private/ChaosVDGeometryDataComponent.h
+++ b/Engine/Plugins/ChaosVD/Source/ChaosVD/Private/ChaosVDGeometryDataComponent.h
@@ -98,7 +98,7 @@ struct FChaosVDMeshDataInstanceState
 	EChaosVDMeshComponent MeshComponentType = EChaosVDMeshComponent::Invalid;
 
 	/* CVD Debug - Pointer to the mesh component used to render this Mesh */
-	UPROPERTY(VisibleAnywhere, Category="CVD GeometryData")
+	UPROPERTY()
 	TObjectPtr<UMeshComponent> MeshComponent;
 
 	/* CVD Debug - Instance index of mesh component used to render this Mesh */
