[ChaosVD] Refactor for selection system for solver data visualization

Before this change, the selection system for debug draw elements was implemented per visualization type. No everything uses the same selection system

- Created a simple selection system that uses a single handle type for the different solver data types.
- Updated all component visualizers to use the new selection system
- Created a base component visualizer class with all the common logic for all visualizers
- Now almost all visualizers use the handle click logic from the base visualizer class.
- Moved most of the logic to create the menu entries for visualizers to a helper method. All visualizers now use this helper to create the menus (except particle data visualizers as it requires a more complex menu for now).
- Refactored the Joint Constraint data inspector widget to make it work with the new generic selection system, and work with different types of constraints
- Deleted the Character Ground Constraints inspector widget and tab logic. This was almost all duplicated from the Joint Constraint Inspector, but the API at the time didn't allow us to re-use it. Now the Character Ground Constraints use the same Constraint Dtaa Inspector as Joint Constraints
- Created a base SolverConstraintDataComponent and moved all the logic to organize and query constraint data to it. This allows Character Constraints, Collision Constraints and Joint Constraints component use the same code, eliminating a lot of code duplication.
- Modified the Collision Data Inspector widget to make it derive from the new generic Constraint Data Inspector, to reduce code duplication.

diff --git a/Engine/Plugins/ChaosVD/Source/ChaosVD/Private/Actors/ChaosVDSolverInfoActor.cpp b/Engine/Plugins/ChaosVD/Source/ChaosVD/Private/Actors/ChaosVDSolverInfoActor.cpp
--- a/Engine/Plugins/ChaosVD/Source/ChaosVD/Private/Actors/ChaosVDSolverInfoActor.cpp
+++ b/Engine/Plugins/ChaosVD/Source/ChaosVD/Private/Actors/ChaosVDSolverInfoActor.cpp
@@ -51,6 +51,19 @@ void AChaosVDSolverInfoActor::SetScene(TWeakPtr<FChaosVDScene> InScene)
 	{
 		RegisterSelectionSetObject(ScenePtr->GetElementSelectionSet());
 	}
+
+	TInlineComponentArray<UChaosVDSolverDataComponent*> SolverDataComponents;
+	GetComponents(SolverDataComponents);
+
+	for (UChaosVDSolverDataComponent* Component : SolverDataComponents)
+	{
+		if(!Component)
+		{
+			continue;
+		}
+
+		Component->SetScene(InScene);
+	}
 }
 
 void AChaosVDSolverInfoActor::RegisterParticleActor(int32 ParticleID, AChaosVDParticleActor* ParticleActor)
@@ -257,7 +270,7 @@ void AChaosVDSolverInfoActor::HandlePostSelectionChange(const UTypedElementSelec
 	{
 		if (AChaosVDParticleActor* SelectedParticle = SelectedParticles[0])
 		{
-			if (const FChaosVDParticleDataWrapper* ParticleData = SelectedParticle->GetParticleData())
+			if (TSharedPtr<const FChaosVDParticleDataWrapper> ParticleData = SelectedParticle->GetParticleData())
 			{
 				SelectedParticlesID.Add(ParticleData->ParticleIndex);
 			}
diff --git a/Engine/Plugins/ChaosVD/Source/ChaosVD/Private/Actors/ChaosVDSolverInfoActor.h b/Engine/Plugins/ChaosVD/Source/ChaosVD/Private/Actors/ChaosVDSolverInfoActor.h
--- a/Engine/Plugins/ChaosVD/Source/ChaosVD/Private/Actors/ChaosVDSolverInfoActor.h
+++ b/Engine/Plugins/ChaosVD/Source/ChaosVD/Private/Actors/ChaosVDSolverInfoActor.h
@@ -121,13 +121,13 @@ void AChaosVDSolverInfoActor::VisitSelectedParticleData(TCallback VisitCallback)
 	for (const int32 SelectedParticleID : SelectedParticlesID)
 	{
 		AChaosVDParticleActor* ParticleActor = GetParticleActor(SelectedParticleID);
-		const FChaosVDParticleDataWrapper* ParticleDataViewer = ParticleActor ? ParticleActor->GetParticleData() : nullptr;
+		TSharedPtr<const FChaosVDParticleDataWrapper> ParticleDataViewer = ParticleActor ? ParticleActor->GetParticleData() : nullptr;
 		if (!ensure(ParticleDataViewer))
 		{
 			continue;
 		}
 
-		if (!VisitCallback(*ParticleDataViewer))
+		if (!VisitCallback(ParticleDataViewer))
 		{
 			return;
 		}
@@ -140,13 +140,13 @@ void AChaosVDSolverInfoActor::VisitAllParticleData(TCallback VisitCallback)
 	for (const TPair<int32, AChaosVDParticleActor*>& ParticleWithIDPair : SolverParticlesByID)
 	{
 		AChaosVDParticleActor* ParticleActor = ParticleWithIDPair.Value;
-		const FChaosVDParticleDataWrapper* ParticleDataViewer = ParticleActor ? ParticleActor->GetParticleData() : nullptr;
+		TSharedPtr<const FChaosVDParticleDataWrapper> ParticleDataViewer = ParticleActor ? ParticleActor->GetParticleData() : nullptr;
 		if (!ensure(ParticleDataViewer))
 		{
 			continue;
 		}
 
-		if (!VisitCallback(*ParticleDataViewer))
+		if (!VisitCallback(ParticleDataViewer))
 		{
 			return;
 		}
diff --git a/Engine/Plugins/ChaosVD/Source/ChaosVD/Private/ChaosVDCharacterGroundConstraintDataInspectorTab.cpp b/Engine/Plugins/ChaosVD/Source/ChaosVD/Private/ChaosVDCharacterGroundConstraintDataInspectorTab.cpp
deleted file mode 100644
diff --git a/Engine/Plugins/ChaosVD/Source/ChaosVD/Private/ChaosVDCharacterGroundConstraintDataInspectorTab.h b/Engine/Plugins/ChaosVD/Source/ChaosVD/Private/ChaosVDCharacterGroundConstraintDataInspectorTab.h
deleted file mode 100644
diff --git a/Engine/Plugins/ChaosVD/Source/ChaosVD/Private/ChaosVDCollisionDataProviderInterface.h b/Engine/Plugins/ChaosVD/Source/ChaosVD/Private/ChaosVDCollisionDataProviderInterface.h
--- a/Engine/Plugins/ChaosVD/Source/ChaosVD/Private/ChaosVDCollisionDataProviderInterface.h
+++ b/Engine/Plugins/ChaosVD/Source/ChaosVD/Private/ChaosVDCollisionDataProviderInterface.h
@@ -5,6 +5,8 @@
 #include "UObject/Interface.h"
 #include "ChaosVDCollisionDataProviderInterface.generated.h"
 
+struct FChaosVDParticlePairMidPhase;
+struct FChaosVDSolverDataSelectionHandle;
 struct FChaosVDCollisionDataFinder;
 
 // This class does not need to be modified.
@@ -25,9 +27,8 @@ public:
 	
 	/**
 	 * Gathers and populates the provided array with any existing collision data for this object (if any) 
-	 * @param OutCollisionDataFound Array to populate with collision data
 	 */
-	virtual void GetCollisionData(TArray<TSharedPtr<FChaosVDCollisionDataFinder>>& OutCollisionDataFound) PURE_VIRTUAL(IChaosVDCollisionDataProviderInterface::GetCollisionData)
+	virtual TConstArrayView<TSharedPtr<FChaosVDParticlePairMidPhase>> GetCollisionData() PURE_VIRTUAL(IChaosVDCollisionDataProviderInterface::GetCollisionData, return TConstArrayView<TSharedPtr<FChaosVDParticlePairMidPhase>>(););
 
 	/**
 	 * Checks if the object implementing the interface has any collision data available
diff --git a/Engine/Plugins/ChaosVD/Source/ChaosVD/Private/ChaosVDParticleActor.cpp b/Engine/Plugins/ChaosVD/Source/ChaosVD/Private/ChaosVDParticleActor.cpp
--- a/Engine/Plugins/ChaosVD/Source/ChaosVD/Private/ChaosVDParticleActor.cpp
+++ b/Engine/Plugins/ChaosVD/Source/ChaosVD/Private/ChaosVDParticleActor.cpp
@@ -400,22 +400,13 @@ FBox AChaosVDParticleActor::GetComponentsBoundingBox(bool bNonColliding, bool bI
 	return BoundingBox;
 }
 
-void AChaosVDParticleActor::GetCollisionData(TArray<TSharedPtr<FChaosVDCollisionDataFinder>>& OutCollisionDataFound)
+TConstArrayView<TSharedPtr<FChaosVDParticlePairMidPhase>> AChaosVDParticleActor::GetCollisionData()
 {
 	if (const TArray<TSharedPtr<FChaosVDParticlePairMidPhase>>* MidPhases = GetCollisionMidPhasesArray())
 	{
-		OutCollisionDataFound.Reserve(MidPhases->Num());
-
-		for (const TSharedPtr<FChaosVDParticlePairMidPhase>& MidPhasePtr : *MidPhases)
-		{
-			TSharedPtr<FChaosVDCollisionDataFinder> FinderData = MakeShared<FChaosVDCollisionDataFinder>();
-			FinderData->OwningMidPhase = MidPhasePtr;
-			FinderData->OwningConstraint = MidPhasePtr->Constraints.Num() > 0 ? &MidPhasePtr->Constraints[0] : nullptr;
-			FinderData->ContactIndex = INDEX_NONE;
-
-			OutCollisionDataFound.Add(FinderData);
-		}
+		return *MidPhases;
 	}
+	return TConstArrayView<TSharedPtr<FChaosVDParticlePairMidPhase>>();
 }
 
 bool AChaosVDParticleActor::HasCollisionData()
@@ -430,7 +421,7 @@ bool AChaosVDParticleActor::HasCollisionData()
 
 FName AChaosVDParticleActor::GetProviderName()
 {
-	return GetFName();
+	return ParticleDataPtr ? FName(ParticleDataPtr->DebugName) : NAME_None;
 }
 
 void AChaosVDParticleActor::UpdateMeshInstancesSelectionState()
@@ -447,20 +438,20 @@ void AChaosVDParticleActor::UpdateMeshInstancesSelectionState()
 
 void AChaosVDParticleActor::GetCharacterGroundConstraintData(TArray<TSharedPtr<FChaosVDCharacterGroundConstraint>>& OutConstraintsFound)
 {
-	if (const TArray<TSharedPtr<FChaosVDCharacterGroundConstraint>>* Constraints = GetCharacterGroundConstraintArray())
+	if (const TArray<TSharedPtr<FChaosVDConstraintDataWrapperBase>>* Constraints = GetCharacterGroundConstraintArray())
 	{
 		OutConstraintsFound.Reserve(Constraints->Num());
 
-		for (const TSharedPtr<FChaosVDCharacterGroundConstraint>& Constraint : *Constraints)
+		for (const TSharedPtr<FChaosVDConstraintDataWrapperBase>& Constraint : *Constraints)
 		{
-			OutConstraintsFound.Add(Constraint);
+			OutConstraintsFound.Add(StaticCastSharedPtr<FChaosVDCharacterGroundConstraint>(Constraint));
 		}
 	}
 }
 
 bool AChaosVDParticleActor::HasCharacterGroundConstraintData()
 {
-	if (const TArray<TSharedPtr<FChaosVDCharacterGroundConstraint>>* Constraints = GetCharacterGroundConstraintArray())
+	if (const TArray<TSharedPtr<FChaosVDConstraintDataWrapperBase>>* Constraints = GetCharacterGroundConstraintArray())
 	{
 		return Constraints->Num() > 0;
 	}
@@ -498,7 +489,7 @@ const TArray<TSharedPtr<FChaosVDParticlePairMidPhase>>* AChaosVDParticleActor::G
 	return nullptr;
 }
 
-const TArray<TSharedPtr<FChaosVDCharacterGroundConstraint>>* AChaosVDParticleActor::GetCharacterGroundConstraintArray() const
+const TArray<TSharedPtr<FChaosVDConstraintDataWrapperBase>>* AChaosVDParticleActor::GetCharacterGroundConstraintArray() const
 {
 	if (!ParticleDataPtr.IsValid())
 	{
@@ -515,7 +506,7 @@ const TArray<TSharedPtr<FChaosVDCharacterGroundConstraint>>* AChaosVDParticleAct
 	{
 		if (const UChaosVDSolverCharacterGroundConstraintDataComponent* ConstraintDataComponent = SolverInfoActor->GetCharacterGroundConstraintDataComponent())
 		{
-			return ConstraintDataComponent->GetConstraintsForParticle(ParticleDataPtr->ParticleIndex);
+			return ConstraintDataComponent->GetConstraintsForParticle(ParticleDataPtr->ParticleIndex, EChaosVDParticlePairSlot::Primary);
 		}
 	}
 
diff --git a/Engine/Plugins/ChaosVD/Source/ChaosVD/Private/ChaosVDParticleActor.h b/Engine/Plugins/ChaosVD/Source/ChaosVD/Private/ChaosVDParticleActor.h
--- a/Engine/Plugins/ChaosVD/Source/ChaosVD/Private/ChaosVDParticleActor.h
+++ b/Engine/Plugins/ChaosVD/Source/ChaosVD/Private/ChaosVDParticleActor.h
@@ -71,7 +71,7 @@ public:
 
 	virtual void Destroyed() override;
 
-	virtual const FChaosVDParticleDataWrapper* GetParticleData() override { return ParticleDataPtr.Get(); }
+	virtual TSharedPtr<const FChaosVDParticleDataWrapper> GetParticleData() override { return ParticleDataPtr; }
 	
 #if WITH_EDITOR
 	virtual bool IsSelectedInEditor() const override;
@@ -103,7 +103,7 @@ public:
 	virtual FBox GetComponentsBoundingBox(bool bNonColliding, bool bIncludeFromChildActors) const override;
 
 	// BEGIN IChaosVDCollisionDataProvider Interface
-	virtual void GetCollisionData(TArray<TSharedPtr<FChaosVDCollisionDataFinder>>& OutCollisionDataFound) override;
+	virtual TConstArrayView<TSharedPtr<FChaosVDParticlePairMidPhase>> GetCollisionData() override;
 	virtual bool HasCollisionData() override;
 	virtual FName GetProviderName() override;
 	// END IChaosVDCollisionDataProvider Interface
@@ -141,7 +141,7 @@ protected:
 
 	const TArray<TSharedPtr<FChaosVDParticlePairMidPhase>>* GetCollisionMidPhasesArray() const;
 
-	const TArray<TSharedPtr<FChaosVDCharacterGroundConstraint>>* GetCharacterGroundConstraintArray() const;
+	const TArray<TSharedPtr<FChaosVDConstraintDataWrapperBase>>* GetCharacterGroundConstraintArray() const;
 
 	void UpdateShapeDataComponents();
 
diff --git a/Engine/Plugins/ChaosVD/Source/ChaosVD/Private/ChaosVDParticleActorCustomization.cpp b/Engine/Plugins/ChaosVD/Source/ChaosVD/Private/ChaosVDParticleActorCustomization.cpp
--- a/Engine/Plugins/ChaosVD/Source/ChaosVD/Private/ChaosVDParticleActorCustomization.cpp
+++ b/Engine/Plugins/ChaosVD/Source/ChaosVD/Private/ChaosVDParticleActorCustomization.cpp
@@ -110,7 +110,7 @@ void FChaosVDParticleActorCustomization::HandleParticleDataUpdated()
 	}
 	else
 	{
-		const FChaosVDParticleDataWrapper* ParticleDataPtr = ParticleActor->GetParticleData();
+		TSharedPtr<const FChaosVDParticleDataWrapper> ParticleDataPtr = ParticleActor->GetParticleData();
 		CachedParticleData = ParticleDataPtr ? *ParticleDataPtr : FChaosVDParticleDataWrapper();
 	}
 }
diff --git a/Engine/Plugins/ChaosVD/Source/ChaosVD/Private/ChaosVDScene.cpp b/Engine/Plugins/ChaosVD/Source/ChaosVD/Private/ChaosVDScene.cpp
--- a/Engine/Plugins/ChaosVD/Source/ChaosVD/Private/ChaosVDScene.cpp
+++ b/Engine/Plugins/ChaosVD/Source/ChaosVD/Private/ChaosVDScene.cpp
@@ -22,17 +22,13 @@
 #include "Engine/StreamableManager.h"
 #include "Engine/World.h"
 #include "Elements/Actor/ActorElementData.h"
-#include "Elements/Component/ComponentElementData.h"
 #include "Elements/Framework/EngineElementsLibrary.h"
 #include "Elements/Framework/TypedElementSelectionSet.h"
-#include "Elements/Object/ObjectElementData.h"
 #include "Materials/Material.h"
 #include "Misc/ScopedSlowTask.h"
 #include "Selection.h"
-#include "Settings/ChaosVDParticleVisualizationSettings.h"
 #include "Settings/ChaosVDCoreSettings.h"
 #include "UObject/Package.h"
-#include "WorldPersistentFolders.h"
 #include "Actors/ChaosVDGeometryContainer.h"
 
 #define LOCTEXT_NAMESPACE "ChaosVisualDebugger"
@@ -225,7 +221,7 @@ void FChaosVDScene::UpdateFromRecordedStepData(const int32 SolverID, const FChao
 		if (bShouldDestroyParticleAnyway || InFrameData.ParticlesDestroyedIDs.Contains(ParticleActorWithID.Key))
 		{
 			// In large maps moving at high speed (like when moving on a vehicle), level streaming adds/removes hundreds of actors (and therefore particles) constantly.
-			// Destroying particle actors is expensive, specially if we need to spawn them again sooner as we will nee to rebuild-them.
+			// Destroying particle actors is expensive, specially if we need to spawn them again sooner as we will need to rebuild-them.
 			// So, we deactivate them instead.
 
 			// TODO: We need an actor pool system, so we can keep memory under control as well.
@@ -691,6 +687,8 @@ void FChaosVDScene::InitializeSelectionSets()
 	FString ObjectSelectionObjectName = FString::Printf(TEXT("CVDSelectedObjects-%s"), *FGuid::NewGuid().ToString());
 	ObjectSelection = USelection::CreateObjectSelection(GetTransientPackage(), *ObjectSelectionObjectName, RF_Transactional);
 	ObjectSelection->SetElementSelectionSet(SelectionSet);
+
+	SolverDataSelectionObject = MakeShared<FChaosVDSolverDataSelection>();
 }
 
 void FChaosVDScene::DeInitializeSelectionSets()
diff --git a/Engine/Plugins/ChaosVD/Source/ChaosVD/Private/ChaosVDScene.h b/Engine/Plugins/ChaosVD/Source/ChaosVD/Private/ChaosVDScene.h
--- a/Engine/Plugins/ChaosVD/Source/ChaosVD/Private/ChaosVDScene.h
+++ b/Engine/Plugins/ChaosVD/Source/ChaosVD/Private/ChaosVDScene.h
@@ -3,6 +3,8 @@
 #pragma once
 
 #include "ChaosVDRecording.h"
+#include "ChaosVDSolverDataSelection.h"
+#include "Components/ChaosVDSolverDataComponent.h"
 #include "Containers/UnrealString.h"
 #include "Containers/Map.h"
 #include "Elements/Framework/TypedElementListFwd.h"
@@ -141,6 +143,8 @@ public:
 	/** Updates the render state of the hit proxies of an array of actors. This used to update the selection outline state */
 	void UpdateSelectionProxiesForActors(TArrayView<AActor*> SelectedActors);
 
+	TWeakPtr<FChaosVDSolverDataSelection> GetSolverDataSelectionObject() { return SolverDataSelectionObject ? SolverDataSelectionObject : nullptr;}
+
 	TSharedPtr<FChaosVDRecording> LoadedRecording;
 
 private:
@@ -222,6 +226,8 @@ private:
 	FChaosVDSolverVisibilityChangedDelegate SolverVisibilityChangedDelegate;
 
 	TMap<uint32, TArray<IChaosVDGeometryOwnerInterface*>> ObjectsWaitingForGeometry;
+	
+	TSharedPtr<FChaosVDSolverDataSelection> SolverDataSelectionObject;
 
 	friend FChaosVDSelectionCustomization;
 };
diff --git a/Engine/Plugins/ChaosVD/Source/ChaosVD/Private/ChaosVDSolverDataSelection.cpp b/Engine/Plugins/ChaosVD/Source/ChaosVD/Private/ChaosVDSolverDataSelection.cpp
new file mode 100644
--- /dev/null
+++ b/Engine/Plugins/ChaosVD/Source/ChaosVD/Private/ChaosVDSolverDataSelection.cpp
@@ -0,0 +1,42 @@
+// Copyright Epic Games, Inc. All Rights Reserved.
+
+#include "ChaosVDSolverDataSelection.h"
+
+void FChaosVDSolverDataSelectionHandle::SetOwner(const TSharedPtr<FChaosVDSolverDataSelection>& InOwner)
+{
+	Owner = InOwner;
+}
+
+bool FChaosVDSolverDataSelectionHandle::IsSelected()
+{
+	if (TSharedPtr<FChaosVDSolverDataSelection> OwnerPtr = Owner.Pin())
+	{
+		return OwnerPtr->IsSelectionHandleSelected(AsShared());
+	}
+
+	return false;
+}
+
+bool FChaosVDSolverDataSelectionHandle::IsValid() const
+{
+	return SelectedDataStruct.IsValid() && SelectedDataStruct->IsValid();
+}
+
+void FChaosVDSolverDataSelection::SelectData(const TSharedPtr<FChaosVDSolverDataSelectionHandle>& InSelectionHandle)
+{
+	CurrentSelectedSolverDataHandle = InSelectionHandle;
+
+	SolverDataSelectionChangeDelegate.Broadcast(CurrentSelectedSolverDataHandle);
+}
+
+bool FChaosVDSolverDataSelection::IsSelectionHandleSelected(const TSharedPtr<FChaosVDSolverDataSelectionHandle>& InSelectionHandle) const
+{
+	const bool bBothHandlesAreValid = CurrentSelectedSolverDataHandle && InSelectionHandle;
+
+	if (!bBothHandlesAreValid)
+	{
+		return false;
+	}
+
+	return (*CurrentSelectedSolverDataHandle) == (*InSelectionHandle);
+}
diff --git a/Engine/Plugins/ChaosVD/Source/ChaosVD/Private/ChaosVDSolverDataSelection.h b/Engine/Plugins/ChaosVD/Source/ChaosVD/Private/ChaosVDSolverDataSelection.h
new file mode 100644
--- /dev/null
+++ b/Engine/Plugins/ChaosVD/Source/ChaosVD/Private/ChaosVDSolverDataSelection.h
@@ -0,0 +1,185 @@
+// Copyright Epic Games, Inc. All Rights Reserved.
+
+#pragma once
+#include "Templates/SharedPointer.h"
+#include "UObject/StructOnScope.h"
+#include "ChaosVDSolverDataSelection.generated.h"
+
+class FChaosVDSolverDataSelection;
+
+USTRUCT()
+struct FChaosVDSelectionContext
+{
+	GENERATED_BODY()
+};
+
+struct FChaosVDSolverDataSelectionHandle : public TSharedFromThis<FChaosVDSolverDataSelectionHandle>
+{
+	virtual ~FChaosVDSolverDataSelectionHandle() = default;
+	
+	template<typename DataStructType>
+	void SetHandleData(const TSharedPtr<DataStructType>& Data);
+	
+	template<typename ContextDataStructType>
+	void SetHandleContext(ContextDataStructType&& ContextData);
+
+	void SetOwner(const TSharedPtr<FChaosVDSolverDataSelection>& InOwner);
+
+	virtual bool IsSelected();
+
+	bool IsValid() const;
+
+	template<typename DataStructType>
+	bool IsA() const;
+
+	template<typename DataStructType>
+	DataStructType* GetData() const;
+
+	template<typename DataStructType>
+	TSharedPtr<DataStructType> GetDataAsShared() const;
+
+	template<typename ContextDataStructType>
+	ContextDataStructType* GetContextData() const;
+
+	TSharedPtr<FStructOnScope> GetDataAsStructScope() { return SelectedDataStruct; }
+
+	bool operator==(const FChaosVDSolverDataSelectionHandle& Other) const
+	{
+		return DataSharedPtr.Get() == Other.DataSharedPtr.Get();
+	}
+
+	bool operator!=(const FChaosVDSolverDataSelectionHandle& Other) const
+	{
+		return DataSharedPtr.Get() != Other.DataSharedPtr.Get();
+	}
+
+private:
+
+	template<typename DataStructType>
+	bool IsA_Internal(const TSharedPtr<FStructOnScope>& InStructOnScope) const;
+
+	TSharedPtr<FStructOnScope> SelectedDataStruct;
+	TSharedPtr<FStructOnScope> SelectedDataContext;
+
+	TSharedPtr<void> DataSharedPtr;
+	TSharedPtr<void> SelectedDataContextSharedPtr;
+
+	TWeakPtr<FChaosVDSolverDataSelection> Owner;
+};
+
+template <typename DataStructType>
+void FChaosVDSolverDataSelectionHandle::SetHandleData(const TSharedPtr<DataStructType>& Data)
+{
+	if (Data)
+	{
+		SelectedDataStruct = MakeShared<FStructOnScope>(DataStructType::StaticStruct(), reinterpret_cast<uint8*>(Data.Get()));
+		DataSharedPtr = Data;
+	}
+	else
+	{
+		SelectedDataStruct->Reset();
+		DataSharedPtr = nullptr;
+	}
+}
+
+template <typename ContextDataStructType>
+void FChaosVDSolverDataSelectionHandle::SetHandleContext(ContextDataStructType&& ContextData)
+{
+	SelectedDataContextSharedPtr = MakeShared<ContextDataStructType>();
+	ContextDataStructType* ContextDataRawPtr = StaticCastSharedPtr<ContextDataStructType>(SelectedDataContextSharedPtr).Get();
+	*ContextDataRawPtr = MoveTemp(ContextData);
+
+	SelectedDataContext = MakeShared<FStructOnScope>(ContextDataStructType::StaticStruct(), reinterpret_cast<uint8*>(ContextDataRawPtr));		
+}
+
+template <typename DataStructType>
+bool FChaosVDSolverDataSelectionHandle::IsA() const
+{
+	return IsA_Internal<DataStructType>(SelectedDataStruct);
+}
+
+template <typename DataStructType>
+DataStructType* FChaosVDSolverDataSelectionHandle::GetData() const
+{
+	if (IsA_Internal<DataStructType>(SelectedDataStruct))
+	{
+		return reinterpret_cast<DataStructType*>(SelectedDataStruct->GetStructMemory());
+	}
+	
+	return nullptr;
+}
+
+template <typename DataStructType>
+TSharedPtr<DataStructType> FChaosVDSolverDataSelectionHandle::GetDataAsShared() const
+{
+	if (IsA_Internal<DataStructType>(SelectedDataStruct))
+	{
+		return StaticCastSharedPtr<DataStructType>(DataSharedPtr);
+	}
+	return nullptr;
+}
+
+template <typename ContextDataStructType>
+ContextDataStructType* FChaosVDSolverDataSelectionHandle::GetContextData() const
+{
+	if (IsA_Internal<ContextDataStructType>(SelectedDataContext))
+	{
+		return reinterpret_cast<ContextDataStructType*>(SelectedDataContext->GetStructMemory());
+	}
+
+	return nullptr;
+}
+
+template <typename DataStructType>
+bool FChaosVDSolverDataSelectionHandle::IsA_Internal(const TSharedPtr<FStructOnScope>& InStructOnScope) const
+{
+	if (InStructOnScope)
+	{
+		const UStruct* HandleStruct = InStructOnScope->GetStruct();
+		return HandleStruct && (DataStructType::StaticStruct() == HandleStruct || HandleStruct->IsChildOf(DataStructType::StaticStruct()));
+	}
+
+	return false;
+}
+
+
+DECLARE_MULTICAST_DELEGATE_OneParam(FChaosVDSolverDataSelectionChangedDelegate, const TSharedPtr<FChaosVDSolverDataSelectionHandle>& SelectionHandle)
+
+class FChaosVDSolverDataSelection : public TSharedFromThis<FChaosVDSolverDataSelection>
+{
+public:
+
+	void SelectData(const TSharedPtr<FChaosVDSolverDataSelectionHandle>& InSelectionHandle);
+
+	template<typename SolverDataType>
+	TSharedPtr<FChaosVDSolverDataSelectionHandle> MakeSelectionHandle(const TSharedPtr<SolverDataType>& InSolverData);
+
+	FChaosVDSolverDataSelectionChangedDelegate& GetDataSelectionChangedDelegate() { return SolverDataSelectionChangeDelegate; }
+	TSharedPtr<FChaosVDSolverDataSelectionHandle> GetCurrentSelectionHandle() { return CurrentSelectedSolverDataHandle; }
+
+	template <typename SolverDataType>
+	bool IsDataSelected(const TSharedPtr<SolverDataType>& InSolverData);
+	
+	bool IsSelectionHandleSelected(const TSharedPtr<FChaosVDSolverDataSelectionHandle>& InSelectionHandle) const;
+
+private:
+	FChaosVDSolverDataSelectionChangedDelegate SolverDataSelectionChangeDelegate;
+	TSharedPtr<FChaosVDSolverDataSelectionHandle> CurrentSelectedSolverDataHandle;
+};
+
+template <typename SolverDataType>
+TSharedPtr<FChaosVDSolverDataSelectionHandle> FChaosVDSolverDataSelection::MakeSelectionHandle(const TSharedPtr<SolverDataType>& InSolverData)
+{
+	TSharedPtr<FChaosVDSolverDataSelectionHandle> NewSelectionHandle = MakeShared<FChaosVDSolverDataSelectionHandle>();
+	NewSelectionHandle->SetHandleData(StaticCastSharedPtr<SolverDataType>(InSolverData));
+	NewSelectionHandle->SetOwner(AsShared());
+
+	return NewSelectionHandle;
+}
+
+template <typename SolverDataType>
+bool FChaosVDSolverDataSelection::IsDataSelected(const TSharedPtr<SolverDataType>& InSolverData)
+{
+	TSharedPtr<FChaosVDSolverDataSelectionHandle> SelectionHandle = MakeSelectionHandle(InSolverData);
+	return SelectionHandle->IsSelected();
+}
\ No newline at end of file
diff --git a/Engine/Plugins/ChaosVD/Source/ChaosVD/Private/ChaosVDTabsIDs.cpp b/Engine/Plugins/ChaosVD/Source/ChaosVD/Private/ChaosVDTabsIDs.cpp
--- a/Engine/Plugins/ChaosVD/Source/ChaosVD/Private/ChaosVDTabsIDs.cpp
+++ b/Engine/Plugins/ChaosVD/Source/ChaosVD/Private/ChaosVDTabsIDs.cpp
@@ -12,5 +12,4 @@ const FName FChaosVDTabID::SolversTrack = TEXT("SolversTrack");
 const FName FChaosVDTabID::StatusBar = TEXT("StatusBar");
 const FName FChaosVDTabID::CollisionDataDetails = TEXT("CollisionDataDetails");
 const FName FChaosVDTabID::SceneQueryDataDetails = TEXT("SceneQueryDataDetails");
-const FName FChaosVDTabID::JointsDataDetails = TEXT("JointsDataDetails");
-const FName FChaosVDTabID::CharacterGroundConstraintDataDetails = TEXT("CharacterGroundConstraintDataDetails");
\ No newline at end of file
+const FName FChaosVDTabID::ConstraintsInspector = TEXT("ConstraintsDetails");
\ No newline at end of file
diff --git a/Engine/Plugins/ChaosVD/Source/ChaosVD/Private/ChaosVDTabsIDs.h b/Engine/Plugins/ChaosVD/Source/ChaosVD/Private/ChaosVDTabsIDs.h
--- a/Engine/Plugins/ChaosVD/Source/ChaosVD/Private/ChaosVDTabsIDs.h
+++ b/Engine/Plugins/ChaosVD/Source/ChaosVD/Private/ChaosVDTabsIDs.h
@@ -17,6 +17,5 @@ public:
 	static const FName StatusBar;
 	static const FName CollisionDataDetails;
 	static const FName SceneQueryDataDetails;
-	static const FName JointsDataDetails;
-	static const FName CharacterGroundConstraintDataDetails;
+	static const FName ConstraintsInspector;
 };
diff --git a/Engine/Plugins/ChaosVD/Source/ChaosVD/Private/ChaosVDWorldOutlinerMode.cpp b/Engine/Plugins/ChaosVD/Source/ChaosVD/Private/ChaosVDWorldOutlinerMode.cpp
--- a/Engine/Plugins/ChaosVD/Source/ChaosVD/Private/ChaosVDWorldOutlinerMode.cpp
+++ b/Engine/Plugins/ChaosVD/Source/ChaosVD/Private/ChaosVDWorldOutlinerMode.cpp
@@ -94,7 +94,7 @@ void FChaosVDActorTreeItem::UpdateDisplayString()
 {
 	if (AChaosVDParticleActor* CVDActor = Cast<AChaosVDParticleActor>(Actor.Get()))
 	{
-		if (const FChaosVDParticleDataWrapper* ParticleData = CVDActor->GetParticleData())
+		if (TSharedPtr<const FChaosVDParticleDataWrapper> ParticleData = CVDActor->GetParticleData())
 		{
 			const bool bHasDebugName = !ParticleData->DebugName.IsEmpty();
 			DisplayString = bHasDebugName ? ParticleData->DebugName : TEXT("Unnamed Particle - ID : ") + FString::FromInt(ParticleData->ParticleIndex);
diff --git a/Engine/Plugins/ChaosVD/Source/ChaosVD/Private/Components/ChaosVDConstraintDataComponent.cpp b/Engine/Plugins/ChaosVD/Source/ChaosVD/Private/Components/ChaosVDConstraintDataComponent.cpp
new file mode 100644
--- /dev/null
+++ b/Engine/Plugins/ChaosVD/Source/ChaosVD/Private/Components/ChaosVDConstraintDataComponent.cpp
@@ -0,0 +1,36 @@
+// Copyright Epic Games, Inc. All Rights Reserved.
+
+#include "ChaosVDConstraintDataComponent.h"
+
+#include "ChaosVDConstraintDataHelpers.h"
+
+UChaosVDConstraintDataComponent::UChaosVDConstraintDataComponent()
+{
+	SetCanEverAffectNavigation(false);
+	bNavigationRelevant = false;
+	PrimaryComponentTick.bCanEverTick = false;
+}
+
+const FChaosVDConstraintDataArray* UChaosVDConstraintDataComponent::GetConstraintsForParticle(int32 ParticleID, EChaosVDParticlePairSlot Options) const
+{
+	return Chaos::VisualDebugger::Utils::GetDataFromParticlePairMaps<FChaosVDConstraintDataByParticleMap, TSharedPtr<FChaosVDConstraintDataWrapperBase>>(ConstraintByParticle0, ConstraintByParticle1, ParticleID, Options);
+}
+
+TSharedPtr<FChaosVDConstraintDataWrapperBase> UChaosVDConstraintDataComponent::GetConstraintByIndex(int32 ConstraintIndex)
+{
+	if (TSharedPtr<FChaosVDConstraintDataWrapperBase>* ConstraintPtrPtr = ConstraintByConstraintIndex.Find(ConstraintIndex))
+	{
+		return *ConstraintPtrPtr;
+	}
+
+	return nullptr;
+}
+
+void UChaosVDConstraintDataComponent::ClearData()
+{
+	AllConstraints.Empty();
+	ConstraintByParticle0.Empty();
+	ConstraintByParticle1.Empty();
+	ConstraintByConstraintIndex.Empty();
+}
+
diff --git a/Engine/Plugins/ChaosVD/Source/ChaosVD/Private/Components/ChaosVDConstraintDataComponent.h b/Engine/Plugins/ChaosVD/Source/ChaosVD/Private/Components/ChaosVDConstraintDataComponent.h
new file mode 100644
--- /dev/null
+++ b/Engine/Plugins/ChaosVD/Source/ChaosVD/Private/Components/ChaosVDConstraintDataComponent.h
@@ -0,0 +1,66 @@
+// Copyright Epic Games, Inc. All Rights Reserved.
+
+#pragma once
+
+#include "Components/ChaosVDConstraintDataHelpers.h"
+#include "ChaosVDSolverDataComponent.h"
+#include "DataWrappers/ChaosVDParticleDataWrapper.h"
+#include "ChaosVDConstraintDataComponent.generated.h"
+
+struct FChaosVDJointConstraint;
+typedef TMap<int32, TArray<TSharedPtr<FChaosVDConstraintDataWrapperBase>>> FChaosVDConstraintDataByParticleMap;
+typedef TMap<int32, TSharedPtr<FChaosVDConstraintDataWrapperBase>> FChaosVDConstraintDataByConstraintIndexMap;
+typedef TArray<TSharedPtr<FChaosVDConstraintDataWrapperBase>> FChaosVDConstraintDataArray;
+
+UCLASS()
+class UChaosVDConstraintDataComponent : public UChaosVDSolverDataComponent
+{
+	GENERATED_BODY()
+
+public:
+	UChaosVDConstraintDataComponent();
+
+	template<typename ConstraintDataType>
+	void UpdateConstraintData(const TArray<TSharedPtr<ConstraintDataType>>& InData);
+	
+	const FChaosVDConstraintDataArray& GetAllConstraints() const { return AllConstraints; }
+	const FChaosVDConstraintDataArray* GetConstraintsForParticle(int32 ParticleID, EChaosVDParticlePairSlot Options) const;
+
+	virtual void ClearData() override;
+
+protected:
+
+	TSharedPtr<FChaosVDConstraintDataWrapperBase> GetConstraintByIndex(int32 ConstraintIndex);
+
+	FChaosVDConstraintDataArray AllConstraints;
+	
+	FChaosVDConstraintDataByParticleMap ConstraintByParticle0;
+	FChaosVDConstraintDataByParticleMap ConstraintByParticle1;
+
+	FChaosVDConstraintDataByConstraintIndexMap ConstraintByConstraintIndex;
+
+};
+
+template <typename ConstraintDataType>
+void UChaosVDConstraintDataComponent::UpdateConstraintData(const TArray<TSharedPtr<ConstraintDataType>>& InData)
+{
+	ClearData();
+
+	AllConstraints.Reserve(InData.Num());
+	ConstraintByParticle0.Reserve(InData.Num());
+	ConstraintByParticle1.Reserve(InData.Num());
+	ConstraintByConstraintIndex.Reserve(InData.Num());
+
+	for (const TSharedPtr<FChaosVDConstraintDataWrapperBase>& Constraint : InData)
+	{
+		if (!Constraint)
+		{
+			continue;
+		}
+
+		AllConstraints.Emplace(Constraint);
+		ConstraintByConstraintIndex.Add(Constraint->GetConstraintIndex(), Constraint);
+		Chaos::VisualDebugger::Utils::AddDataDataToParticleIDMap(ConstraintByParticle0, Constraint, Constraint->GetParticleIDAtSlot(EChaosVDParticlePairIndex::Index_0));
+		Chaos::VisualDebugger::Utils::AddDataDataToParticleIDMap(ConstraintByParticle1, Constraint, Constraint->GetParticleIDAtSlot(EChaosVDParticlePairIndex::Index_1));
+	}
+}
diff --git a/Engine/Plugins/ChaosVD/Source/ChaosVD/Private/Components/ChaosVDSceneQueryDataComponent.cpp b/Engine/Plugins/ChaosVD/Source/ChaosVD/Private/Components/ChaosVDSceneQueryDataComponent.cpp
--- a/Engine/Plugins/ChaosVD/Source/ChaosVD/Private/Components/ChaosVDSceneQueryDataComponent.cpp
+++ b/Engine/Plugins/ChaosVD/Source/ChaosVD/Private/Components/ChaosVDSceneQueryDataComponent.cpp
@@ -4,38 +4,7 @@
 
 #include "ChaosVDRecording.h"
 
-void FChaosVDSceneQuerySelectionHandle::SetIsSelected(bool bNewSelected)
-{
-	if (const TSharedPtr<FChaosVDQueryDataWrapper> QueryDataPtr = QueryData.Pin())
-	{
-		QueryDataPtr->bIsSelectedInEditor = bNewSelected;
-		if (QueryDataPtr && QueryDataPtr->SQVisitData.IsValidIndex(SQVisitIndex))
-		{
-			QueryDataPtr->SQVisitData[SQVisitIndex].bIsSelectedInEditor = bNewSelected;
-		}
-		else if (QueryDataPtr)
-		{
-			// If we end up with a invalid index, make sure that no visit is selected
-			SQVisitIndex = INDEX_NONE;
-			for (FChaosVDQueryVisitStep& VisitData : QueryDataPtr->SQVisitData)
-			{
-				VisitData.bIsSelectedInEditor = false;
-			}
-		}
-	}
-}
-
-bool FChaosVDSceneQuerySelectionHandle::IsSelected() const
-{
-	if (const TSharedPtr<FChaosVDQueryDataWrapper> QueryDataPtr = QueryData.Pin())
-	{
-		return QueryDataPtr->bIsSelectedInEditor;
-	}
-
-	return false;
-}
-
-UChaosVDSceneQueryDataComponent::UChaosVDSceneQueryDataComponent() : CurrentSQSelectionHandle({nullptr, INDEX_NONE})
+UChaosVDSceneQueryDataComponent::UChaosVDSceneQueryDataComponent()
 {
 	PrimaryComponentTick.bCanEverTick = false;
 	
@@ -51,9 +20,6 @@ void UChaosVDSceneQueryDataComponent::UpdateQueriesFromFrameData(const FChaosVDG
 	RecordedQueriesByID.Empty(RecordedQueriesNum);
 	RecordedQueries.Empty(RecordedQueriesNum);
 
-	// Until we have a way to track and auto-select new queries instances between frames, just clear the selection
-	SelectQuery(FChaosVDSceneQuerySelectionHandle());
-
 	for (const TPair<int32, TSharedPtr<FChaosVDQueryDataWrapper>>& QueryIDPair : InGameFrameData.RecordedSceneQueries)
 	{
 		if (TSharedPtr<FChaosVDQueryDataWrapper> QueryData = QueryIDPair.Value)
@@ -92,39 +58,9 @@ TSharedPtr<FChaosVDQueryDataWrapper> UChaosVDSceneQueryDataComponent::GetQueryBy
 	return nullptr;
 }
 
-void UChaosVDSceneQueryDataComponent::SelectQuery(int32 QueryID)
-{
-	CurrentSQSelectionHandle.SetIsSelected(false);
-	CurrentSQSelectionHandle = FChaosVDSceneQuerySelectionHandle(GetQueryByID(QueryID), 0);
-	CurrentSQSelectionHandle.SetIsSelected(true);
-	
-	SelectionChangeDelegate.Broadcast(CurrentSQSelectionHandle);
-}
-
-void UChaosVDSceneQueryDataComponent::SelectQuery(const FChaosVDSceneQuerySelectionHandle& SelectionHandle)
-{
-	CurrentSQSelectionHandle.SetIsSelected(false);
-	CurrentSQSelectionHandle = SelectionHandle;
-	CurrentSQSelectionHandle.SetIsSelected(true);
-	
-	SelectionChangeDelegate.Broadcast(CurrentSQSelectionHandle);
-}
-
-bool UChaosVDSceneQueryDataComponent::IsQuerySelected(int32 QueryID) const
-{
-	if (const TSharedPtr<FChaosVDQueryDataWrapper> QueryDataPtr = CurrentSQSelectionHandle.GetQueryData().Pin())
-	{
-		return QueryDataPtr->ID == QueryID;
-	}
-
-	return false;
-}
-
 void UChaosVDSceneQueryDataComponent::ClearData()
 {
 	RecordedQueriesByType.Reset();
 	RecordedQueriesByID.Reset();
 	RecordedQueries.Reset();
-	
-	CurrentSQSelectionHandle = FChaosVDSceneQuerySelectionHandle();
 }
diff --git a/Engine/Plugins/ChaosVD/Source/ChaosVD/Private/Components/ChaosVDSceneQueryDataComponent.h b/Engine/Plugins/ChaosVD/Source/ChaosVD/Private/Components/ChaosVDSceneQueryDataComponent.h
--- a/Engine/Plugins/ChaosVD/Source/ChaosVD/Private/Components/ChaosVDSceneQueryDataComponent.h
+++ b/Engine/Plugins/ChaosVD/Source/ChaosVD/Private/Components/ChaosVDSceneQueryDataComponent.h
@@ -3,47 +3,16 @@
 #pragma once
 
 #include "DataWrappers/ChaosVDQueryDataWrappers.h"
-#include "SolverDataComponent.h"
+#include "ChaosVDSolverDataComponent.h"
+#include "ChaosVDSolverDataSelection.h"
 
 #include "ChaosVDSceneQueryDataComponent.generated.h"
 
 struct FChaosVDGameFrameData;
 
-/** Struct used to pass data about a specific query to other objects */
-struct FChaosVDSceneQuerySelectionHandle
-{
-	FChaosVDSceneQuerySelectionHandle()
-	{		
-	}
-
-	FChaosVDSceneQuerySelectionHandle(const TWeakPtr<FChaosVDQueryDataWrapper>& InQueryData, int32 SQVisitIndex)
-		: QueryData(InQueryData)
-		, SQVisitIndex(SQVisitIndex)
-	{
-		if (const TSharedPtr<FChaosVDQueryDataWrapper> QueryDataPtr = QueryData.Pin())
-		{
-			QueryID = QueryDataPtr->ID;
-		}
-	}
-
-	void SetIsSelected(bool bNewSelected);
-	bool IsSelected() const;
-
-	TWeakPtr<FChaosVDQueryDataWrapper> GetQueryData() const { return QueryData; } 
-	int32 GetQueryID() const { return QueryID; } 
-	int32 GetSQVisitIndex() const { return SQVisitIndex; } 
-
-private:
-	TWeakPtr<FChaosVDQueryDataWrapper> QueryData = nullptr;
-	int32 SQVisitIndex = INDEX_NONE;
-	int32 QueryID = INDEX_NONE;
-};
-
-DECLARE_MULTICAST_DELEGATE_OneParam(FChaosVDSQSelectionChangedDelegate, const FChaosVDSceneQuerySelectionHandle& SelectionHandle)
-
 /** Actor Component that contains all the scene queries recorded at the current loaded frame */
 UCLASS()
-class CHAOSVD_API UChaosVDSceneQueryDataComponent : public USolverDataComponent
+class UChaosVDSceneQueryDataComponent : public UChaosVDSolverDataComponent
 {
 	GENERATED_BODY()
 
@@ -56,15 +25,6 @@ public:
 	TConstArrayView<TSharedPtr<FChaosVDQueryDataWrapper>> GetAllQueries() const;
 	TSharedPtr<FChaosVDQueryDataWrapper> GetQueryByID(int32 QueryID) const;
 
-	void SelectQuery(int32 QueryID);
-	void SelectQuery(const FChaosVDSceneQuerySelectionHandle& SelectionHandle);
-
-	bool IsQuerySelected(int32 QueryID) const;
-
-	FChaosVDSceneQuerySelectionHandle GetSelectedQueryHandle() const { return CurrentSQSelectionHandle; }
-
-	FChaosVDSQSelectionChangedDelegate& GetOnSelectionChangeDelegate() { return SelectionChangeDelegate; }
-
 	virtual void ClearData() override;
 
 protected:
@@ -72,8 +32,4 @@ protected:
 	TMap<EChaosVDSceneQueryType, TArray<TSharedPtr<FChaosVDQueryDataWrapper>>> RecordedQueriesByType;
 	TMap<int32, TSharedPtr<FChaosVDQueryDataWrapper>> RecordedQueriesByID;
 	TArray<TSharedPtr<FChaosVDQueryDataWrapper>> RecordedQueries;
-	
-	FChaosVDSQSelectionChangedDelegate SelectionChangeDelegate;
-
-	FChaosVDSceneQuerySelectionHandle CurrentSQSelectionHandle;
 };
diff --git a/Engine/Plugins/ChaosVD/Source/ChaosVD/Private/Components/ChaosVDSolverCharacterGroundConstraintDataComponent.cpp b/Engine/Plugins/ChaosVD/Source/ChaosVD/Private/Components/ChaosVDSolverCharacterGroundConstraintDataComponent.cpp
--- a/Engine/Plugins/ChaosVD/Source/ChaosVD/Private/Components/ChaosVDSolverCharacterGroundConstraintDataComponent.cpp
+++ b/Engine/Plugins/ChaosVD/Source/ChaosVD/Private/Components/ChaosVDSolverCharacterGroundConstraintDataComponent.cpp
@@ -2,102 +2,76 @@
 
 #include "Components/ChaosVDSolverCharacterGroundConstraintDataComponent.h"
 
-#include "Actors/ChaosVDSolverInfoActor.h"
-#include "Components/ChaosVDConstraintDataHelpers.h"
-#include "DataWrappers/ChaosVDCharacterGroundConstraintDataWrappers.h"
+#include "ChaosVDCharacterGroundConstraintDataProviderInterface.h"
+#include "ChaosVDScene.h"
+#include "Selection.h"
+#include "Settings/ChaosVDCharacterConstraintsVisualizationSettings.h"
 
-void FChaosVDCharacterGroundConstraintSelectionHandle::SetIsSelected(bool bNewSelected)
+void UChaosVDSolverCharacterGroundConstraintDataComponent::HandleSceneUpdated()
 {
-	if (const TSharedPtr<FChaosVDCharacterGroundConstraint> ConstraintDataPtr = ConstraintData.Pin())
+	if (const UChaosVDCharacterConstraintsVisualizationSettings* CharacterConstraintsVisualizationSettings = GetDefault<UChaosVDCharacterConstraintsVisualizationSettings>())
 	{
-		ConstraintDataPtr->bIsSelectedInEditor = bNewSelected;
+		if (!CharacterConstraintsVisualizationSettings->bAutoSelectConstraintFromSelectedParticle)
+		{
+			return;
+		}
+
 	}
-}
 
-bool FChaosVDCharacterGroundConstraintSelectionHandle::IsSelected() const
-{
-	if (const TSharedPtr<FChaosVDCharacterGroundConstraint> ConstraintDataPtr = ConstraintData.Pin())
+	const TSharedPtr<FChaosVDScene> ScenePtr = SceneWeakPtr.Pin();
+	if (!ScenePtr)
 	{
-		return ConstraintDataPtr->bIsSelectedInEditor;
+		return;
+	}
+	
+	USelection* ActorSelection = ScenePtr->GetActorSelectionObject();
+	if (!ActorSelection)
+	{
+		return;
 	}
 
-	return false;
-}
-
-UChaosVDSolverCharacterGroundConstraintDataComponent::UChaosVDSolverCharacterGroundConstraintDataComponent()
-{
-	SetCanEverAffectNavigation(false);
-	bNavigationRelevant = false;
-	PrimaryComponentTick.bCanEverTick = false;
-}
-
-void UChaosVDSolverCharacterGroundConstraintDataComponent::UpdateConstraintData(const TArray<TSharedPtr<FChaosVDCharacterGroundConstraint>>& InConstraintData)
-{
-	ClearData();
-
-	AllConstraints = InConstraintData;
-
-	ConstraintByCharacterParticle.Reserve(InConstraintData.Num());
-	ConstraintByConstraintIndex.Reserve(InConstraintData.Num());
+	TSharedPtr<FChaosVDSolverDataSelection> SolverDataSelection = ScenePtr->GetSolverDataSelectionObject().Pin();
+	if(!SolverDataSelection)
+	{
+		return;
+	}
 
-	for (const TSharedPtr<FChaosVDCharacterGroundConstraint>& Constraint : InConstraintData)
+	if (ActorSelection->Num() > 0)
 	{
-		if (!Constraint)
+		if (IChaosVDCharacterGroundConstraintDataProviderInterface* DataProvider = Cast<IChaosVDCharacterGroundConstraintDataProviderInterface>(ActorSelection->GetSelectedObject(0)))
 		{
-			continue;
+			if (DataProvider->HasCharacterGroundConstraintData())
+			{
+				TArray<TSharedPtr<FChaosVDCharacterGroundConstraint>> FoundConstraintData;
+				DataProvider->GetCharacterGroundConstraintData(FoundConstraintData);
+				SolverDataSelection->SelectData(SolverDataSelection->MakeSelectionHandle(FoundConstraintData[0]));
+			}
 		}
-
-		ConstraintByConstraintIndex.Add(Constraint->ConstraintIndex, Constraint);
-		Chaos::VisualDebugger::Utils::AddDataDataToParticleIDMap(ConstraintByCharacterParticle, Constraint, Constraint->CharacterParticleIndex);
 	}
 }
 
-const FChaosVDCharacterGroundDataArray* UChaosVDSolverCharacterGroundConstraintDataComponent::GetConstraintsForParticle(int32 ParticleID) const
+void UChaosVDSolverCharacterGroundConstraintDataComponent::BeginDestroy()
 {
-	return ConstraintByCharacterParticle.Find(ParticleID);
-}
-
-void UChaosVDSolverCharacterGroundConstraintDataComponent::SelectConstraint(int32 ConstraintIndex)
-{
-	CurrentConstraintSelectionHandle.SetIsSelected(false);
-	CurrentConstraintSelectionHandle = FChaosVDCharacterGroundConstraintSelectionHandle(GetConstraintByIndex(ConstraintIndex));
-	CurrentConstraintSelectionHandle.SetIsSelected(true);
-
-	SelectionChangeDelegate.Broadcast(CurrentConstraintSelectionHandle);
-}
-
-void UChaosVDSolverCharacterGroundConstraintDataComponent::SelectConstraint(const FChaosVDCharacterGroundConstraintSelectionHandle& SelectionHandle)
-{
-	CurrentConstraintSelectionHandle.SetIsSelected(false);
-	CurrentConstraintSelectionHandle = SelectionHandle;
-	CurrentConstraintSelectionHandle.SetIsSelected(true);
-	
-	SelectionChangeDelegate.Broadcast(CurrentConstraintSelectionHandle);
-}
+	Super::BeginDestroy();
 
-bool UChaosVDSolverCharacterGroundConstraintDataComponent::IsConstraintSelected(int32 ConstraintIndex) const
-{
-	if (const TSharedPtr<FChaosVDCharacterGroundConstraint> ConstraintDataPtr = CurrentConstraintSelectionHandle.GetData().Pin())
+	const TSharedPtr<FChaosVDScene> ScenePtr = SceneWeakPtr.Pin();
+	if (!ScenePtr)
 	{
-		return ConstraintDataPtr->bIsSelectedInEditor && ConstraintDataPtr->ConstraintIndex == ConstraintIndex;
+		return;
 	}
-
-	return false;
+	
+	ScenePtr->OnSceneUpdated().RemoveAll(this);
 }
 
-TSharedPtr<FChaosVDCharacterGroundConstraint> UChaosVDSolverCharacterGroundConstraintDataComponent::GetConstraintByIndex(int32 ConstraintIndex)
+void UChaosVDSolverCharacterGroundConstraintDataComponent::SetScene(const TWeakPtr<FChaosVDScene>& InSceneWeakPtr)
 {
-	if (TSharedPtr<FChaosVDCharacterGroundConstraint>* ConstraintPtrPtr = ConstraintByConstraintIndex.Find(ConstraintIndex))
+	Super::SetScene(InSceneWeakPtr);
+	
+	const TSharedPtr<FChaosVDScene> ScenePtr = SceneWeakPtr.Pin();
+	if (!ScenePtr)
 	{
-		return *ConstraintPtrPtr;
+		return;
 	}
-
-	return nullptr;
-}
-
-void UChaosVDSolverCharacterGroundConstraintDataComponent::ClearData()
-{
-	AllConstraints.Empty();
-	ConstraintByCharacterParticle.Empty();
-	ConstraintByConstraintIndex.Empty();
+	
+	ScenePtr->OnSceneUpdated().AddUObject(this, &UChaosVDSolverCharacterGroundConstraintDataComponent::HandleSceneUpdated);
 }
diff --git a/Engine/Plugins/ChaosVD/Source/ChaosVD/Private/Components/ChaosVDSolverCharacterGroundConstraintDataComponent.h b/Engine/Plugins/ChaosVD/Source/ChaosVD/Private/Components/ChaosVDSolverCharacterGroundConstraintDataComponent.h
--- a/Engine/Plugins/ChaosVD/Source/ChaosVD/Private/Components/ChaosVDSolverCharacterGroundConstraintDataComponent.h
+++ b/Engine/Plugins/ChaosVD/Source/ChaosVD/Private/Components/ChaosVDSolverCharacterGroundConstraintDataComponent.h
@@ -2,76 +2,17 @@
 
 #pragma once
 
-#include "SolverDataComponent.h"
+#include "ChaosVDConstraintDataComponent.h"
 #include "ChaosVDSolverCharacterGroundConstraintDataComponent.generated.h"
 
-struct FChaosVDCharacterGroundConstraint;
-
-typedef TMap<int32, TArray<TSharedPtr<FChaosVDCharacterGroundConstraint>>> FChaosVDCharacterGroundDataByParticleMap;
-typedef TMap<int32, TSharedPtr<FChaosVDCharacterGroundConstraint>> FChaosVDCharacterGroundDataByConstraintIndexMap;
-typedef TArray<TSharedPtr<FChaosVDCharacterGroundConstraint>> FChaosVDCharacterGroundDataArray;
-
-enum class EChaosVDParticlePairSlot : uint8;
-
-/** Struct used to pass data about a specific character ground constraint to other objects */
-struct FChaosVDCharacterGroundConstraintSelectionHandle
-{
-	FChaosVDCharacterGroundConstraintSelectionHandle()
-	{		
-	}
-
-	FChaosVDCharacterGroundConstraintSelectionHandle(const TWeakPtr<FChaosVDCharacterGroundConstraint>& InConstraintData)
-		: ConstraintData(InConstraintData)
-	{
-	}
-
-	void SetIsSelected(bool bNewSelected);
-	bool IsSelected() const;
-
-	TWeakPtr<FChaosVDCharacterGroundConstraint> GetData() const { return ConstraintData; }
-
-private:
-	TWeakPtr<FChaosVDCharacterGroundConstraint> ConstraintData = nullptr;
-};
-
-DECLARE_MULTICAST_DELEGATE_OneParam(FChaosVDCharacterGroundSelectionChangedDelegate, const FChaosVDCharacterGroundConstraintSelectionHandle& SelectionHandle)
-
 UCLASS()
-class CHAOSVD_API UChaosVDSolverCharacterGroundConstraintDataComponent : public USolverDataComponent
+class CHAOSVD_API UChaosVDSolverCharacterGroundConstraintDataComponent : public UChaosVDConstraintDataComponent
 {
 	GENERATED_BODY()
-
 public:
+	virtual void SetScene(const TWeakPtr<FChaosVDScene>& InSceneWeakPtr) override;
 
-	UChaosVDSolverCharacterGroundConstraintDataComponent();
-
-	void UpdateConstraintData(const TArray<TSharedPtr<FChaosVDCharacterGroundConstraint>>& InConstraintData);
-	
-	const FChaosVDCharacterGroundDataArray& GetAllConstraints() const { return AllConstraints; }
-	const FChaosVDCharacterGroundDataArray* GetConstraintsForParticle(int32 ParticleID) const;
-
-	void SelectConstraint(int32 ConstraintIndex);
-	void SelectConstraint(const FChaosVDCharacterGroundConstraintSelectionHandle& SelectionHandle);
-
-	bool IsConstraintSelected(int32 ConstraintIndex) const;
-
-	FChaosVDCharacterGroundSelectionChangedDelegate& OnSelectionChanged() { return SelectionChangeDelegate; }
-
-	const FChaosVDCharacterGroundConstraintSelectionHandle& GetCurrentSelectionHandle() const { return CurrentConstraintSelectionHandle; }
-
-	virtual void ClearData() override;
-
-protected:
-
-	TSharedPtr<FChaosVDCharacterGroundConstraint> GetConstraintByIndex(int32 ConstraintIndex);
-
-	FChaosVDCharacterGroundDataArray AllConstraints;
-	
-	FChaosVDCharacterGroundDataByParticleMap ConstraintByCharacterParticle;
-
-	FChaosVDCharacterGroundDataByConstraintIndexMap ConstraintByConstraintIndex;
-
-	FChaosVDCharacterGroundSelectionChangedDelegate SelectionChangeDelegate;
+	void HandleSceneUpdated();
 
-	FChaosVDCharacterGroundConstraintSelectionHandle CurrentConstraintSelectionHandle;
+	virtual void BeginDestroy() override;
 };
diff --git a/Engine/Plugins/ChaosVD/Source/ChaosVD/Private/Components/ChaosVDSolverCollisionDataComponent.h b/Engine/Plugins/ChaosVD/Source/ChaosVD/Private/Components/ChaosVDSolverCollisionDataComponent.h
--- a/Engine/Plugins/ChaosVD/Source/ChaosVD/Private/Components/ChaosVDSolverCollisionDataComponent.h
+++ b/Engine/Plugins/ChaosVD/Source/ChaosVD/Private/Components/ChaosVDSolverCollisionDataComponent.h
@@ -8,7 +8,7 @@
 #include "Containers/Map.h"
 #include "DataWrappers/ChaosVDCollisionDataWrappers.h"
 #include "GameFramework/Info.h"
-#include "SolverDataComponent.h"
+#include "ChaosVDSolverDataComponent.h"
 #include "Templates/SharedPointer.h"
 
 #include "ChaosVDSolverCollisionDataComponent.generated.h"
@@ -19,7 +19,7 @@ typedef TMap<int32, TArray<TSharedPtr<FChaosVDParticlePairMidPhase>>> FChaosVDMi
 typedef TMap<int32, TArray<FChaosVDConstraint*>> FChaosVDConstraintByParticleMap;
 
 UCLASS()
-class UChaosVDSolverCollisionDataComponent : public USolverDataComponent
+class UChaosVDSolverCollisionDataComponent : public UChaosVDSolverDataComponent
 {
 	GENERATED_BODY()
 
diff --git a/Engine/Plugins/ChaosVD/Source/ChaosVD/Private/Components/ChaosVDSolverDataComponent.cpp b/Engine/Plugins/ChaosVD/Source/ChaosVD/Private/Components/ChaosVDSolverDataComponent.cpp
new file mode 100644
--- /dev/null
+++ b/Engine/Plugins/ChaosVD/Source/ChaosVD/Private/Components/ChaosVDSolverDataComponent.cpp
@@ -0,0 +1,9 @@
+// Copyright Epic Games, Inc. All Rights Reserved.
+
+#include "Components/ChaosVDSolverDataComponent.h"
+
+
+void UChaosVDSolverDataComponent::SetScene(const TWeakPtr<FChaosVDScene>& InSceneWeakPtr)
+{
+	SceneWeakPtr = InSceneWeakPtr;
+}
diff --git a/Engine/Plugins/ChaosVD/Source/ChaosVD/Private/Components/ChaosVDSolverDataComponent.h b/Engine/Plugins/ChaosVD/Source/ChaosVD/Private/Components/ChaosVDSolverDataComponent.h
new file mode 100644
--- /dev/null
+++ b/Engine/Plugins/ChaosVD/Source/ChaosVD/Private/Components/ChaosVDSolverDataComponent.h
@@ -0,0 +1,24 @@
+// Copyright Epic Games, Inc. All Rights Reserved.
+
+#pragma once
+
+#include "Components/ActorComponent.h"
+#include "ChaosVDSolverDataComponent.generated.h"
+
+class FChaosVDScene;
+/**
+ * Base class for all components that stores recorded solver data
+ */
+UCLASS(Abstract)
+class UChaosVDSolverDataComponent : public UActorComponent
+{
+	GENERATED_BODY()
+public:
+	virtual void ClearData() PURE_VIRTUAL(UChaosVDSolverDataComponent::ClearData);
+	virtual void SetScene(const TWeakPtr<FChaosVDScene>& InSceneWeakPtr);
+
+protected:
+
+	TWeakPtr<FChaosVDScene> SceneWeakPtr;
+};
+
diff --git a/Engine/Plugins/ChaosVD/Source/ChaosVD/Private/Components/ChaosVDSolverJointConstraintDataComponent.cpp b/Engine/Plugins/ChaosVD/Source/ChaosVD/Private/Components/ChaosVDSolverJointConstraintDataComponent.cpp
--- a/Engine/Plugins/ChaosVD/Source/ChaosVD/Private/Components/ChaosVDSolverJointConstraintDataComponent.cpp
+++ b/Engine/Plugins/ChaosVD/Source/ChaosVD/Private/Components/ChaosVDSolverJointConstraintDataComponent.cpp
@@ -1,106 +1,3 @@
 ﻿// Copyright Epic Games, Inc. All Rights Reserved.
 
 #include "Components/ChaosVDSolverJointConstraintDataComponent.h"
-
-#include "Actors/ChaosVDSolverInfoActor.h"
-#include "Components/ChaosVDConstraintDataHelpers.h"
-#include "DataWrappers/ChaosVDJointDataWrappers.h"
-
-void FChaosVDJointConstraintSelectionHandle::SetIsSelected(bool bNewSelected)
-{
-	if (const TSharedPtr<FChaosVDJointConstraint> JointDataPtr = JointData.Pin())
-	{
-		JointDataPtr->bIsSelectedInEditor = bNewSelected;
-	}
-}
-
-bool FChaosVDJointConstraintSelectionHandle::IsSelected() const
-{
-	if (const TSharedPtr<FChaosVDJointConstraint> JointDataPtr = JointData.Pin())
-	{
-		return JointDataPtr->bIsSelectedInEditor;
-	}
-
-	return false;
-}
-
-UChaosVDSolverJointConstraintDataComponent::UChaosVDSolverJointConstraintDataComponent()
-{
-	SetCanEverAffectNavigation(false);
-	bNavigationRelevant = false;
-	PrimaryComponentTick.bCanEverTick = false;
-}
-
-void UChaosVDSolverJointConstraintDataComponent::UpdateConstraintData(const TArray<TSharedPtr<FChaosVDJointConstraint>>& InJointData)
-{
-	ClearData();
-
-	AllJointConstraints = InJointData;
-
-	JointConstraintByParticle0.Reserve(InJointData.Num());
-	JointConstraintByParticle1.Reserve(InJointData.Num());
-	JointConstraintByConstraintIndex.Reserve(InJointData.Num());
-
-	for (const TSharedPtr<FChaosVDJointConstraint>& JointConstraint : InJointData)
-	{
-		if (!JointConstraint)
-		{
-			continue;
-		}
-
-		JointConstraintByConstraintIndex.Add(JointConstraint->ConstraintIndex, JointConstraint);
-		Chaos::VisualDebugger::Utils::AddDataDataToParticleIDMap(JointConstraintByParticle0, JointConstraint, JointConstraint->ParticleParIndexes[0]);
-		Chaos::VisualDebugger::Utils::AddDataDataToParticleIDMap(JointConstraintByParticle1, JointConstraint, JointConstraint->ParticleParIndexes[1]);
-	}
-}
-
-const FChaosVDJointDataArray* UChaosVDSolverJointConstraintDataComponent::GetJointConstraintsForParticle(int32 ParticleID, EChaosVDParticlePairSlot Options) const
-{
-	return Chaos::VisualDebugger::Utils::GetDataFromParticlePairMaps<FChaosVDJointDataByParticleMap, TSharedPtr<FChaosVDJointConstraint>>(JointConstraintByParticle0, JointConstraintByParticle1, ParticleID, Options);
-}
-
-void UChaosVDSolverJointConstraintDataComponent::SelectJoint(int32 ConstraintIndex)
-{
-	CurrentJointSelectionHandle.SetIsSelected(false);
-	CurrentJointSelectionHandle = FChaosVDJointConstraintSelectionHandle(GetJointConstraintByIndex(ConstraintIndex));
-	CurrentJointSelectionHandle.SetIsSelected(true);
-
-	SelectionChangeDelegate.Broadcast(CurrentJointSelectionHandle);
-}
-
-void UChaosVDSolverJointConstraintDataComponent::SelectJoint(const FChaosVDJointConstraintSelectionHandle& SelectionHandle)
-{
-	CurrentJointSelectionHandle.SetIsSelected(false);
-	CurrentJointSelectionHandle = SelectionHandle;
-	CurrentJointSelectionHandle.SetIsSelected(true);
-	
-	SelectionChangeDelegate.Broadcast(CurrentJointSelectionHandle);
-}
-
-bool UChaosVDSolverJointConstraintDataComponent::IsJointSelected(int32 ConstraintIndex) const
-{
-	if (const TSharedPtr<FChaosVDJointConstraint> JointDataPtr = CurrentJointSelectionHandle.GetData().Pin())
-	{
-		return JointDataPtr->bIsSelectedInEditor && JointDataPtr->ConstraintIndex == ConstraintIndex;
-	}
-
-	return false;
-}
-
-TSharedPtr<FChaosVDJointConstraint> UChaosVDSolverJointConstraintDataComponent::GetJointConstraintByIndex(int32 ConstraintIndex)
-{
-	if (TSharedPtr<FChaosVDJointConstraint>* JointConstraintPtrPtr = JointConstraintByConstraintIndex.Find(ConstraintIndex))
-	{
-		return *JointConstraintPtrPtr;
-	}
-
-	return nullptr;
-}
-
-void UChaosVDSolverJointConstraintDataComponent::ClearData()
-{
-	AllJointConstraints.Empty();
-	JointConstraintByParticle0.Empty();
-	JointConstraintByParticle1.Empty();
-	JointConstraintByConstraintIndex.Empty();
-}
diff --git a/Engine/Plugins/ChaosVD/Source/ChaosVD/Private/Components/ChaosVDSolverJointConstraintDataComponent.h b/Engine/Plugins/ChaosVD/Source/ChaosVD/Private/Components/ChaosVDSolverJointConstraintDataComponent.h
--- a/Engine/Plugins/ChaosVD/Source/ChaosVD/Private/Components/ChaosVDSolverJointConstraintDataComponent.h
+++ b/Engine/Plugins/ChaosVD/Source/ChaosVD/Private/Components/ChaosVDSolverJointConstraintDataComponent.h
@@ -2,77 +2,12 @@
 
 #pragma once
 
-#include "SolverDataComponent.h"
+#include "ChaosVDConstraintDataComponent.h"
 #include "ChaosVDSolverJointConstraintDataComponent.generated.h"
 
-struct FChaosVDJointConstraint;
-
-typedef TMap<int32, TArray<TSharedPtr<FChaosVDJointConstraint>>> FChaosVDJointDataByParticleMap;
-typedef TMap<int32, TSharedPtr<FChaosVDJointConstraint>> FChaosVDJointDataByConstraintIndexMap;
-typedef TArray<TSharedPtr<FChaosVDJointConstraint>> FChaosVDJointDataArray;
-
-enum class EChaosVDParticlePairSlot : uint8;
-
-/** Struct used to pass data about a specific joint constraint to other objects */
-struct FChaosVDJointConstraintSelectionHandle
-{
-	FChaosVDJointConstraintSelectionHandle()
-	{		
-	}
-
-	FChaosVDJointConstraintSelectionHandle(const TWeakPtr<FChaosVDJointConstraint>& InJointData)
-		: JointData(InJointData)
-	{
-	}
-
-	void SetIsSelected(bool bNewSelected);
-	bool IsSelected() const;
-
-	TWeakPtr<FChaosVDJointConstraint> GetData() const { return JointData; } 
-
-private:
-	TWeakPtr<FChaosVDJointConstraint> JointData = nullptr;
-};
-
-DECLARE_MULTICAST_DELEGATE_OneParam(FChaosVDJointSlectionChangedDelegate, const FChaosVDJointConstraintSelectionHandle& SelectionHandle)
 
 UCLASS()
-class CHAOSVD_API UChaosVDSolverJointConstraintDataComponent : public USolverDataComponent
+class CHAOSVD_API UChaosVDSolverJointConstraintDataComponent : public UChaosVDConstraintDataComponent
 {
 	GENERATED_BODY()
-
-public:
-
-	UChaosVDSolverJointConstraintDataComponent();
-
-	void UpdateConstraintData(const TArray<TSharedPtr<FChaosVDJointConstraint>>& InJointData);
-	
-	const FChaosVDJointDataArray& GetAllJointConstraints() const { return AllJointConstraints; }
-	const FChaosVDJointDataArray* GetJointConstraintsForParticle(int32 ParticleID, EChaosVDParticlePairSlot Options) const;
-
-	void SelectJoint(int32 ConstraintIndex);
-	void SelectJoint(const FChaosVDJointConstraintSelectionHandle& SelectionHandle);
-
-	bool IsJointSelected(int32 ConstraintIndex) const;
-
-	FChaosVDJointSlectionChangedDelegate& OnSelectionChanged() { return SelectionChangeDelegate; }
-
-	const FChaosVDJointConstraintSelectionHandle& GetCurrentSelectionHandle() const { return CurrentJointSelectionHandle; }
-
-	virtual void ClearData() override;
-
-protected:
-
-	TSharedPtr<FChaosVDJointConstraint> GetJointConstraintByIndex(int32 ConstraintIndex);
-
-	FChaosVDJointDataArray AllJointConstraints;
-	
-	FChaosVDJointDataByParticleMap JointConstraintByParticle0;
-	FChaosVDJointDataByParticleMap JointConstraintByParticle1;
-
-	FChaosVDJointDataByConstraintIndexMap JointConstraintByConstraintIndex;
-
-	FChaosVDJointSlectionChangedDelegate SelectionChangeDelegate;
-
-	FChaosVDJointConstraintSelectionHandle CurrentJointSelectionHandle;
 };
diff --git a/Engine/Plugins/ChaosVD/Source/ChaosVD/Private/Components/SolverDataComponent.h b/Engine/Plugins/ChaosVD/Source/ChaosVD/Private/Components/SolverDataComponent.h
deleted file mode 100644
diff --git a/Engine/Plugins/ChaosVD/Source/ChaosVD/Private/Settings/ChaosVDCharacterConstraintsVisualizationSettings.cpp b/Engine/Plugins/ChaosVD/Source/ChaosVD/Private/Settings/ChaosVDCharacterConstraintsVisualizationSettings.cpp
--- a/Engine/Plugins/ChaosVD/Source/ChaosVD/Private/Settings/ChaosVDCharacterConstraintsVisualizationSettings.cpp
+++ b/Engine/Plugins/ChaosVD/Source/ChaosVD/Private/Settings/ChaosVDCharacterConstraintsVisualizationSettings.cpp
@@ -2,7 +2,7 @@
 
 #include "ChaosVDCharacterConstraintsVisualizationSettings.h"
 
-void UChaosVDCharacterConstraintsVisualizationSettings::SetCharacterGroundConstraintDataVisualizationFlags(EChaosVDCharacterGroundConstraintDataVisualizationFlags NewFlags)
+void UChaosVDCharacterConstraintsVisualizationSettings::SetDataVisualizationFlags(EChaosVDCharacterGroundConstraintDataVisualizationFlags NewFlags)
 {
 	if (UChaosVDCharacterConstraintsVisualizationSettings* Settings = GetMutableDefault<UChaosVDCharacterConstraintsVisualizationSettings>())
 	{
@@ -11,7 +11,7 @@ void UChaosVDCharacterConstraintsVisualizationSettings::SetCharacterGroundConstr
 	}
 }
 
-EChaosVDCharacterGroundConstraintDataVisualizationFlags UChaosVDCharacterConstraintsVisualizationSettings::GetCharacterGroundConstraintDataVisualizationFlags()
+EChaosVDCharacterGroundConstraintDataVisualizationFlags UChaosVDCharacterConstraintsVisualizationSettings::GetDataVisualizationFlags()
 {
 	if (UChaosVDCharacterConstraintsVisualizationSettings* Settings = GetMutableDefault<UChaosVDCharacterConstraintsVisualizationSettings>())
 	{
diff --git a/Engine/Plugins/ChaosVD/Source/ChaosVD/Private/Settings/ChaosVDCharacterConstraintsVisualizationSettings.h b/Engine/Plugins/ChaosVD/Source/ChaosVD/Private/Settings/ChaosVDCharacterConstraintsVisualizationSettings.h
--- a/Engine/Plugins/ChaosVD/Source/ChaosVD/Private/Settings/ChaosVDCharacterConstraintsVisualizationSettings.h
+++ b/Engine/Plugins/ChaosVD/Source/ChaosVD/Private/Settings/ChaosVDCharacterConstraintsVisualizationSettings.h
@@ -43,6 +43,10 @@ public:
 	/** If true, any debug draw text available will be drawn */
 	UPROPERTY(EditAnywhere, Category=DebugDraw)
 	bool bShowDebugText = false;
+	
+	/** Color used for torque vector */
+	UPROPERTY(config, EditAnywhere, Category = DetailsPanel)
+	bool bAutoSelectConstraintFromSelectedParticle = false;
 
 	/** The depth priority used for while drawing data. Can be World or Foreground (with this one the shapes will be drawn on top of the geometry and be always visible) */
 	UPROPERTY(config, EditAnywhere, Category=DebugDraw)
@@ -69,8 +73,8 @@ public:
 	UPROPERTY(config, EditAnywhere, Category = DebugDraw)
 	FColor TorqueColor = FColor(255, 0, 255);
 
-	static void SetCharacterGroundConstraintDataVisualizationFlags(EChaosVDCharacterGroundConstraintDataVisualizationFlags NewFlags);
-	static EChaosVDCharacterGroundConstraintDataVisualizationFlags GetCharacterGroundConstraintDataVisualizationFlags();
+	static void SetDataVisualizationFlags(EChaosVDCharacterGroundConstraintDataVisualizationFlags NewFlags);
+	static EChaosVDCharacterGroundConstraintDataVisualizationFlags GetDataVisualizationFlags();
 
 private:
 	/** Set of flags to enable/disable visualization of specific character ground constraint data as debug draw */
diff --git a/Engine/Plugins/ChaosVD/Source/ChaosVD/Private/Settings/ChaosVDCollisionVisualizationSettings.cpp b/Engine/Plugins/ChaosVD/Source/ChaosVD/Private/Settings/ChaosVDCollisionVisualizationSettings.cpp
--- a/Engine/Plugins/ChaosVD/Source/ChaosVD/Private/Settings/ChaosVDCollisionVisualizationSettings.cpp
+++ b/Engine/Plugins/ChaosVD/Source/ChaosVD/Private/Settings/ChaosVDCollisionVisualizationSettings.cpp
@@ -2,7 +2,7 @@
 
 #include "ChaosVDCollisionVisualizationSettings.h"
 
-void UChaosVDCollisionDataVisualizationSettings::SetCollisionDataVisualizationFlags(EChaosVDCollisionVisualizationFlags NewFlags)
+void UChaosVDCollisionDataVisualizationSettings::SetDataVisualizationFlags(EChaosVDCollisionVisualizationFlags NewFlags)
 {
 	if (UChaosVDCollisionDataVisualizationSettings* Settings = GetMutableDefault<UChaosVDCollisionDataVisualizationSettings>())
 	{
@@ -11,7 +11,7 @@ void UChaosVDCollisionDataVisualizationSettings::SetCollisionDataVisualizationFl
 	}
 }
 
-EChaosVDCollisionVisualizationFlags UChaosVDCollisionDataVisualizationSettings::GetCollisionDataVisualizationFlags()
+EChaosVDCollisionVisualizationFlags UChaosVDCollisionDataVisualizationSettings::GetDataVisualizationFlags()
 {
 	if (UChaosVDCollisionDataVisualizationSettings* Settings = GetMutableDefault<UChaosVDCollisionDataVisualizationSettings>())
 	{
diff --git a/Engine/Plugins/ChaosVD/Source/ChaosVD/Private/Settings/ChaosVDCollisionVisualizationSettings.h b/Engine/Plugins/ChaosVD/Source/ChaosVD/Private/Settings/ChaosVDCollisionVisualizationSettings.h
--- a/Engine/Plugins/ChaosVD/Source/ChaosVD/Private/Settings/ChaosVDCollisionVisualizationSettings.h
+++ b/Engine/Plugins/ChaosVD/Source/ChaosVD/Private/Settings/ChaosVDCollisionVisualizationSettings.h
@@ -43,8 +43,8 @@ public:
 	UPROPERTY(EditAnywhere, config, Category=DebugDraw)
 	float ContactNormalScale = 30.0f;
 
-	static void SetCollisionDataVisualizationFlags(EChaosVDCollisionVisualizationFlags NewFlags);
-	static EChaosVDCollisionVisualizationFlags GetCollisionDataVisualizationFlags();
+	static void SetDataVisualizationFlags(EChaosVDCollisionVisualizationFlags NewFlags);
+	static EChaosVDCollisionVisualizationFlags GetDataVisualizationFlags();
 
 private:
 	/** Set of flags to enable/disable visualization of specific collision data as debug draw */
diff --git a/Engine/Plugins/ChaosVD/Source/ChaosVD/Private/Settings/ChaosVDJointConstraintVisualizationSettings.cpp b/Engine/Plugins/ChaosVD/Source/ChaosVD/Private/Settings/ChaosVDJointConstraintVisualizationSettings.cpp
--- a/Engine/Plugins/ChaosVD/Source/ChaosVD/Private/Settings/ChaosVDJointConstraintVisualizationSettings.cpp
+++ b/Engine/Plugins/ChaosVD/Source/ChaosVD/Private/Settings/ChaosVDJointConstraintVisualizationSettings.cpp
@@ -2,7 +2,7 @@
 
 #include "ChaosVDJointConstraintVisualizationSettings.h"
 
-void UChaosVDJointConstraintsVisualizationSettings::SetJointsDataVisualizationFlags(EChaosVDJointsDataVisualizationFlags NewFlags)
+void UChaosVDJointConstraintsVisualizationSettings::SetDataVisualizationFlags(EChaosVDJointsDataVisualizationFlags NewFlags)
 {
 	if (UChaosVDJointConstraintsVisualizationSettings* Settings = GetMutableDefault<UChaosVDJointConstraintsVisualizationSettings>())
 	{
@@ -11,7 +11,7 @@ void UChaosVDJointConstraintsVisualizationSettings::SetJointsDataVisualizationFl
 	}
 }
 
-EChaosVDJointsDataVisualizationFlags UChaosVDJointConstraintsVisualizationSettings::GetJointsDataVisualizationFlags()
+EChaosVDJointsDataVisualizationFlags UChaosVDJointConstraintsVisualizationSettings::GetDataVisualizationFlags()
 {
 	if (UChaosVDJointConstraintsVisualizationSettings* Settings = GetMutableDefault<UChaosVDJointConstraintsVisualizationSettings>())
 	{
diff --git a/Engine/Plugins/ChaosVD/Source/ChaosVD/Private/Settings/ChaosVDJointConstraintVisualizationSettings.h b/Engine/Plugins/ChaosVD/Source/ChaosVD/Private/Settings/ChaosVDJointConstraintVisualizationSettings.h
--- a/Engine/Plugins/ChaosVD/Source/ChaosVD/Private/Settings/ChaosVDJointConstraintVisualizationSettings.h
+++ b/Engine/Plugins/ChaosVD/Source/ChaosVD/Private/Settings/ChaosVDJointConstraintVisualizationSettings.h
@@ -67,8 +67,8 @@ public:
 	UPROPERTY(EditAnywhere, config, Category=DebugDraw)
 	float ConstraintAxisLength = 10.0f;
 
-	static void SetJointsDataVisualizationFlags(EChaosVDJointsDataVisualizationFlags NewFlags);
-	static EChaosVDJointsDataVisualizationFlags GetJointsDataVisualizationFlags();
+	static void SetDataVisualizationFlags(EChaosVDJointsDataVisualizationFlags NewFlags);
+	static EChaosVDJointsDataVisualizationFlags GetDataVisualizationFlags();
 
 private:
 	/** Set of flags to enable/disable visualization of specific joint constraints data as debug draw */
diff --git a/Engine/Plugins/ChaosVD/Source/ChaosVD/Private/Settings/ChaosVDSceneQueryVisualizationSettings.cpp b/Engine/Plugins/ChaosVD/Source/ChaosVD/Private/Settings/ChaosVDSceneQueryVisualizationSettings.cpp
--- a/Engine/Plugins/ChaosVD/Source/ChaosVD/Private/Settings/ChaosVDSceneQueryVisualizationSettings.cpp
+++ b/Engine/Plugins/ChaosVD/Source/ChaosVD/Private/Settings/ChaosVDSceneQueryVisualizationSettings.cpp
@@ -2,7 +2,7 @@
 
 #include "ChaosVDSceneQueryVisualizationSettings.h"
 
-void UChaosVDSceneQueriesVisualizationSettings::SetSceneQueryDataVisualizationFlags(EChaosVDSceneQueryVisualizationFlags NewFlags)
+void UChaosVDSceneQueriesVisualizationSettings::SetDataVisualizationFlags(EChaosVDSceneQueryVisualizationFlags NewFlags)
 {
 	if (UChaosVDSceneQueriesVisualizationSettings* Settings = GetMutableDefault<UChaosVDSceneQueriesVisualizationSettings>())
 	{
@@ -11,7 +11,7 @@ void UChaosVDSceneQueriesVisualizationSettings::SetSceneQueryDataVisualizationFl
 	}
 }
 
-EChaosVDSceneQueryVisualizationFlags UChaosVDSceneQueriesVisualizationSettings::GetSceneQueryDataVisualizationFlags()
+EChaosVDSceneQueryVisualizationFlags UChaosVDSceneQueriesVisualizationSettings::GetDataVisualizationFlags()
 {
 	if (UChaosVDSceneQueriesVisualizationSettings* Settings = GetMutableDefault<UChaosVDSceneQueriesVisualizationSettings>())
 	{
diff --git a/Engine/Plugins/ChaosVD/Source/ChaosVD/Private/Settings/ChaosVDSceneQueryVisualizationSettings.h b/Engine/Plugins/ChaosVD/Source/ChaosVD/Private/Settings/ChaosVDSceneQueryVisualizationSettings.h
--- a/Engine/Plugins/ChaosVD/Source/ChaosVD/Private/Settings/ChaosVDSceneQueryVisualizationSettings.h
+++ b/Engine/Plugins/ChaosVD/Source/ChaosVD/Private/Settings/ChaosVDSceneQueryVisualizationSettings.h
@@ -28,8 +28,8 @@ class UChaosVDSceneQueriesVisualizationSettings : public UChaosVDVisualizationSe
 
 public:
 
-	static void SetSceneQueryDataVisualizationFlags(EChaosVDSceneQueryVisualizationFlags NewFlags);
-	static EChaosVDSceneQueryVisualizationFlags GetSceneQueryDataVisualizationFlags();
+	static void SetDataVisualizationFlags(EChaosVDSceneQueryVisualizationFlags NewFlags);
+	static EChaosVDSceneQueryVisualizationFlags GetDataVisualizationFlags();
 
 	/** If true, any debug draw text available will be drawn */
 	UPROPERTY(EditAnywhere, Category=DebugDraw)
diff --git a/Engine/Plugins/ChaosVD/Source/ChaosVD/Private/Trace/ChaosVDTraceProvider.cpp b/Engine/Plugins/ChaosVD/Source/ChaosVD/Private/Trace/ChaosVDTraceProvider.cpp
--- a/Engine/Plugins/ChaosVD/Source/ChaosVD/Private/Trace/ChaosVDTraceProvider.cpp
+++ b/Engine/Plugins/ChaosVD/Source/ChaosVD/Private/Trace/ChaosVDTraceProvider.cpp
@@ -73,6 +73,11 @@ void FChaosVDTraceProvider::CommitProcessedGameFramesToRecording()
 		TSharedPtr<FChaosVDGameFrameData> ProcessedGameFrameData;
 		DeQueueGameFrameForProcessing(ProcessedGameFrameData);
 
+		if (StartLastCommitedFrameTimeSeconds == 0.0)
+		{
+			StartLastCommitedFrameTimeSeconds = FPlatformTime::Seconds();
+		}
+
 		if (ProcessedGameFrameData.IsValid())
 		{
 			InternalRecording->GetAvailableSolverIDsAtGameFrame(*ProcessedGameFrameData, SolverIDs);
@@ -209,6 +214,16 @@ void FChaosVDTraceProvider::HandleAnalysisComplete()
 		}
 	}
 
+	if (TSharedPtr<FChaosVDRecording> Recording = GetRecordingForSession())
+	{
+		double TotalTimeProcessingFrames = FPlatformTime::Seconds() - StartLastCommitedFrameTimeSeconds;
+
+		int32 NumOfGameFramesProcessed = Recording->GetAvailableGameFramesNumber();
+		double AvgTimePerFrameSeconds = TotalTimeProcessingFrames / NumOfGameFramesProcessed;
+		
+		UE_LOG(LogChaosVDEditor, Log, TEXT(" [%d] Game frames Processed at [%f] ms per frame on average"), NumOfGameFramesProcessed, AvgTimePerFrameSeconds * 1000.0);
+	}
+
 	UE_LOG(LogChaosVDEditor, Log, TEXT("Total size of loaded data => [%s]"), *FText::AsMemory(TotalBytes, &SizeFormattingOptions,nullptr, EMemoryUnitStandard::IEC).ToString());
 }
 
diff --git a/Engine/Plugins/ChaosVD/Source/ChaosVD/Private/Trace/ChaosVDTraceProvider.h b/Engine/Plugins/ChaosVD/Source/ChaosVD/Private/Trace/ChaosVDTraceProvider.h
--- a/Engine/Plugins/ChaosVD/Source/ChaosVD/Private/Trace/ChaosVDTraceProvider.h
+++ b/Engine/Plugins/ChaosVD/Source/ChaosVD/Private/Trace/ChaosVDTraceProvider.h
@@ -105,4 +105,6 @@ private:
 	int32 CurrentGameFrameQueueSize = 0;
 
 	bool bDefaultDataProcessorsRegistered = false;
+
+	double StartLastCommitedFrameTimeSeconds = 0.0;
 };
diff --git a/Engine/Plugins/ChaosVD/Source/ChaosVD/Private/Visualizers/ChaosVDCharacterGroundConstraintsDataComponentVisualizer.cpp b/Engine/Plugins/ChaosVD/Source/ChaosVD/Private/Visualizers/ChaosVDCharacterGroundConstraintsDataComponentVisualizer.cpp
--- a/Engine/Plugins/ChaosVD/Source/ChaosVD/Private/Visualizers/ChaosVDCharacterGroundConstraintsDataComponentVisualizer.cpp
+++ b/Engine/Plugins/ChaosVD/Source/ChaosVD/Private/Visualizers/ChaosVDCharacterGroundConstraintsDataComponentVisualizer.cpp
@@ -4,7 +4,6 @@
 
 #include "ChaosVDScene.h"
 #include "ChaosVDTabsIDs.h"
-#include "EditorModeManager.h"
 #include "EditorViewportClient.h"
 #include "SceneManagement.h"
 #include "SceneView.h"
@@ -15,10 +14,7 @@
 #include "ToolMenus.h"
 #include "ToolMenuEntry.h"
 #include "ToolMenuSection.h"
-#include "Utils/ChaosVDUserInterfaceUtils.h"
 #include "Visualizers/ChaosVDDebugDrawUtils.h"
-#include "Widgets/SChaosVDEnumFlagsMenu.h"
-#include "Widgets/SChaosVDMainTab.h"
 #include "Widgets/SChaosVDViewportToolbar.h"
 
 #define LOCTEXT_NAMESPACE "ChaosVisualDebugger"
@@ -29,42 +25,24 @@ bool FChaosVDCharacterGroundConstraintVisualizationDataContext::IsVisualizationF
 	return EnumHasAnyFlags(FlagsAsParticleFlags, Flag);
 }
 
-IMPLEMENT_HIT_PROXY(HChaosVDCharacterGroundConstraintProxy, HComponentVisProxy)
-
-
 FChaosVDCharacterGroundConstraintDataComponentVisualizer::FChaosVDCharacterGroundConstraintDataComponentVisualizer()
 {
-	RegisterVisualizerMenus();
+	FChaosVDCharacterGroundConstraintDataComponentVisualizer::RegisterVisualizerMenus();
+	InspectorTabID = FChaosVDTabID::ConstraintsInspector;
 }
 
 void FChaosVDCharacterGroundConstraintDataComponentVisualizer::RegisterVisualizerMenus()
 {
-		UToolMenus* ToolMenus = UToolMenus::Get();
+	FName MenuSection("CharacterGroundConstraintDataVisualization.Show");
+	FText MenuSectionLabel = LOCTEXT("CharacterConstraintDataVisualizationShowMenuLabel", "Character Ground Constraints Data Visualization");
+	FText FlagsMenuLabel = LOCTEXT("CharacterConstraintDataVisualizationFlagsMenuLabel", "Character Ground Constraints Data Flags");
+	FText FlagsMenuTooltip = LOCTEXT("CharacterConstraintDataVisualizationFlagsMenuToolTip", "Set of flags to enable/disable visibility of specific types of Character Constraints data");
+	FSlateIcon FlagsMenuIcon = FSlateIcon(FAppStyle::Get().GetStyleSetName(), TEXT("ClassIcon.Character"));
+
+	FText SettingsMenuLabel = LOCTEXT("CharacterConstraintDataVisualizationMenuLabel", "Character Ground Constraints Visualization Settings");
+	FText SettingsMenuTooltip = LOCTEXT("CharacterConstraintDataVisualizationMenuToolTip", "Options to change how the recorded Character Constraints data is debug drawn");
 	
-	if (!ensure(ToolMenus))
-	{
-		return;
-	}
-
-	if (UToolMenu* Menu = ToolMenus->ExtendMenu(SChaosVDViewportToolbar::ShowMenuName))
-	{
-		FToolMenuSection& Section = Menu->AddSection("CharacterConstraintDataVisualization.Show", LOCTEXT("CharacterConstraintDataVisualizationShowMenuLabel", "Character Constraints Data Visualization"));
-		
-		Section.AddSubMenu(TEXT("CharacterConstraintDataVisualizationFlags"), LOCTEXT("CharacterConstraintDataVisualizationFlagsMenuLabel", "Character Constraints Data Flags"), LOCTEXT("CharacterConstraintDataVisualizationFlagsMenuToolTip", "Set of flags to enable/disable visibility of specific types of Character Constraints data"), FNewToolMenuDelegate::CreateLambda([](UToolMenu* Menu)
-		                   {
-			                   TSharedRef<SWidget> VisualizationFlagsWidget = SNew(SChaosVDEnumFlagsMenu<EChaosVDCharacterGroundConstraintDataVisualizationFlags>)
-				                   .CurrentValue_Static(&UChaosVDCharacterConstraintsVisualizationSettings::GetCharacterGroundConstraintDataVisualizationFlags)
-				                   .OnEnumSelectionChanged_Static(&UChaosVDCharacterConstraintsVisualizationSettings::SetCharacterGroundConstraintDataVisualizationFlags);
-			
-			                   FToolMenuEntry FlagsMenuEntry = FToolMenuEntry::InitWidget("CharacterConstraintDataVisualizationFlags", VisualizationFlagsWidget,FText::GetEmpty());
-			                   Menu->AddMenuEntry(NAME_None, FlagsMenuEntry);
-		                   }),
-		                   false, FSlateIcon(FAppStyle::Get().GetStyleSetName(), TEXT("ClassIcon.Character")));
-
-		using namespace Chaos::VisualDebugger::Utils;
-		Section.AddSubMenu(TEXT("CharacterConstraintDataVisualizationSettings"), LOCTEXT("CharacterConstraintDataVisualizationMenuLabel", "Character Constraints Visualization Settings"), LOCTEXT("CharacterConstraintDataVisualizationMenuToolTip", "Options to change how the recorded Character Constraints data is debug drawn"), FNewToolMenuDelegate::CreateStatic(&CreateMenuEntryForDefaultObject<UChaosVDCharacterConstraintsVisualizationSettings>, EChaosVDSaveSettingsOptions::ShowSaveButton),
-		                   false, FSlateIcon(FAppStyle::Get().GetStyleSetName(), TEXT("Icons.Toolbar.Settings")));
-	}
+	CreateGenericVisualizerMenu<UChaosVDCharacterConstraintsVisualizationSettings, EChaosVDCharacterGroundConstraintDataVisualizationFlags>(SChaosVDViewportToolbar::ShowMenuName, MenuSection, MenuSectionLabel, FlagsMenuLabel, FlagsMenuTooltip, FlagsMenuIcon, SettingsMenuLabel, SettingsMenuTooltip);
 }
 
 void FChaosVDCharacterGroundConstraintDataComponentVisualizer::DrawVisualization(const UActorComponent* Component, const FSceneView* View, FPrimitiveDrawInterface* PDI)
@@ -103,60 +81,48 @@ void FChaosVDCharacterGroundConstraintDataComponentVisualizer::DrawVisualization
 	VisualizationContext.SpaceTransform = SolverInfoActor->GetSimulationTransform();
 	VisualizationContext.SolverInfoActor = SolverInfoActor;
 
-	VisualizationContext.VisualizationFlags = static_cast<uint32>(UChaosVDCharacterConstraintsVisualizationSettings::GetCharacterGroundConstraintDataVisualizationFlags());
+	VisualizationContext.VisualizationFlags = static_cast<uint32>(UChaosVDCharacterConstraintsVisualizationSettings::GetDataVisualizationFlags());
 	VisualizationContext.DebugDrawSettings = GetDefault<UChaosVDCharacterConstraintsVisualizationSettings>();
 
 	if (!VisualizationContext.IsVisualizationFlagEnabled(EChaosVDCharacterGroundConstraintDataVisualizationFlags::EnableDraw))
 	{
 		return;
 	}
+	
+	TSharedPtr<FChaosVDSolverDataSelection> SelectionObject = CVDScene->GetSolverDataSelectionObject().Pin();
+	TSharedPtr<FChaosVDSolverDataSelectionHandle> SelectionHandle;
+	if (SelectionObject)
+	{
+		SelectionHandle = SelectionObject->GetCurrentSelectionHandle();
+	}
 
 	// If nothing is selected, fallback to draw all character ground constraints
-	const bool bDrawOnlySelected = VisualizationContext.IsVisualizationFlagEnabled(EChaosVDCharacterGroundConstraintDataVisualizationFlags::OnlyDrawSelected) && ConstraintDataComponent->GetCurrentSelectionHandle().IsSelected();
+	const bool bDrawOnlySelected = VisualizationContext.IsVisualizationFlagEnabled(EChaosVDCharacterGroundConstraintDataVisualizationFlags::OnlyDrawSelected) && SelectionHandle && SelectionHandle->IsSelected();
 	if (bDrawOnlySelected)
 	{
-		if (const TSharedPtr<FChaosVDCharacterGroundConstraint> Constraint = ConstraintDataComponent->GetCurrentSelectionHandle().GetData().Pin())
+		if (FChaosVDCharacterGroundConstraint* Constraint = SelectionHandle->GetData<FChaosVDCharacterGroundConstraint>())
 		{
-			VisualizationContext.DataSelectionHandle = FChaosVDCharacterGroundConstraintSelectionHandle(Constraint);
+			VisualizationContext.DataSelectionHandle = SelectionHandle;
 			DrawConstraint(Component, *Constraint, VisualizationContext, View, PDI);
 		}
 	}
 	else
 	{
-		for (const TSharedPtr<FChaosVDCharacterGroundConstraint>& Constraint : ConstraintDataComponent->GetAllConstraints())
+		for (const TSharedPtr<FChaosVDConstraintDataWrapperBase>& Constraint : ConstraintDataComponent->GetAllConstraints())
 		{
-			if (Constraint)
+			static_assert(std::is_base_of_v<FChaosVDConstraintDataWrapperBase, FChaosVDCharacterGroundConstraint>, "Only FChaosVDCharacterGroundConstraint is supported");
+			if (FChaosVDCharacterGroundConstraint* CharacterGroundConstraint = static_cast<FChaosVDCharacterGroundConstraint*>(Constraint.Get()))
 			{
-				VisualizationContext.DataSelectionHandle = FChaosVDCharacterGroundConstraintSelectionHandle(Constraint);
-				DrawConstraint(Component, *Constraint, VisualizationContext, View, PDI);
+				VisualizationContext.DataSelectionHandle = SelectionObject->MakeSelectionHandle(StaticCastSharedPtr<FChaosVDCharacterGroundConstraint>(Constraint));
+				DrawConstraint(Component, *CharacterGroundConstraint, VisualizationContext, View, PDI);
 			}
 		}
 	}
 }
 
-bool FChaosVDCharacterGroundConstraintDataComponentVisualizer::VisProxyHandleClick(FEditorViewportClient* InViewportClient, HComponentVisProxy* VisProxy, const FViewportClick& Click)
+bool FChaosVDCharacterGroundConstraintDataComponentVisualizer::CanHandleClick(const HChaosVDComponentVisProxy& VisProxy)
 {
-	const HChaosVDCharacterGroundConstraintProxy* ConstraintDataProxy = HitProxyCast<HChaosVDCharacterGroundConstraintProxy>(VisProxy);
-	if (ConstraintDataProxy == nullptr)
-	{
-		return false;
-	}
-	
-	if (const UChaosVDSolverCharacterGroundConstraintDataComponent* ConstraintDataComponent = Cast<UChaosVDSolverCharacterGroundConstraintDataComponent>(ConstraintDataProxy->Component.Get()))
-	{
-		// Bring the constraint details tab into focus if available
-		const TSharedPtr<SChaosVDMainTab> MainTabToolkitHost = InViewportClient->GetModeTools() ? StaticCastSharedPtr<SChaosVDMainTab>(InViewportClient->GetModeTools()->GetToolkitHost()) : nullptr;
-		if (const TSharedPtr<FTabManager> TabManager = MainTabToolkitHost ? MainTabToolkitHost->GetTabManager() : nullptr)
-		{
-			TabManager->TryInvokeTab(FChaosVDTabID::CharacterGroundConstraintDataDetails);
-		}
-
-		const_cast<UChaosVDSolverCharacterGroundConstraintDataComponent*>(ConstraintDataComponent)->SelectConstraint(ConstraintDataProxy->DataSelectionHandle);
-
-		return true;
-	}
-
-	return false;
+	return VisProxy.DataSelectionHandle && VisProxy.DataSelectionHandle->IsA<FChaosVDCharacterGroundConstraint>();
 }
 
 void FChaosVDCharacterGroundConstraintDataComponentVisualizer::DrawConstraint(const UActorComponent* Component, const FChaosVDCharacterGroundConstraint& InConstraintData, FChaosVDCharacterGroundConstraintVisualizationDataContext& VisualizationContext, const FSceneView* View, FPrimitiveDrawInterface* PDI)
@@ -192,7 +158,7 @@ void FChaosVDCharacterGroundConstraintDataComponentVisualizer::DrawConstraint(co
 		return;
 	}
 
-	const FChaosVDParticleDataWrapper* CharacterParticleData = nullptr;
+	TSharedPtr<const FChaosVDParticleDataWrapper> CharacterParticleData = nullptr;
 
 	if (AChaosVDParticleActor* CharacterParticle = VisualizationContext.SolverInfoActor->GetParticleActor(InConstraintData.CharacterParticleIndex))
 	{
@@ -211,9 +177,9 @@ void FChaosVDCharacterGroundConstraintDataComponentVisualizer::DrawConstraint(co
 		return;
 	}
 
-	PDI->SetHitProxy(new HChaosVDCharacterGroundConstraintProxy(Component, VisualizationContext.DataSelectionHandle));
+	PDI->SetHitProxy(new HChaosVDComponentVisProxy(Component, VisualizationContext.DataSelectionHandle));
 
-	const float LineThickness = InConstraintData.bIsSelectedInEditor ?  DebugDrawSettings->BaseLineThickness * 1.5f :  DebugDrawSettings->BaseLineThickness;
+	const float LineThickness = VisualizationContext.DataSelectionHandle && VisualizationContext.DataSelectionHandle->IsSelected() ?  DebugDrawSettings->BaseLineThickness * 1.5f :  DebugDrawSettings->BaseLineThickness;
 
 	const FVector CharacterPos = CharacterParticleData->ParticlePositionRotation.MX;
 	const FVector UpDir = InConstraintData.Settings.VerticalAxis;
diff --git a/Engine/Plugins/ChaosVD/Source/ChaosVD/Private/Visualizers/ChaosVDCharacterGroundConstraintsDataComponentVisualizer.h b/Engine/Plugins/ChaosVD/Source/ChaosVD/Private/Visualizers/ChaosVDCharacterGroundConstraintsDataComponentVisualizer.h
--- a/Engine/Plugins/ChaosVD/Source/ChaosVD/Private/Visualizers/ChaosVDCharacterGroundConstraintsDataComponentVisualizer.h
+++ b/Engine/Plugins/ChaosVD/Source/ChaosVD/Private/Visualizers/ChaosVDCharacterGroundConstraintsDataComponentVisualizer.h
@@ -1,11 +1,11 @@
 // Copyright Epic Games, Inc. All Rights Reserved.
 
 #pragma once
+#include "ChaosVDComponentVisualizerBase.h"
 #include "ComponentVisualizer.h"
-#include "IChaosVDParticleVisualizationDataProvider.h"
 #include "Chaos/Core.h"
-#include "Components/ChaosVDSolverCharacterGroundConstraintDataComponent.h"
 
+struct FChaosVDSolverDataSelectionHandle;
 class AChaosVDSolverInfoActor;
 enum class EChaosVDCharacterGroundConstraintDataVisualizationFlags : uint32;
 struct FChaosVDCharacterGroundConstraintDebugDrawSettings;
@@ -14,42 +14,26 @@ struct FChaosVDCharacterGroundConstraint;
 /** Visualization context structure specific for character ground constraint visualizations */
 struct FChaosVDCharacterGroundConstraintVisualizationDataContext : public FChaosVDVisualizationContext
 {
-	FChaosVDCharacterGroundConstraintSelectionHandle DataSelectionHandle = FChaosVDCharacterGroundConstraintSelectionHandle(nullptr);
+	TSharedPtr<FChaosVDSolverDataSelectionHandle> DataSelectionHandle = nullptr;
 	
 	AChaosVDSolverInfoActor* SolverInfoActor = nullptr;
 
 	bool IsVisualizationFlagEnabled(EChaosVDCharacterGroundConstraintDataVisualizationFlags Flag) const;
 };
 
-/** Custom Hit Proxy for debug drawn scene queries */
-struct HChaosVDCharacterGroundConstraintProxy : public HComponentVisProxy
-{
-	DECLARE_HIT_PROXY()
-	
-	HChaosVDCharacterGroundConstraintProxy(const UActorComponent* Component, const FChaosVDCharacterGroundConstraintSelectionHandle& InConstraintFinderData) : HComponentVisProxy(Component, HPP_UI), DataSelectionHandle(InConstraintFinderData)
-	{	
-	}
-
-	virtual EMouseCursor::Type GetMouseCursor() override
-	{
-		return EMouseCursor::Crosshairs;
-	}
-
-	FChaosVDCharacterGroundConstraintSelectionHandle DataSelectionHandle;
-};
-
 /**
  * Component visualizer in charge of generating debug draw visualizations for character ground constraints in a UChaosVDSolverCharacterGroundConstraintDataComponent
  */
-class FChaosVDCharacterGroundConstraintDataComponentVisualizer final : public FComponentVisualizer
+class FChaosVDCharacterGroundConstraintDataComponentVisualizer final : public FChaosVDComponentVisualizerBase
 {
 public:
 	FChaosVDCharacterGroundConstraintDataComponentVisualizer();
-	
-	void RegisterVisualizerMenus();
+
+	virtual void RegisterVisualizerMenus() override;
 
 	virtual void DrawVisualization(const UActorComponent* Component, const FSceneView* View, FPrimitiveDrawInterface* PDI) override;
-	virtual bool VisProxyHandleClick(FEditorViewportClient* InViewportClient, HComponentVisProxy* VisProxy, const FViewportClick& Click) override;
+
+	virtual bool CanHandleClick(const HChaosVDComponentVisProxy& VisProxy) override;
 
 protected:
 
diff --git a/Engine/Plugins/ChaosVD/Source/ChaosVD/Private/Visualizers/ChaosVDComponentVisualizerBase.cpp b/Engine/Plugins/ChaosVD/Source/ChaosVD/Private/Visualizers/ChaosVDComponentVisualizerBase.cpp
new file mode 100644
--- /dev/null
+++ b/Engine/Plugins/ChaosVD/Source/ChaosVD/Private/Visualizers/ChaosVDComponentVisualizerBase.cpp
@@ -0,0 +1,64 @@
+// Copyright Epic Games, Inc. All Rights Reserved.
+
+#include "ChaosVDComponentVisualizerBase.h"
+
+#include "ChaosVDEngine.h"
+#include "ChaosVDScene.h"
+#include "ChaosVDSolverDataSelection.h"
+#include "EditorModeManager.h"
+#include "EditorViewportClient.h"
+#include "Widgets/SChaosVDMainTab.h"
+
+IMPLEMENT_HIT_PROXY(HChaosVDComponentVisProxy, HComponentVisProxy)
+
+
+bool FChaosVDComponentVisualizerBase::CanHandleClick(const HChaosVDComponentVisProxy& VisProxy)
+{
+	return false;
+}
+
+bool FChaosVDComponentVisualizerBase::VisProxyHandleClick(FEditorViewportClient* InViewportClient, HComponentVisProxy* VisProxy, const FViewportClick& Click)
+{
+	const HChaosVDComponentVisProxy* ChaosVDHitProxy = HitProxyCast<HChaosVDComponentVisProxy>(VisProxy);
+	if (ChaosVDHitProxy == nullptr)
+	{
+		return false;
+	}
+
+	if (!CanHandleClick(*ChaosVDHitProxy))
+	{
+		return false;
+	}
+
+	const TSharedPtr<SChaosVDMainTab> MainTabToolkitHost = InViewportClient->GetModeTools() ? StaticCastSharedPtr<SChaosVDMainTab>(InViewportClient->GetModeTools()->GetToolkitHost()) : nullptr;
+	if (!MainTabToolkitHost)
+	{
+		return false;
+	}
+
+	if (TSharedPtr<FChaosVDScene> CVDScene = MainTabToolkitHost->GetChaosVDEngineInstance()->GetCurrentScene())
+	{
+		// Bring the constraint details tab into focus if available
+		if (const TSharedPtr<FTabManager> TabManager = MainTabToolkitHost ? MainTabToolkitHost->GetTabManager() : nullptr)
+		{
+			TabManager->TryInvokeTab(InspectorTabID);
+		}
+
+		return SelectVisualizedData(*ChaosVDHitProxy, CVDScene.ToSharedRef());
+	}
+
+	return false;
+}
+
+bool FChaosVDComponentVisualizerBase::SelectVisualizedData(const HChaosVDComponentVisProxy& VisProxy, const TSharedRef<FChaosVDScene>& InCVDScene)
+{
+	if (TSharedPtr<FChaosVDSolverDataSelection> SelectionObject = InCVDScene->GetSolverDataSelectionObject().Pin())
+	{
+		SelectionObject->SelectData(VisProxy.DataSelectionHandle);
+
+		return true;
+	}
+
+	return false;
+}
+
diff --git a/Engine/Plugins/ChaosVD/Source/ChaosVD/Private/Visualizers/ChaosVDComponentVisualizerBase.h b/Engine/Plugins/ChaosVD/Source/ChaosVD/Private/Visualizers/ChaosVDComponentVisualizerBase.h
new file mode 100644
--- /dev/null
+++ b/Engine/Plugins/ChaosVD/Source/ChaosVD/Private/Visualizers/ChaosVDComponentVisualizerBase.h
@@ -0,0 +1,100 @@
+// Copyright Epic Games, Inc. All Rights Reserved.
+
+#pragma once
+
+#include "ComponentVisualizer.h"
+#include "HAL/Platform.h"
+#include "Templates/SharedPointer.h"
+#include "ToolMenu.h"
+#include "ToolMenus.h"
+#include "ToolMenuEntry.h"
+#include "ToolMenuSection.h"
+#include "Utils/ChaosVDUserInterfaceUtils.h"
+#include "Widgets/SChaosVDEnumFlagsMenu.h"
+#include "Widgets/SWidget.h"
+
+class FChaosVDSolverDataSelection;
+struct FChaosVDSolverDataSelectionHandle;
+class UChaosVDSettingsObjectBase;
+class FChaosVDScene;
+
+/** Context needed to be able to visualize data in the viewport */
+struct FChaosVDVisualizationContext
+{
+	FTransform SpaceTransform;
+	TWeakPtr<FChaosVDScene> CVDScene;
+	int32 SolverID = INDEX_NONE;
+	uint32 VisualizationFlags = 0;
+	const UChaosVDSettingsObjectBase* DebugDrawSettings = nullptr;
+	TSharedPtr<FChaosVDSolverDataSelection> SolverDataSelectionObject = nullptr;
+};
+
+/** Custom Hit Proxy for debug drawn particle data */
+struct HChaosVDComponentVisProxy : public HComponentVisProxy
+{
+	DECLARE_HIT_PROXY()
+
+	HChaosVDComponentVisProxy(const UActorComponent* InComponent, const TSharedPtr<FChaosVDSolverDataSelectionHandle>& InDataSelectionHandle)
+		: HComponentVisProxy(InComponent), DataSelectionHandle(InDataSelectionHandle)
+	{
+	}
+
+	virtual EMouseCursor::Type GetMouseCursor() override
+	{
+		return EMouseCursor::Crosshairs;
+	}
+
+	TSharedPtr<FChaosVDSolverDataSelectionHandle> DataSelectionHandle;
+};
+
+/** Base class used for all component visualizers in CVD - It provides a common code to handle selection and clicks*/
+class FChaosVDComponentVisualizerBase : public FComponentVisualizer
+{
+public:
+
+	virtual bool VisProxyHandleClick(FEditorViewportClient* InViewportClient, HComponentVisProxy* VisProxy, const FViewportClick& Click) override;
+
+protected:
+	
+	virtual void RegisterVisualizerMenus() = 0;
+	
+	virtual bool CanHandleClick(const HChaosVDComponentVisProxy& VisProxy);
+
+	virtual bool SelectVisualizedData(const HChaosVDComponentVisProxy& VisProxy, const TSharedRef<FChaosVDScene>& InCVDScene);
+
+	template<typename ObjectSettingsType, typename VisualizationFlagsType>
+	void CreateGenericVisualizerMenu(FName MenuToExtend, FName SectionName, const FText& InSectionLabel, const FText& InFlagsMenuLabel,  const FText& InFlagsMenuTooltip, FSlateIcon FlagsMenuIcon,  const FText& InSettingsMenuLabel, const FText& InSettingsMenuTooltip);
+
+	FName InspectorTabID = NAME_None;
+};
+
+template <typename ObjectSettingsType, typename VisualizationFlagsType>
+void FChaosVDComponentVisualizerBase::CreateGenericVisualizerMenu(FName MenuToExtend, FName SectionName, const FText& InSectionLabel, const FText& InFlagsMenuLabel,  const FText& InFlagsMenuTooltip, FSlateIcon FlagsMenuIcon,  const FText& InSettingsMenuLabel, const FText& InSettingsMenuTooltip)
+{
+	UToolMenus* ToolMenus = UToolMenus::Get();
+	
+	if (!ensure(ToolMenus))
+	{
+		return;
+	}
+
+	if (UToolMenu* Menu = ToolMenus->ExtendMenu(MenuToExtend))
+	{
+		FToolMenuSection& Section = Menu->AddSection(SectionName, InSectionLabel);
+		
+		Section.AddSubMenu(FName(InSectionLabel.ToString()), InFlagsMenuLabel, InFlagsMenuTooltip, FNewToolMenuDelegate::CreateLambda([](UToolMenu* Menu)
+						   {
+							   TSharedRef<SWidget> VisualizationFlagsWidget = SNew(SChaosVDEnumFlagsMenu<VisualizationFlagsType>)
+								   .CurrentValue_Static(&ObjectSettingsType::GetDataVisualizationFlags)
+								   .OnEnumSelectionChanged_Static(&ObjectSettingsType::SetDataVisualizationFlags);
+			
+							   FToolMenuEntry FlagsMenuEntry = FToolMenuEntry::InitWidget("VisualizationFlags", VisualizationFlagsWidget,FText::GetEmpty());
+							   Menu->AddMenuEntry(NAME_None, FlagsMenuEntry);
+						   }),
+						   false, FlagsMenuIcon);
+
+		using namespace Chaos::VisualDebugger::Utils;
+		Section.AddSubMenu(FName(InSettingsMenuLabel.ToString()), InSettingsMenuLabel, InSettingsMenuTooltip, FNewToolMenuDelegate::CreateStatic(&CreateMenuEntryForDefaultObject<ObjectSettingsType>, EChaosVDSaveSettingsOptions::ShowSaveButton),
+						   false, FSlateIcon(FAppStyle::Get().GetStyleSetName(), TEXT("Icons.Toolbar.Settings")));
+	}
+}
diff --git a/Engine/Plugins/ChaosVD/Source/ChaosVD/Private/Visualizers/ChaosVDJointConstraintsDataComponentVisualizer.cpp b/Engine/Plugins/ChaosVD/Source/ChaosVD/Private/Visualizers/ChaosVDJointConstraintsDataComponentVisualizer.cpp
--- a/Engine/Plugins/ChaosVD/Source/ChaosVD/Private/Visualizers/ChaosVDJointConstraintsDataComponentVisualizer.cpp
+++ b/Engine/Plugins/ChaosVD/Source/ChaosVD/Private/Visualizers/ChaosVDJointConstraintsDataComponentVisualizer.cpp
@@ -2,10 +2,10 @@
 
 #include "Visualizers/ChaosVDJointConstraintsDataComponentVisualizer.h"
 
+#include "ChaosVDEngine.h"
 #include "ChaosVDScene.h"
 #include "ChaosVDStyle.h"
 #include "ChaosVDTabsIDs.h"
-#include "EditorModeManager.h"
 #include "EditorViewportClient.h"
 #include "SceneManagement.h"
 #include "SceneView.h"
@@ -16,13 +16,8 @@
 #include "ToolMenus.h"
 #include "ToolMenuEntry.h"
 #include "ToolMenuSection.h"
-#include "Utils/ChaosVDUserInterfaceUtils.h"
 #include "Visualizers/ChaosVDDebugDrawUtils.h"
-#include "Widgets/SChaosVDMainTab.h"
 #include "Widgets/SChaosVDViewportToolbar.h"
-#include "Widgets/SChaosVDEnumFlagsMenu.h"
-
-IMPLEMENT_HIT_PROXY(HChaosVDJointConstraintProxy, HComponentVisProxy)
 
 #define LOCTEXT_NAMESPACE "ChaosVisualDebugger"
 
@@ -110,40 +105,25 @@ namespace Chaos::VisualDebugger::Utils
 
 FChaosVDJointConstraintsDataComponentVisualizer::FChaosVDJointConstraintsDataComponentVisualizer()
 {
-	RegisterVisualizerMenus();
+	FChaosVDJointConstraintsDataComponentVisualizer::RegisterVisualizerMenus();
+
+	InspectorTabID = FChaosVDTabID::ConstraintsInspector;
 }
 
 void FChaosVDJointConstraintsDataComponentVisualizer::RegisterVisualizerMenus()
 {
-	UToolMenus* ToolMenus = UToolMenus::Get();
+	FName MenuSection("JointConstraintDataVisualization.Show");
+	FText MenuSectionLabel = LOCTEXT("JointConstraintDataVisualizationShowMenuLabel", "Joint Constraint Data Visualization");
+	FText FlagsMenuLabel = LOCTEXT("JointConstraintDataVisualizationFlagsMenuLabel", "Joint Constraints Data Flags");
+	FText FlagsMenuTooltip = LOCTEXT("JointConstraintDataVisualizationFlagsMenuToolTip", "Set of flags to enable/disable visibility of specific types of joint constraint data");
+	FSlateIcon FlagsMenuIcon = FSlateIcon(FChaosVDStyle::Get().GetStyleSetName(), TEXT("ConnectionIcon"));
+
+	FText SettingsMenuLabel = LOCTEXT("JointConstraintDataVisualizationMenuLabel", "Joint Constraint Visualization Settings");
+	FText SettingsMenuTooltip = LOCTEXT("JointConstraintDataVisualizationMenuToolTip", "Options to change how the recorded joint constraint data is debug drawn");
 	
-	if (!ensure(ToolMenus))
-	{
-		return;
-	}
-
-	if (UToolMenu* Menu = ToolMenus->ExtendMenu(SChaosVDViewportToolbar::ShowMenuName))
-	{
-		FToolMenuSection& Section = Menu->AddSection("JointConstraintDataVisualization.Show", LOCTEXT("JointConstraintDataVisualizationShowMenuLabel", "Joint Constraint Data Visualization"));
-		
-		Section.AddSubMenu(TEXT("JointConstraintDataVisualizationFlags"), LOCTEXT("JointConstraintDataVisualizationFlagsMenuLabel", "Joint Constraints Data Flags"), LOCTEXT("JointConstraintDataVisualizationFlagsMenuToolTip", "Set of flags to enable/disable visibility of specific types of joint constraint data"), FNewToolMenuDelegate::CreateLambda([](UToolMenu* Menu)
-		                   {
-			                   TSharedRef<SWidget> VisualizationFlagsWidget = SNew(SChaosVDEnumFlagsMenu<EChaosVDJointsDataVisualizationFlags>)
-				                   .CurrentValue_Static(&UChaosVDJointConstraintsVisualizationSettings::GetJointsDataVisualizationFlags)
-				                   .OnEnumSelectionChanged_Static(&UChaosVDJointConstraintsVisualizationSettings::SetJointsDataVisualizationFlags);
-			
-			                   FToolMenuEntry FlagsMenuEntry = FToolMenuEntry::InitWidget("JointConstraintDataVisualizationFlags", VisualizationFlagsWidget,FText::GetEmpty());
-			                   Menu->AddMenuEntry(NAME_None, FlagsMenuEntry);
-		                   }),
-		                   false, FSlateIcon(FChaosVDStyle::Get().GetStyleSetName(), TEXT("ConnectionIcon")));
-
-		using namespace Chaos::VisualDebugger::Utils;
-		Section.AddSubMenu(TEXT("JointConstraintDataVisualizationSettings"), LOCTEXT("JointConstraintDataVisualizationMenuLabel", "Joint Constraint Visualization Settings"), LOCTEXT("JointConstraintDataVisualizationMenuToolTip", "Options to change how the recorded joint constraint data is debug drawn"), FNewToolMenuDelegate::CreateStatic(&CreateMenuEntryForDefaultObject<UChaosVDJointConstraintsVisualizationSettings>, EChaosVDSaveSettingsOptions::ShowSaveButton),
-		                   false, FSlateIcon(FAppStyle::Get().GetStyleSetName(), TEXT("Icons.Toolbar.Settings")));
-	}
+	CreateGenericVisualizerMenu<UChaosVDJointConstraintsVisualizationSettings, EChaosVDJointsDataVisualizationFlags>(SChaosVDViewportToolbar::ShowMenuName, MenuSection, MenuSectionLabel, FlagsMenuLabel, FlagsMenuTooltip, FlagsMenuIcon, SettingsMenuLabel, SettingsMenuTooltip);
 }
 
-
 void FChaosVDJointConstraintsDataComponentVisualizer::DrawVisualization(const UActorComponent* Component, const FSceneView* View, FPrimitiveDrawInterface* PDI)
 {
 	const UChaosVDSolverJointConstraintDataComponent* JointConstraintDataComponent = Cast<UChaosVDSolverJointConstraintDataComponent>(Component);
@@ -182,7 +162,7 @@ void FChaosVDJointConstraintsDataComponentVisualizer::DrawVisualization(const UA
 
 	if (const UChaosVDJointConstraintsVisualizationSettings* EditorSettings = GetDefault<UChaosVDJointConstraintsVisualizationSettings>())
 	{
-		VisualizationContext.VisualizationFlags = static_cast<uint32>(UChaosVDJointConstraintsVisualizationSettings::GetJointsDataVisualizationFlags());
+		VisualizationContext.VisualizationFlags = static_cast<uint32>(UChaosVDJointConstraintsVisualizationSettings::GetDataVisualizationFlags());
 		VisualizationContext.bShowDebugText = EditorSettings->bShowDebugText;
 		VisualizationContext.DebugDrawSettings = EditorSettings;
 	}
@@ -192,59 +172,49 @@ void FChaosVDJointConstraintsDataComponentVisualizer::DrawVisualization(const UA
 		return;
 	}
 
-	// If nothing is selected, fallback to draw all joints
-	const bool bDrawOnlySelected = VisualizationContext.IsVisualizationFlagEnabled(EChaosVDJointsDataVisualizationFlags::OnlyDrawSelected) && JointConstraintDataComponent->GetCurrentSelectionHandle().IsSelected();
+	TSharedPtr<FChaosVDSolverDataSelection> SelectionObject = CVDScene->GetSolverDataSelectionObject().Pin();
+	TSharedPtr<FChaosVDSolverDataSelectionHandle> SelectionHandle;
+	if (SelectionObject)
+	{
+		SelectionHandle = SelectionObject->GetCurrentSelectionHandle();
+	}
+	
+	// If nothing is selected, fallback to draw all 
+	const bool bDrawOnlySelected = VisualizationContext.IsVisualizationFlagEnabled(EChaosVDJointsDataVisualizationFlags::OnlyDrawSelected) && SelectionHandle && SelectionHandle->IsSelected();
 	if (bDrawOnlySelected)
 	{
-		if (const TSharedPtr<FChaosVDJointConstraint> JointConstraint = JointConstraintDataComponent->GetCurrentSelectionHandle().GetData().Pin())
+		if (FChaosVDJointConstraint* JointConstraintData = SelectionHandle->GetData<FChaosVDJointConstraint>())
 		{
-			VisualizationContext.DataSelectionHandle = FChaosVDJointConstraintSelectionHandle(JointConstraint);
-			DrawJointConstraint(Component, *JointConstraint, VisualizationContext, View, PDI);
+			VisualizationContext.DataSelectionHandle = SelectionHandle;
+			DrawJointConstraint(Component, *JointConstraintData, VisualizationContext, View, PDI);
 		}
 	}
 	else
 	{
-		for (const TSharedPtr<FChaosVDJointConstraint>& JointConstraint : JointConstraintDataComponent->GetAllJointConstraints())
+		for (const TSharedPtr<FChaosVDConstraintDataWrapperBase>& Constraint : JointConstraintDataComponent->GetAllConstraints())
 		{
-			if (JointConstraint)
+			static_assert(std::is_base_of_v<FChaosVDConstraintDataWrapperBase, FChaosVDJointConstraint>, "Only FChaosVDJointConstraint is supported");
+
+			if (FChaosVDJointConstraint* JointConstraint = static_cast<FChaosVDJointConstraint*>(Constraint.Get()))
 			{
-				VisualizationContext.DataSelectionHandle = FChaosVDJointConstraintSelectionHandle(JointConstraint);
+				VisualizationContext.DataSelectionHandle = SelectionObject->MakeSelectionHandle(StaticCastSharedPtr<FChaosVDJointConstraint>(Constraint));
+
 				DrawJointConstraint(Component, *JointConstraint, VisualizationContext, View, PDI);
 			}
 		}
 	}
 }
 
-bool FChaosVDJointConstraintsDataComponentVisualizer::VisProxyHandleClick(FEditorViewportClient* InViewportClient, HComponentVisProxy* VisProxy, const FViewportClick& Click)
+bool FChaosVDJointConstraintsDataComponentVisualizer::CanHandleClick(const HChaosVDComponentVisProxy& VisProxy)
 {
-	const HChaosVDJointConstraintProxy* JointDataProxy = HitProxyCast<HChaosVDJointConstraintProxy>(VisProxy);
-	if (JointDataProxy == nullptr)
-	{
-		return false;
-	}
-	
-	if (const UChaosVDSolverJointConstraintDataComponent* JointsDataComponent = Cast<UChaosVDSolverJointConstraintDataComponent>(JointDataProxy->Component.Get()))
-	{
-		// Bring the SQ Inspector into focus if available
-		const TSharedPtr<SChaosVDMainTab> MainTabToolkitHost = InViewportClient->GetModeTools() ? StaticCastSharedPtr<SChaosVDMainTab>(InViewportClient->GetModeTools()->GetToolkitHost()) : nullptr;
-		if (const TSharedPtr<FTabManager> TabManager = MainTabToolkitHost ? MainTabToolkitHost->GetTabManager() : nullptr)
-		{
-			TabManager->TryInvokeTab(FChaosVDTabID::JointsDataDetails);
-		}
-
-		const_cast<UChaosVDSolverJointConstraintDataComponent*>(JointsDataComponent)->SelectJoint(JointDataProxy->DataSelectionHandle);
-
-		return true;
-	}
-
-	return false;
+	return VisProxy.DataSelectionHandle && VisProxy.DataSelectionHandle->IsA<FChaosVDJointConstraint>();
 }
 
-void FChaosVDJointConstraintsDataComponentVisualizer::DebugDrawAllAxis(const FChaosVDJointConstraint& InJointConstraintData, FChaosVDJointVisualizationDataContext& VisualizationContext, FPrimitiveDrawInterface* PDI, const float LineThickness, const FVector& InPosition, const Chaos::FMatrix33& InRotationMatrix, TConstArrayView<FLinearColor> AxisColors)
+void FChaosVDJointConstraintsDataComponentVisualizer::DebugDrawAllAxis(const FChaosVDJointConstraint& InJointConstraintData, FChaosVDJointVisualizationDataContext& VisualizationContext, FPrimitiveDrawInterface* PDI, const float LineThickness, const FVector& InPosition, const Chaos::FMatrix33& InRotationMatrix, TConstArrayView<FLinearColor> AxisColors,  bool bIsSelected)
 {
 	for(int32 AxisIndex = 0; AxisIndex < 3 ; AxisIndex++)
 	{
-		FChaosVDDebugDrawUtils::DrawArrowVector(PDI, InPosition, InPosition + VisualizationContext.DebugDrawSettings->GeneralScale * VisualizationContext.DebugDrawSettings->ConstraintAxisLength * VisualizationContext.SpaceTransform.TransformVector(InRotationMatrix.GetAxis(AxisIndex)), FText::GetEmpty(), Chaos::VisualDebugger::Utils::GenerateSelectionAwareDebugColor(AxisColors[AxisIndex], InJointConstraintData.bIsSelectedInEditor), VisualizationContext.DebugDrawSettings->DepthPriority, LineThickness * 0.2f);
+		FChaosVDDebugDrawUtils::DrawArrowVector(PDI, InPosition, InPosition + VisualizationContext.DebugDrawSettings->GeneralScale * VisualizationContext.DebugDrawSettings->ConstraintAxisLength * VisualizationContext.SpaceTransform.TransformVector(InRotationMatrix.GetAxis(AxisIndex)), FText::GetEmpty(), Chaos::VisualDebugger::Utils::GenerateSelectionAwareDebugColor(AxisColors[AxisIndex], bIsSelected), VisualizationContext.DebugDrawSettings->DepthPriority, LineThickness * 0.2f);
 	}
 }
 
@@ -280,8 +250,8 @@ void FChaosVDJointConstraintsDataComponentVisualizer::DrawJointConstraint(const
 		return;
 	}
 
-	const FChaosVDParticleDataWrapper* ParticleData0 = nullptr;
-	const FChaosVDParticleDataWrapper* ParticleData1 = nullptr;
+	TSharedPtr<const FChaosVDParticleDataWrapper> ParticleData0 = nullptr;
+	TSharedPtr<const FChaosVDParticleDataWrapper> ParticleData1 = nullptr;
 
 	if (AChaosVDParticleActor* Particle0 = VisualizationContext.SolverInfoActor->GetParticleActor(InJointConstraintData.ParticleParIndexes[0]))
 	{
@@ -324,9 +294,11 @@ void FChaosVDJointConstraintsDataComponentVisualizer::DrawJointConstraint(const
 		return;
 	}
 
-	PDI->SetHitProxy(new HChaosVDJointConstraintProxy(Component, VisualizationContext.DataSelectionHandle));
+	bool bIsSelected = VisualizationContext.DataSelectionHandle && VisualizationContext.DataSelectionHandle->IsSelected();
+
+	PDI->SetHitProxy(new HChaosVDComponentVisProxy(Component, VisualizationContext.DataSelectionHandle));
 
-	const float LineThickness = InJointConstraintData.bIsSelectedInEditor ?  VisualizationContext.DebugDrawSettings->BaseLineThickness * 1.5f :  VisualizationContext.DebugDrawSettings->BaseLineThickness;
+	const float LineThickness = bIsSelected ? VisualizationContext.DebugDrawSettings->BaseLineThickness * 1.5f :  VisualizationContext.DebugDrawSettings->BaseLineThickness;
 
 	FVector Pa = ParticleData1->ParticlePositionRotation.MX;
 	FVector Pb = ParticleData0->ParticlePositionRotation.MX;
@@ -370,10 +342,10 @@ void FChaosVDJointConstraintsDataComponentVisualizer::DrawJointConstraint(const
 			Sb = FMath::Lerp(Pb, Xb, FMath::Clamp<double>(CoMSize / Lena, 0., 1.));
 		}
 
-		FChaosVDDebugDrawUtils::DrawLine(PDI, Pa, Sa, GenerateSelectionAwareDebugColor(White, InJointConstraintData.bIsSelectedInEditor), FText::GetEmpty(), VisualizationContext.DebugDrawSettings->DepthPriority, ConnectorThickness);
-		FChaosVDDebugDrawUtils::DrawLine(PDI, Pb, Sb, GenerateSelectionAwareDebugColor(White, InJointConstraintData.bIsSelectedInEditor), FText::GetEmpty(), VisualizationContext.DebugDrawSettings->DepthPriority, ConnectorThickness);
-		FChaosVDDebugDrawUtils::DrawLine(PDI, Sa, Xa, GenerateSelectionAwareDebugColor(Red, InJointConstraintData.bIsSelectedInEditor), FText::GetEmpty(), VisualizationContext.DebugDrawSettings->DepthPriority, ConnectorThickness);
-		FChaosVDDebugDrawUtils::DrawLine(PDI, Sb, Xb, GenerateSelectionAwareDebugColor(Cyan, InJointConstraintData.bIsSelectedInEditor), FText::GetEmpty(), VisualizationContext.DebugDrawSettings->DepthPriority, ConnectorThickness);
+		FChaosVDDebugDrawUtils::DrawLine(PDI, Pa, Sa, GenerateSelectionAwareDebugColor(White, bIsSelected), FText::GetEmpty(), VisualizationContext.DebugDrawSettings->DepthPriority, ConnectorThickness);
+		FChaosVDDebugDrawUtils::DrawLine(PDI, Pb, Sb, GenerateSelectionAwareDebugColor(White,bIsSelected), FText::GetEmpty(), VisualizationContext.DebugDrawSettings->DepthPriority, ConnectorThickness);
+		FChaosVDDebugDrawUtils::DrawLine(PDI, Sa, Xa, GenerateSelectionAwareDebugColor(Red, bIsSelected), FText::GetEmpty(), VisualizationContext.DebugDrawSettings->DepthPriority, ConnectorThickness);
+		FChaosVDDebugDrawUtils::DrawLine(PDI, Sb, Xb, GenerateSelectionAwareDebugColor(Cyan, bIsSelected), FText::GetEmpty(), VisualizationContext.DebugDrawSettings->DepthPriority, ConnectorThickness);
 	}
 
 	if (VisualizationContext.IsVisualizationFlagEnabled(EChaosVDJointsDataVisualizationFlags::CenterOfMassConnector))
@@ -394,31 +366,31 @@ void FChaosVDJointConstraintsDataComponentVisualizer::DrawJointConstraint(const
 			Sb = FMath::Lerp(Cb, Xb, FMath::Clamp<double>(CoMSize / Lena, 0., 1.));
 		}
 
-		FChaosVDDebugDrawUtils::DrawLine(PDI, Ca, Sa, GenerateSelectionAwareDebugColor(Black, InJointConstraintData.bIsSelectedInEditor), FText::GetEmpty(), VisualizationContext.DebugDrawSettings->DepthPriority, ConnectorThickness);
-		FChaosVDDebugDrawUtils::DrawLine(PDI, Cb, Sb, GenerateSelectionAwareDebugColor(Black, InJointConstraintData.bIsSelectedInEditor), FText::GetEmpty(), VisualizationContext.DebugDrawSettings->DepthPriority, ConnectorThickness);
-		FChaosVDDebugDrawUtils::DrawLine(PDI, Sa, Xa, GenerateSelectionAwareDebugColor(Red, InJointConstraintData.bIsSelectedInEditor), FText::GetEmpty(), VisualizationContext.DebugDrawSettings->DepthPriority, ConnectorThickness);
-		FChaosVDDebugDrawUtils::DrawLine(PDI, Sb, Xb, GenerateSelectionAwareDebugColor(Cyan, InJointConstraintData.bIsSelectedInEditor), FText::GetEmpty(), VisualizationContext.DebugDrawSettings->DepthPriority, ConnectorThickness);
+		FChaosVDDebugDrawUtils::DrawLine(PDI, Ca, Sa, GenerateSelectionAwareDebugColor(Black, bIsSelected), FText::GetEmpty(), VisualizationContext.DebugDrawSettings->DepthPriority, ConnectorThickness);
+		FChaosVDDebugDrawUtils::DrawLine(PDI, Cb, Sb, GenerateSelectionAwareDebugColor(Black, bIsSelected), FText::GetEmpty(), VisualizationContext.DebugDrawSettings->DepthPriority, ConnectorThickness);
+		FChaosVDDebugDrawUtils::DrawLine(PDI, Sa, Xa, GenerateSelectionAwareDebugColor(Red, bIsSelected), FText::GetEmpty(), VisualizationContext.DebugDrawSettings->DepthPriority, ConnectorThickness);
+		FChaosVDDebugDrawUtils::DrawLine(PDI, Sb, Xb, GenerateSelectionAwareDebugColor(Cyan, bIsSelected), FText::GetEmpty(), VisualizationContext.DebugDrawSettings->DepthPriority, ConnectorThickness);
 	}
 
 	if (VisualizationContext.IsVisualizationFlagEnabled(EChaosVDJointsDataVisualizationFlags::Stretch))
 	{
 		const float StretchThickness = 3.0f * LineThickness;
-		FChaosVDDebugDrawUtils::DrawLine(PDI, Xa, Xb, GenerateSelectionAwareDebugColor(Magenta, InJointConstraintData.bIsSelectedInEditor), FText::GetEmpty(), VisualizationContext.DebugDrawSettings->DepthPriority, StretchThickness);
+		FChaosVDDebugDrawUtils::DrawLine(PDI, Xa, Xb, GenerateSelectionAwareDebugColor(Magenta, bIsSelected), FText::GetEmpty(), VisualizationContext.DebugDrawSettings->DepthPriority, StretchThickness);
 	}
 
 	if (VisualizationContext.IsVisualizationFlagEnabled(EChaosVDJointsDataVisualizationFlags::Axes))
 	{
 		const FLinearColor AxisAColors[3] = {Red, Green, Blue};
 		const FLinearColor AxisBColors[3] = {Cyan, Magenta, Yellow};
-		DebugDrawAllAxis(InJointConstraintData, VisualizationContext, PDI, LineThickness, Xa, Ra, AxisAColors);
-		DebugDrawAllAxis(InJointConstraintData, VisualizationContext, PDI, LineThickness, Xb, Rb, AxisBColors);
+		DebugDrawAllAxis(InJointConstraintData, VisualizationContext, PDI, LineThickness, Xa, Ra, AxisAColors, bIsSelected);
+		DebugDrawAllAxis(InJointConstraintData, VisualizationContext, PDI, LineThickness, Xb, Rb, AxisBColors, bIsSelected);
 	}
 
 	// NOTE: GetLinearImpulse is the positional impulse (pushout)
 	if (VisualizationContext.IsVisualizationFlagEnabled(EChaosVDJointsDataVisualizationFlags::PushOut))
 	{
 		FLinearColor PushOutImpulseColor = FColor(0, 250, 250);
-		FChaosVDDebugDrawUtils::DrawLine(PDI, Xa, Xa + VisualizationContext.DebugDrawSettings->LinearImpulseScale * VisualizationContext.SpaceTransform.TransformVectorNoScale(InJointConstraintData.JointState.LinearImpulse), GenerateSelectionAwareDebugColor(PushOutImpulseColor, InJointConstraintData.bIsSelectedInEditor), FText::GetEmpty(), VisualizationContext.DebugDrawSettings->DepthPriority, LineThickness);
+		FChaosVDDebugDrawUtils::DrawLine(PDI, Xa, Xa + VisualizationContext.DebugDrawSettings->LinearImpulseScale * VisualizationContext.SpaceTransform.TransformVectorNoScale(InJointConstraintData.JointState.LinearImpulse), GenerateSelectionAwareDebugColor(PushOutImpulseColor, bIsSelected), FText::GetEmpty(), VisualizationContext.DebugDrawSettings->DepthPriority, LineThickness);
 	}
 
 	//TODO: Should we draw the Angular Impulse
diff --git a/Engine/Plugins/ChaosVD/Source/ChaosVD/Private/Visualizers/ChaosVDJointConstraintsDataComponentVisualizer.h b/Engine/Plugins/ChaosVD/Source/ChaosVD/Private/Visualizers/ChaosVDJointConstraintsDataComponentVisualizer.h
--- a/Engine/Plugins/ChaosVD/Source/ChaosVD/Private/Visualizers/ChaosVDJointConstraintsDataComponentVisualizer.h
+++ b/Engine/Plugins/ChaosVD/Source/ChaosVD/Private/Visualizers/ChaosVDJointConstraintsDataComponentVisualizer.h
@@ -1,10 +1,10 @@
 ﻿// Copyright Epic Games, Inc. All Rights Reserved.
 
 #pragma once
-#include "ComponentVisualizer.h"
+#include "ChaosVDComponentVisualizerBase.h"
+#include "ChaosVDSolverDataSelection.h"
 #include "Chaos/Core.h"
 #include "Components/ChaosVDSolverJointConstraintDataComponent.h"
-#include "IChaosVDParticleVisualizationDataProvider.h"
 #include "Settings/ChaosVDJointConstraintVisualizationSettings.h"
 
 struct FChaosVDJointConstraint;
@@ -15,7 +15,8 @@ class UChaosVDJointConstraintsVisualizationSettings;
 /** Visualization context structure specific for Joints visualizations */
 struct FChaosVDJointVisualizationDataContext : public FChaosVDVisualizationContext
 {
-	FChaosVDJointConstraintSelectionHandle DataSelectionHandle = FChaosVDJointConstraintSelectionHandle(nullptr);
+	TSharedPtr<FChaosVDSolverDataSelectionHandle> DataSelectionHandle = MakeShared<FChaosVDSolverDataSelectionHandle>();
+
 	bool bIsServerVisualizationEnabled = false;
 	
 	AChaosVDSolverInfoActor* SolverInfoActor = nullptr;
@@ -31,38 +32,22 @@ struct FChaosVDJointVisualizationDataContext : public FChaosVDVisualizationConte
 	}
 };
 
-/** Custom Hit Proxy for debug drawn scene queries */
-struct HChaosVDJointConstraintProxy : public HComponentVisProxy
-{
-	DECLARE_HIT_PROXY()
-	
-	HChaosVDJointConstraintProxy(const UActorComponent* Component, const FChaosVDJointConstraintSelectionHandle& InContactFinderData) : HComponentVisProxy(Component, HPP_UI), DataSelectionHandle(InContactFinderData)
-	{	
-	}
-
-	virtual EMouseCursor::Type GetMouseCursor() override
-	{
-		return EMouseCursor::Crosshairs;
-	}
-
-	FChaosVDJointConstraintSelectionHandle DataSelectionHandle;
-};
-
 /**
  * Component visualizer in charge of generating debug draw visualizations for Joint Constraints in a UChaosVDSolverJointConstraintDataComponent
  */
-class FChaosVDJointConstraintsDataComponentVisualizer final : public FComponentVisualizer
+class FChaosVDJointConstraintsDataComponentVisualizer final : public FChaosVDComponentVisualizerBase
 {
 public:
 	FChaosVDJointConstraintsDataComponentVisualizer();
-	
-	void RegisterVisualizerMenus();
+
+	virtual void RegisterVisualizerMenus() override;
 
 	virtual void DrawVisualization(const UActorComponent* Component, const FSceneView* View, FPrimitiveDrawInterface* PDI) override;
-	virtual bool VisProxyHandleClick(FEditorViewportClient* InViewportClient, HComponentVisProxy* VisProxy, const FViewportClick& Click) override;
+
+	virtual bool CanHandleClick(const HChaosVDComponentVisProxy& VisProxy) override;
 
 protected:
 
-	void DebugDrawAllAxis(const FChaosVDJointConstraint& InJointConstraintData, FChaosVDJointVisualizationDataContext& VisualizationContext, FPrimitiveDrawInterface* PDI, const float LineThickness, const FVector& InPosition, const Chaos::FMatrix33& InRotationMatrix, TConstArrayView<FLinearColor>);
+	void DebugDrawAllAxis(const FChaosVDJointConstraint& InJointConstraintData, FChaosVDJointVisualizationDataContext& VisualizationContext, FPrimitiveDrawInterface* PDI, const float LineThickness, const FVector& InPosition, const Chaos::FMatrix33& InRotationMatrix, TConstArrayView<FLinearColor> AxisColors, bool bIsSelected);
 	void DrawJointConstraint(const UActorComponent* Component, const FChaosVDJointConstraint& InJointConstraintData, FChaosVDJointVisualizationDataContext& VisualizationContext, const FSceneView* View, FPrimitiveDrawInterface* PDI);
 };
diff --git a/Engine/Plugins/ChaosVD/Source/ChaosVD/Private/Visualizers/ChaosVDParticleDataComponentVisualizer.cpp b/Engine/Plugins/ChaosVD/Source/ChaosVD/Private/Visualizers/ChaosVDParticleDataComponentVisualizer.cpp
--- a/Engine/Plugins/ChaosVD/Source/ChaosVD/Private/Visualizers/ChaosVDParticleDataComponentVisualizer.cpp
+++ b/Engine/Plugins/ChaosVD/Source/ChaosVD/Private/Visualizers/ChaosVDParticleDataComponentVisualizer.cpp
@@ -5,6 +5,7 @@
 #include "Actors/ChaosVDSolverInfoActor.h"
 #include "ChaosVDParticleActor.h"
 #include "ChaosVDScene.h"
+#include "ChaosVDTabsIDs.h"
 #include "Components/ChaosVDParticleDataComponent.h"
 #include "DataWrappers/ChaosVDParticleDataWrapper.h"
 #include "SceneView.h"
@@ -18,8 +19,6 @@
 #include "Widgets/SChaosVDEnumFlagsMenu.h"
 #include "Widgets/SChaosVDViewportToolbar.h"
 
-IMPLEMENT_HIT_PROXY(HChaosVDParticleDataProxy, HComponentVisProxy)
-
 #define LOCTEXT_NAMESPACE "ChaosVisualDebugger"
 
 /** Sets a Hit proxy which will be cleared out as soon this struct goes out of scope*/
@@ -47,7 +46,9 @@ struct FChaosVDScopedParticleHitProxy
 
 FChaosVDParticleDataComponentVisualizer::FChaosVDParticleDataComponentVisualizer()
 {
-	RegisterVisualizerMenus();
+	FChaosVDParticleDataComponentVisualizer::RegisterVisualizerMenus();
+
+	InspectorTabID = FChaosVDTabID::DetailsPanel;
 }
 
 void FChaosVDParticleDataComponentVisualizer::RegisterVisualizerMenus()
@@ -152,6 +153,8 @@ void FChaosVDParticleDataComponentVisualizer::DrawVisualization(const UActorComp
 	VisualizationContext.GeometryGenerator = CVDScene->GetGeometryGenerator();
 	VisualizationContext.bShowDebugText = VisualizationSettings->bShowDebugText;
 	VisualizationContext.DebugDrawSettings = VisualizationSettings;
+	VisualizationContext.SolverDataSelectionObject = CVDScene->GetSolverDataSelectionObject().Pin();
+	
 
 	if (!VisualizationContext.IsVisualizationFlagEnabled(EChaosVDParticleDataVisualizationFlags::EnableDraw))
 	{
@@ -161,7 +164,7 @@ void FChaosVDParticleDataComponentVisualizer::DrawVisualization(const UActorComp
 	if (VisualizationContext.IsVisualizationFlagEnabled(EChaosVDParticleDataVisualizationFlags::DrawDataOnlyForSelectedParticle))
 	{
 		VisualizationContext.bIsSelectedData = true;
-		SolverDataActor->VisitSelectedParticleData([this, PDI, View, &VisualizationContext, Component](const FChaosVDParticleDataWrapper& InParticleDataViewer)
+		SolverDataActor->VisitSelectedParticleData([this, PDI, View, &VisualizationContext, Component](const TSharedPtr<const FChaosVDParticleDataWrapper>& InParticleDataViewer)
 		{
 			DrawVisualizationForParticleData(Component, PDI, View, VisualizationContext, InParticleDataViewer);
 
@@ -170,9 +173,9 @@ void FChaosVDParticleDataComponentVisualizer::DrawVisualization(const UActorComp
 	}
 	else
 	{
-		SolverDataActor->VisitAllParticleData([this, PDI, View, &VisualizationContext, Component, SolverDataActor](const FChaosVDParticleDataWrapper& InParticleDataViewer)
+		SolverDataActor->VisitAllParticleData([this, PDI, View, &VisualizationContext, Component, SolverDataActor](const TSharedPtr<const FChaosVDParticleDataWrapper>& InParticleDataViewer)
 		{
-			VisualizationContext.bIsSelectedData = SolverDataActor->IsParticleSelectedByID(InParticleDataViewer.ParticleIndex);
+			VisualizationContext.bIsSelectedData = InParticleDataViewer && SolverDataActor->IsParticleSelectedByID(InParticleDataViewer->ParticleIndex);
 			DrawVisualizationForParticleData(Component, PDI, View, VisualizationContext, InParticleDataViewer);
 
 			// If we reach the debug draw limit for this frame, there is no need to continue processing particles
@@ -181,17 +184,15 @@ void FChaosVDParticleDataComponentVisualizer::DrawVisualization(const UActorComp
 	}
 }
 
-bool FChaosVDParticleDataComponentVisualizer::VisProxyHandleClick(FEditorViewportClient* InViewportClient, HComponentVisProxy* VisProxy, const FViewportClick& Click)
+bool FChaosVDParticleDataComponentVisualizer::CanHandleClick(const HChaosVDComponentVisProxy& VisProxy)
 {
-	bool bHandled = false;
+	return VisProxy.DataSelectionHandle && VisProxy.DataSelectionHandle->IsA<FChaosVDParticleDataWrapper>();
+}
 
-	const HChaosVDParticleDataProxy* ParticleDataProxy = HitProxyCast<HChaosVDParticleDataProxy>(VisProxy);
-	if (!ParticleDataProxy)
-	{
-		return bHandled;
-	}
-	
-	const UChaosVDParticleDataComponent* ParticleDataComponent = Cast<UChaosVDParticleDataComponent>(VisProxy->Component.Get());
+bool FChaosVDParticleDataComponentVisualizer::SelectVisualizedData(const HChaosVDComponentVisProxy& VisProxy, const TSharedRef<FChaosVDScene>& InCVDScene)
+{
+	bool bHandled = false;
+	const UChaosVDParticleDataComponent* ParticleDataComponent = Cast<UChaosVDParticleDataComponent>(VisProxy.Component.Get());
 	if (!ParticleDataComponent)
 	{
 		return bHandled;
@@ -203,9 +204,12 @@ bool FChaosVDParticleDataComponentVisualizer::VisProxyHandleClick(FEditorViewpor
 		return bHandled;
 	}
 
-	bHandled = SolverDataActor->SelectParticleByID(ParticleDataProxy->DataSelectionHandle.ParticleIndex);
+	if (TSharedPtr<const FChaosVDParticleDataWrapper> ParticleDataViewer = VisProxy.DataSelectionHandle ?  VisProxy.DataSelectionHandle->GetDataAsShared<const FChaosVDParticleDataWrapper>() : nullptr)
+	{
+		bHandled = SolverDataActor->SelectParticleByID(ParticleDataViewer->ParticleIndex);
+	}
 
-	return bHandled;	
+	return bHandled;
 }
 
 void FChaosVDParticleDataComponentVisualizer::DrawParticleVector(FPrimitiveDrawInterface* PDI, const FVector& StartLocation, const FVector& InVector, EChaosVDParticleDataVisualizationFlags VectorID, const FChaosVDParticleDataVisualizationContext& InVisualizationContext, float LineThickness)
@@ -224,7 +228,7 @@ void FChaosVDParticleDataComponentVisualizer::DrawParticleVector(FPrimitiveDrawI
 	FChaosVDDebugDrawUtils::DrawArrowVector(PDI, StartLocation, StartLocation +  InVisualizationContext.DebugDrawSettings->GetScaleFortDataID(VectorID) * InVector, FText::AsCultureInvariant(DebugText), InVisualizationContext.DebugDrawSettings->ColorSettings.GetColorForDataID(VectorID, InVisualizationContext.bIsSelectedData),  InVisualizationContext.DebugDrawSettings->DepthPriority, LineThickness);
 }
 
-void FChaosVDParticleDataComponentVisualizer::DrawVisualizationForParticleData(const UActorComponent* Component, FPrimitiveDrawInterface* PDI, const FSceneView* View, const FChaosVDParticleDataVisualizationContext& InVisualizationContext, const FChaosVDParticleDataWrapper& InParticleDataViewer)
+void FChaosVDParticleDataComponentVisualizer::DrawVisualizationForParticleData(const UActorComponent* Component, FPrimitiveDrawInterface* PDI, const FSceneView* View, const FChaosVDParticleDataVisualizationContext& InVisualizationContext, const TSharedPtr<const FChaosVDParticleDataWrapper>& InParticleDataViewer)
 {
 	using namespace Chaos::VisualDebugger::ParticleDataUnitsStrings;
 
@@ -238,7 +242,12 @@ void FChaosVDParticleDataComponentVisualizer::DrawVisualizationForParticleData(c
 		return;
 	}
 
-	const FVector& OwnerLocation = InVisualizationContext.SpaceTransform.TransformPosition(InParticleDataViewer.ParticlePositionRotation.MX);
+	if (!ensure(InVisualizationContext.SolverDataSelectionObject))
+	{
+		return;
+	}
+
+	const FVector& OwnerLocation = InVisualizationContext.SpaceTransform.TransformPosition(InParticleDataViewer->ParticlePositionRotation.MX);
 
 	// TODO: See how expensive is get the bounds. It is not something we have recorded
 	constexpr float VisibleRadius = 50.0f;
@@ -248,31 +257,31 @@ void FChaosVDParticleDataComponentVisualizer::DrawVisualizationForParticleData(c
 		return;
 	}
 
-	const FQuat& OwnerRotation =  InVisualizationContext.SpaceTransform.TransformRotation(InParticleDataViewer.ParticlePositionRotation.MR);
-	const FVector OwnerCoMLocation = InVisualizationContext.SpaceTransform.TransformPosition(InParticleDataViewer.ParticlePositionRotation.MX + (InParticleDataViewer.ParticlePositionRotation.MR *  InParticleDataViewer.ParticleMassProps.MCenterOfMass));
+	const FQuat& OwnerRotation =  InVisualizationContext.SpaceTransform.TransformRotation(InParticleDataViewer->ParticlePositionRotation.MR);
+	const FVector OwnerCoMLocation = InVisualizationContext.SpaceTransform.TransformPosition(InParticleDataViewer->ParticlePositionRotation.MX + (InParticleDataViewer->ParticlePositionRotation.MR *  InParticleDataViewer->ParticleMassProps.MCenterOfMass));
 	
-	FChaosVDScopedParticleHitProxy ScopedHitProxy(PDI, new HChaosVDParticleDataProxy(Component, {InParticleDataViewer.ParticleIndex, InParticleDataViewer.SolverID} ));
+	FChaosVDScopedParticleHitProxy ScopedHitProxy(PDI, new HChaosVDComponentVisProxy(Component, InVisualizationContext.SolverDataSelectionObject->MakeSelectionHandle(ConstCastSharedPtr<FChaosVDParticleDataWrapper>(InParticleDataViewer))));
 
 	constexpr float DefaultLineThickness = 1.5f;
 	constexpr float SelectedLineThickness = 3.5f;
 	const float LineThickness = InVisualizationContext.bIsSelectedData ? SelectedLineThickness : DefaultLineThickness;
 	
 
-	if (InParticleDataViewer.ParticleVelocities.HasValidData())
+	if (InParticleDataViewer->ParticleVelocities.HasValidData())
 	{
-		DrawParticleVector(PDI, OwnerCoMLocation, InParticleDataViewer.ParticleVelocities.MV,EChaosVDParticleDataVisualizationFlags::Velocity, InVisualizationContext, LineThickness); 
-		DrawParticleVector(PDI, OwnerCoMLocation, InParticleDataViewer.ParticleVelocities.MW,EChaosVDParticleDataVisualizationFlags::AngularVelocity, InVisualizationContext, LineThickness); 
+		DrawParticleVector(PDI, OwnerCoMLocation, InParticleDataViewer->ParticleVelocities.MV,EChaosVDParticleDataVisualizationFlags::Velocity, InVisualizationContext, LineThickness); 
+		DrawParticleVector(PDI, OwnerCoMLocation, InParticleDataViewer->ParticleVelocities.MW,EChaosVDParticleDataVisualizationFlags::AngularVelocity, InVisualizationContext, LineThickness); 
 	}
 
-	if (InParticleDataViewer.ParticleDynamics.HasValidData())
+	if (InParticleDataViewer->ParticleDynamics.HasValidData())
 	{
-		DrawParticleVector(PDI, OwnerCoMLocation, InParticleDataViewer.ParticleDynamics.MAcceleration,EChaosVDParticleDataVisualizationFlags::Acceleration, InVisualizationContext, LineThickness); 
-		DrawParticleVector(PDI, OwnerCoMLocation, InParticleDataViewer.ParticleDynamics.MAngularAcceleration,EChaosVDParticleDataVisualizationFlags::AngularAcceleration, InVisualizationContext, LineThickness); 
-		DrawParticleVector(PDI, OwnerCoMLocation, InParticleDataViewer.ParticleDynamics.MLinearImpulseVelocity,EChaosVDParticleDataVisualizationFlags::LinearImpulse, InVisualizationContext, LineThickness); 
-		DrawParticleVector(PDI, OwnerCoMLocation, InParticleDataViewer.ParticleDynamics.MAngularImpulseVelocity,EChaosVDParticleDataVisualizationFlags::AngularImpulse, InVisualizationContext, LineThickness);
+		DrawParticleVector(PDI, OwnerCoMLocation, InParticleDataViewer->ParticleDynamics.MAcceleration,EChaosVDParticleDataVisualizationFlags::Acceleration, InVisualizationContext, LineThickness); 
+		DrawParticleVector(PDI, OwnerCoMLocation, InParticleDataViewer->ParticleDynamics.MAngularAcceleration,EChaosVDParticleDataVisualizationFlags::AngularAcceleration, InVisualizationContext, LineThickness); 
+		DrawParticleVector(PDI, OwnerCoMLocation, InParticleDataViewer->ParticleDynamics.MLinearImpulseVelocity,EChaosVDParticleDataVisualizationFlags::LinearImpulse, InVisualizationContext, LineThickness); 
+		DrawParticleVector(PDI, OwnerCoMLocation, InParticleDataViewer->ParticleDynamics.MAngularImpulseVelocity,EChaosVDParticleDataVisualizationFlags::AngularImpulse, InVisualizationContext, LineThickness);
 	}
 
-	if (InParticleDataViewer.ParticleMassProps.HasValidData())
+	if (InParticleDataViewer->ParticleMassProps.HasValidData())
 	{
 		if (InVisualizationContext.IsVisualizationFlagEnabled(EChaosVDParticleDataVisualizationFlags::CenterOfMass))
 		{
@@ -288,17 +297,17 @@ void FChaosVDParticleDataComponentVisualizer::DrawVisualizationForParticleData(c
 	}
 
 	// TODO: This is a Proof of concept to test how debug draw connectivity data will look
-	if (InParticleDataViewer.ParticleCluster.HasValidData())
+	if (InParticleDataViewer->ParticleCluster.HasValidData())
 	{
 		if (InVisualizationContext.IsVisualizationFlagEnabled(EChaosVDParticleDataVisualizationFlags::ClusterConnectivityEdge))
 		{
-			for (const FChaosVDConnectivityEdge& ConnectivityEdge : InParticleDataViewer.ParticleCluster.ConnectivityEdges)
+			for (const FChaosVDConnectivityEdge& ConnectivityEdge : InParticleDataViewer->ParticleCluster.ConnectivityEdges)
 			{
 				if (const TSharedPtr<FChaosVDScene> ScenePtr = InVisualizationContext.CVDScene.Pin())
 				{
 					if (AChaosVDParticleActor* SiblingParticle = ScenePtr->GetParticleActor(InVisualizationContext.SolverID, ConnectivityEdge.SiblingParticleID))
 					{
-						if (const FChaosVDParticleDataWrapper* SiblingParticleData = SiblingParticle->GetParticleData())
+						if (TSharedPtr<const FChaosVDParticleDataWrapper> SiblingParticleData = SiblingParticle->GetParticleData())
 						{
 							FColor DebugDrawColor = InVisualizationContext.DebugDrawSettings->ColorSettings.GetColorForDataID(EChaosVDParticleDataVisualizationFlags::ClusterConnectivityEdge, InVisualizationContext.bIsSelectedData);
 							FVector BoxExtents(2,2,2);
diff --git a/Engine/Plugins/ChaosVD/Source/ChaosVD/Private/Visualizers/ChaosVDParticleDataComponentVisualizer.h b/Engine/Plugins/ChaosVD/Source/ChaosVD/Private/Visualizers/ChaosVDParticleDataComponentVisualizer.h
--- a/Engine/Plugins/ChaosVD/Source/ChaosVD/Private/Visualizers/ChaosVDParticleDataComponentVisualizer.h
+++ b/Engine/Plugins/ChaosVD/Source/ChaosVD/Private/Visualizers/ChaosVDParticleDataComponentVisualizer.h
@@ -2,6 +2,7 @@
 
 #pragma once
 
+#include "ChaosVDComponentVisualizerBase.h"
 #include "ComponentVisualizer.h"
 #include "HAL/Platform.h"
 #include "IChaosVDParticleVisualizationDataProvider.h"
@@ -11,12 +12,6 @@
 class FChaosVDGeometryBuilder;
 struct FChaosParticleDataDebugDrawSettings;
 
-struct FChaosVDVisualizedParticleDataSelectionHandle
-{
-	int32 ParticleIndex = INDEX_NONE;
-	int32 SolverID = INDEX_NONE;
-};
-
 class UChaosVDParticleVisualizationDebugDrawSettings;
 
 struct FChaosVDParticleDataVisualizationContext : public FChaosVDVisualizationContext
@@ -34,38 +29,24 @@ struct FChaosVDParticleDataVisualizationContext : public FChaosVDVisualizationCo
 	}
 };
 
-/** Custom Hit Proxy for debug drawn particle data */
-struct HChaosVDParticleDataProxy : public HComponentVisProxy
-{
-	DECLARE_HIT_PROXY()
-	
-	HChaosVDParticleDataProxy(const UActorComponent* Component, const FChaosVDVisualizedParticleDataSelectionHandle& InContactFinderData) : HComponentVisProxy(Component, HPP_UI), DataSelectionHandle(InContactFinderData)
-	{	
-	}
-
-	virtual EMouseCursor::Type GetMouseCursor() override
-	{
-		return EMouseCursor::Crosshairs;
-	}
-
-	FChaosVDVisualizedParticleDataSelectionHandle DataSelectionHandle;
-};
-
 /**
  * Component visualizer in charge of generating debug draw visualizations for for particles
  */
-class FChaosVDParticleDataComponentVisualizer : public FComponentVisualizer
+class FChaosVDParticleDataComponentVisualizer : public FChaosVDComponentVisualizerBase
 {
 public:
 	FChaosVDParticleDataComponentVisualizer();
 
-	void RegisterVisualizerMenus();
+	virtual void RegisterVisualizerMenus() override;
 
 	virtual void DrawVisualization(const UActorComponent* Component, const FSceneView* View, FPrimitiveDrawInterface* PDI) override;
-	virtual bool VisProxyHandleClick(FEditorViewportClient* InViewportClient, HComponentVisProxy* VisProxy, const FViewportClick& Click) override;
+
+	virtual bool CanHandleClick(const HChaosVDComponentVisProxy& VisProxy) override;
+
+	virtual bool SelectVisualizedData(const HChaosVDComponentVisProxy& VisProxy, const TSharedRef<FChaosVDScene>& InCVDScene) override;
 
 protected:
 	
 	void DrawParticleVector(FPrimitiveDrawInterface* PDI, const FVector& StartLocation, const FVector& InVector, EChaosVDParticleDataVisualizationFlags VectorID, const FChaosVDParticleDataVisualizationContext& InVisualizationContext, float LineThickness);
-	void DrawVisualizationForParticleData(const UActorComponent* Component, FPrimitiveDrawInterface* PDI, const FSceneView* View, const FChaosVDParticleDataVisualizationContext& InVisualizationContext, const FChaosVDParticleDataWrapper& InParticleDataViewer);
+	void DrawVisualizationForParticleData(const UActorComponent* Component, FPrimitiveDrawInterface* PDI, const FSceneView* View, const FChaosVDParticleDataVisualizationContext& InVisualizationContext, const TSharedPtr<const FChaosVDParticleDataWrapper>& InParticleDataViewer);
 };
diff --git a/Engine/Plugins/ChaosVD/Source/ChaosVD/Private/Visualizers/ChaosVDSceneQueryDataComponentVisualizer.cpp b/Engine/Plugins/ChaosVD/Source/ChaosVD/Private/Visualizers/ChaosVDSceneQueryDataComponentVisualizer.cpp
--- a/Engine/Plugins/ChaosVD/Source/ChaosVD/Private/Visualizers/ChaosVDSceneQueryDataComponentVisualizer.cpp
+++ b/Engine/Plugins/ChaosVD/Source/ChaosVD/Private/Visualizers/ChaosVDSceneQueryDataComponentVisualizer.cpp
@@ -4,12 +4,10 @@
 
 #include "Actors/ChaosVDSceneQueryDataContainer.h"
 #include "Components/ChaosVDSceneQueryDataComponent.h"
-
 #include "ChaosVDGeometryBuilder.h"
 #include "ChaosVDScene.h"
 #include "ChaosVDStyle.h"
 #include "ChaosVDTabsIDs.h"
-#include "EditorModeManager.h"
 #include "EditorViewportClient.h"
 #include "SceneView.h"
 #include "Actors/ChaosVDSolverInfoActor.h"
@@ -18,53 +16,37 @@
 #include "ToolMenus.h"
 #include "ToolMenuEntry.h"
 #include "ToolMenuSection.h"
-#include "Utils/ChaosVDUserInterfaceUtils.h"
 #include "Visualizers/ChaosVDDebugDrawUtils.h"
-#include "Visualizers/IChaosVDParticleVisualizationDataProvider.h"
-#include "Widgets/SChaosVDMainTab.h"
 #include "Widgets/SChaosVDViewportToolbar.h"
-#include "Widgets/SChaosVDEnumFlagsMenu.h"
 
 class UChaosVDCoreSettings;
 
-IMPLEMENT_HIT_PROXY(HChaosVDSceneQueryProxy, HComponentVisProxy)
-
 #define LOCTEXT_NAMESPACE "ChaosVisualDebugger"
 
 FChaosVDSceneQueryDataComponentVisualizer::FChaosVDSceneQueryDataComponentVisualizer()
 {
-	RegisterVisualizerMenus();
+	FChaosVDSceneQueryDataComponentVisualizer::RegisterVisualizerMenus();
+
+	InspectorTabID = FChaosVDTabID::SceneQueryDataDetails;
 }
 
 void FChaosVDSceneQueryDataComponentVisualizer::RegisterVisualizerMenus()
 {
-	UToolMenus* ToolMenus = UToolMenus::Get();
+	FName MenuSection("SceneQueryDataVisualization.Show");
+	FText MenuSectionLabel = LOCTEXT("SceneQueryDataShowMenuLabel", "Scene Query Data Visualization");
+	FText FlagsMenuLabel = LOCTEXT("SceneQueryDataFlagsMenuLabel", "Scene Query Data Flags");
+	FText FlagsMenuTooltip = LOCTEXT("SceneQueryDataFlagsMenuToolTip", "Set of flags to enable/disable visibility of specific types of scene query data");
+	FSlateIcon FlagsMenuIcon = FSlateIcon(FChaosVDStyle::Get().GetStyleSetName(), TEXT("SceneQueriesInspectorIcon"));
+
+	FText SettingsMenuLabel = LOCTEXT("SceneQuerySettingsMenuLabel", "Scene Query Visualization Settings");
+	FText SettingsMenuTooltip = LOCTEXT("SceneQuerySettingsMenuToolTip", "Options to change how the recorded scene query data is debug drawn");
 	
-	if (!ensure(ToolMenus))
-	{
-		return;
-	}
+	CreateGenericVisualizerMenu<UChaosVDSceneQueriesVisualizationSettings, EChaosVDSceneQueryVisualizationFlags>(SChaosVDViewportToolbar::ShowMenuName, MenuSection, MenuSectionLabel, FlagsMenuLabel, FlagsMenuTooltip, FlagsMenuIcon, SettingsMenuLabel, SettingsMenuTooltip);
+}
 
-	if (UToolMenu* Menu = ToolMenus->ExtendMenu(SChaosVDViewportToolbar::ShowMenuName))
-	{
-		FToolMenuSection& Section = Menu->AddSection("SceneQueryDataVisualization.Show", LOCTEXT("SceneQueryDataShowMenuLabel", "Scene Query Data Visualization"));
-		
-		Section.AddSubMenu(TEXT("SceneQueryDataFlags"), LOCTEXT("SceneQueryDataFlagsMenuLabel", "Scene Query Data Flags"), LOCTEXT("SceneQueryDataFlagsMenuToolTip", "Set of flags to enable/disable visibility of specific types of scene query data"), FNewToolMenuDelegate::CreateLambda([](UToolMenu* Menu)
-						   {
-							   TSharedRef<SWidget> VisualizationFlagsWidget = SNew(SChaosVDEnumFlagsMenu<EChaosVDSceneQueryVisualizationFlags>)
-								   .CurrentValue_Static(&UChaosVDSceneQueriesVisualizationSettings::GetSceneQueryDataVisualizationFlags)
-								   .OnEnumSelectionChanged_Static(&UChaosVDSceneQueriesVisualizationSettings::SetSceneQueryDataVisualizationFlags);
-			
-
-							   FToolMenuEntry FlagsMenuEntry = FToolMenuEntry::InitWidget("CollisionVisualizationFlagsFlags", VisualizationFlagsWidget,FText::GetEmpty());
-							   Menu->AddMenuEntry(NAME_None, FlagsMenuEntry);
-						   }),
-						   false, FSlateIcon(FChaosVDStyle::Get().GetStyleSetName(), TEXT("SceneQueriesInspectorIcon")));
-		
-		using namespace Chaos::VisualDebugger::Utils;
-        Section.AddSubMenu(TEXT("SceneQueryDataSettings"), LOCTEXT("SceneQuerySettingsMenuLabel", "Scene Query Visualization Settings"), LOCTEXT("SceneQuerySettingsMenuToolTip", "Options to change how the recorded scene query data is debug drawn"), FNewToolMenuDelegate::CreateStatic(&CreateMenuEntryForDefaultObject<UChaosVDSceneQueriesVisualizationSettings>, EChaosVDSaveSettingsOptions::ShowSaveButton),
-                           false, FSlateIcon(FAppStyle::Get().GetStyleSetName(), TEXT("Icons.Toolbar.Settings")));
-	}
+bool FChaosVDSceneQueryDataComponentVisualizer::CanHandleClick(const HChaosVDComponentVisProxy& VisProxy)
+{
+	return VisProxy.DataSelectionHandle && VisProxy.DataSelectionHandle->IsA<FChaosVDQueryDataWrapper>();
 }
 
 void FChaosVDSceneQueryDataComponentVisualizer::DrawVisualization(const UActorComponent* Component, const FSceneView* View, FPrimitiveDrawInterface* PDI)
@@ -99,23 +81,31 @@ void FChaosVDSceneQueryDataComponentVisualizer::DrawVisualization(const UActorCo
 		return;
 	}
 
+	TSharedPtr<FChaosVDSolverDataSelection> SolverDataSelectionObject = CVDScene->GetSolverDataSelectionObject().Pin();
+	if (!ensure(SolverDataSelectionObject))
+	{
+		return;
+	}
+
 	FChaosVDSceneQueryVisualizationDataContext VisualizationContext;
 	VisualizationContext.CVDScene = SceneQueryDataContainer->GetScene();
 	VisualizationContext.SpaceTransform = FTransform::Identity;
 	VisualizationContext.GeometryGenerator = GeometryGenerator;
+	VisualizationContext.SolverDataSelectionObject = SolverDataSelectionObject;
 
-	VisualizationContext.VisualizationFlags = static_cast<uint32>(UChaosVDSceneQueriesVisualizationSettings::GetSceneQueryDataVisualizationFlags());
+	VisualizationContext.VisualizationFlags = static_cast<uint32>(UChaosVDSceneQueriesVisualizationSettings::GetDataVisualizationFlags());
 	VisualizationContext.DebugDrawSettings = GetDefault<UChaosVDSceneQueriesVisualizationSettings>();
 
 	if (EnumHasAnyFlags(EChaosVDSceneQueryVisualizationFlags::EnableDraw, static_cast<EChaosVDSceneQueryVisualizationFlags>(VisualizationContext.VisualizationFlags)))
 	{
 		// If Draw only selected Query is enabled, but no query is selected, just draw all queries
 
-		const TSharedPtr<FChaosVDQueryDataWrapper> SelectedQuery = SceneQueryDataComponent->GetSelectedQueryHandle().GetQueryData().Pin();
-		const bool bHasSelectedQuery = SelectedQuery.IsValid();
+		TSharedPtr<FChaosVDSolverDataSelectionHandle> SelectionHandle = SolverDataSelectionObject->GetCurrentSelectionHandle();
+		const bool bHasSelectedQuery = SelectionHandle && SelectionHandle->IsA<FChaosVDQueryDataWrapper>();
+
 		if (bHasSelectedQuery && EnumHasAnyFlags(EChaosVDSceneQueryVisualizationFlags::OnlyDrawSelectedQuery, static_cast<EChaosVDSceneQueryVisualizationFlags>(VisualizationContext.VisualizationFlags)))
 		{
-			DrawSceneQuery(Component, View, PDI, CVDScene, CVDRecording, VisualizationContext, SelectedQuery);
+			DrawSceneQuery(Component, View, PDI, CVDScene, CVDRecording, VisualizationContext, SelectionHandle->GetDataAsShared<FChaosVDQueryDataWrapper>());
 		}
 		else
 		{
@@ -127,31 +117,6 @@ void FChaosVDSceneQueryDataComponentVisualizer::DrawVisualization(const UActorCo
 	}
 }
 
-bool FChaosVDSceneQueryDataComponentVisualizer::VisProxyHandleClick(FEditorViewportClient* InViewportClient, HComponentVisProxy* VisProxy, const FViewportClick& Click)
-{
-	const HChaosVDSceneQueryProxy* SceneQueryDataProxy = HitProxyCast<HChaosVDSceneQueryProxy>(VisProxy);
-	if (SceneQueryDataProxy == nullptr)
-	{
-		return false;
-	}
-	
-	if (const UChaosVDSceneQueryDataComponent* SQDataComponent = Cast<UChaosVDSceneQueryDataComponent>(SceneQueryDataProxy->Component.Get()))
-	{
-		// Bring the SQ Inspector into focus if available
-		const TSharedPtr<SChaosVDMainTab> MainTabToolkitHost = InViewportClient->GetModeTools() ? StaticCastSharedPtr<SChaosVDMainTab>(InViewportClient->GetModeTools()->GetToolkitHost()) : nullptr;
-		if (const TSharedPtr<FTabManager> TabManager = MainTabToolkitHost ? MainTabToolkitHost->GetTabManager() : nullptr)
-		{
-			TabManager->TryInvokeTab(FChaosVDTabID::SceneQueryDataDetails);
-		}
-
-		const_cast<UChaosVDSceneQueryDataComponent*>(SQDataComponent)->SelectQuery(SceneQueryDataProxy->DataSelectionHandle);
-
-		return true;
-	}
-
-	return false;
-}
-
 void FChaosVDSceneQueryDataComponentVisualizer::DrawLineTraceQuery(const UActorComponent* Component, const FChaosVDQueryDataWrapper& SceneQueryData, FChaosVDSceneQueryVisualizationDataContext& VisualizationContext, const FSceneView* View, FPrimitiveDrawInterface* PDI)
 {
 	const UChaosVDSceneQueriesVisualizationSettings* DebugDrawSettings = Cast<UChaosVDSceneQueriesVisualizationSettings>(VisualizationContext.DebugDrawSettings);
@@ -160,7 +125,7 @@ void FChaosVDSceneQueryDataComponentVisualizer::DrawLineTraceQuery(const UActorC
 		return;
 	}
 
-	PDI->SetHitProxy(new HChaosVDSceneQueryProxy(Component, VisualizationContext.DataSelectionHandle));
+	PDI->SetHitProxy(new HChaosVDComponentVisProxy(Component, VisualizationContext.DataSelectionHandle));
 
 	const FText DebugText = DebugDrawSettings->bShowText ? FText::FormatOrdered(LOCTEXT("LineTraceDebugDrawText", "Type: Line Trace \n Tag {1} \n Owner Tag {2}"), FText::AsCultureInvariant(SceneQueryData.CollisionQueryParams.TraceTag.ToString()), FText::AsCultureInvariant(SceneQueryData.CollisionQueryParams.OwnerTag.ToString()))
 							: FText::GetEmpty();
@@ -195,7 +160,7 @@ void FChaosVDSceneQueryDataComponentVisualizer::DrawOverlapQuery(const UActorCom
 		return;
 	}
 	
-	PDI->SetHitProxy(new HChaosVDSceneQueryProxy(Component, VisualizationContext.DataSelectionHandle));
+	PDI->SetHitProxy(new HChaosVDComponentVisProxy(Component, VisualizationContext.DataSelectionHandle));
 
 	const Chaos::FConstImplicitObjectPtr InputShapePtr = VisualizationContext.InputGeometry;
 	if (ensure(InputShapePtr))
@@ -221,7 +186,7 @@ void FChaosVDSceneQueryDataComponentVisualizer::DrawSweepQuery(const UActorCompo
 		return;
 	}
 
-	PDI->SetHitProxy(new HChaosVDSceneQueryProxy(Component, VisualizationContext.DataSelectionHandle));
+	PDI->SetHitProxy(new HChaosVDComponentVisProxy(Component, VisualizationContext.DataSelectionHandle));
 
 	const Chaos::FConstImplicitObjectPtr InputShapePtr = VisualizationContext.InputGeometry;
 	if (ensure(InputShapePtr))
@@ -271,10 +236,14 @@ void FChaosVDSceneQueryDataComponentVisualizer::DrawHits(const UActorComponent*
 		{
 			continue;
 		}
+	
+		TSharedPtr<FChaosVDSolverDataSelectionHandle> HitSelectionHandle = VisualizationContext.SolverDataSelectionObject->MakeSelectionHandle(VisualizationContext.DataSelectionHandle->GetDataAsShared<FChaosVDQueryDataWrapper>());
 
-		VisualizationContext.DataSelectionHandle = FChaosVDSceneQuerySelectionHandle(VisualizationContext.DataSelectionHandle.GetQueryData(), SQVisitIndex);
+		FChaosVDSceneQuerySelectionContext ContextData;
+		ContextData.SQVisitIndex = SQVisitIndex;
+		HitSelectionHandle->SetHandleContext(MoveTemp(ContextData));
 
-		PDI->SetHitProxy(new HChaosVDSceneQueryProxy(Component, VisualizationContext.DataSelectionHandle));
+		PDI->SetHitProxy(new HChaosVDComponentVisProxy(Component, HitSelectionHandle));
 
 		const FText HitPointDebugText = DebugDrawSettings->bShowText ? FText::FormatOrdered(LOCTEXT("SceneQueryHitDebugText", "Distance {0} \n Face Index {1} \n "), SQVisitData.HitData.Distance, SQVisitData.HitData.FaceIdx) : FText::GetEmpty();
 		static FText HitFaceNormalDebugText = FText::AsCultureInvariant(TEXT("Hit Face Normal"));
@@ -294,19 +263,22 @@ void FChaosVDSceneQueryDataComponentVisualizer::DrawHits(const UActorComponent*
 			FChaosVDDebugDrawUtils::DrawArrowVector(PDI, SQVisitData.HitData.WorldPosition, SQVisitData.HitData.WorldPosition + SQVisitData.HitData.WorldNormal * NormalScale, DebugDrawSettings->bShowText ? HitWorldNormalDebugText : FText::GetEmpty(), InColor, ESceneDepthPriorityGroup::SDPG_Foreground);
 		}
 
-		if (SQVisitData.bIsSelectedInEditor)
+		if (TSharedPtr<FChaosVDSolverDataSelectionHandle> CurrentSelection = VisualizationContext.SolverDataSelectionObject->GetCurrentSelectionHandle())
 		{
-			// We don't have an easy way to show something is selected with debug draw
-			// but 3D box surrounding the hit is better than nothing
-			FTransform SelectionBoxTransform;
-			SelectionBoxTransform.SetRotation(FRotationMatrix::MakeFromZ(SQVisitData.HitData.WorldNormal).ToQuat());
-			SelectionBoxTransform.SetLocation(SQVisitData.HitData.WorldPosition);
+			if (IsHitSelected(SQVisitIndex, CurrentSelection.ToSharedRef(), HitSelectionHandle.ToSharedRef()))
+			{
+				// We don't have an easy way to show something is selected with debug draw
+				// but 3D box surrounding the hit is better than nothing
+				FTransform SelectionBoxTransform;
+				SelectionBoxTransform.SetRotation(FRotationMatrix::MakeFromZ(SQVisitData.HitData.WorldNormal).ToQuat());
+				SelectionBoxTransform.SetLocation(SQVisitData.HitData.WorldPosition);
 
-			// The Selection box should be a bit bigger than the configured circle radius for the debug draw hit
-			constexpr float HitSelectionBoxSize = CircleRadius * 1.2f;
+				// The Selection box should be a bit bigger than the configured circle radius for the debug draw hit
+				constexpr float HitSelectionBoxSize = CircleRadius * 1.2f;
 
-			FVector SelectionBoxExtents(HitSelectionBoxSize,HitSelectionBoxSize,HitSelectionBoxSize);
-			FChaosVDDebugDrawUtils::DrawBox(PDI, SelectionBoxExtents, FColor::Yellow, SelectionBoxTransform, FText::GetEmpty(), ESceneDepthPriorityGroup::SDPG_Foreground);
+				FVector SelectionBoxExtents(HitSelectionBoxSize,HitSelectionBoxSize,HitSelectionBoxSize);
+				FChaosVDDebugDrawUtils::DrawBox(PDI, SelectionBoxExtents, FColor::Yellow, SelectionBoxTransform, FText::GetEmpty(), ESceneDepthPriorityGroup::SDPG_Foreground);
+			}
 		}
 
 		PDI->SetHitProxy(nullptr);
@@ -318,11 +290,24 @@ bool FChaosVDSceneQueryDataComponentVisualizer::HasEndLocation(const FChaosVDQue
 	return SceneQueryData.Type != EChaosVDSceneQueryType::Sweep;
 }
 
+bool FChaosVDSceneQueryDataComponentVisualizer::IsHitSelected(int32 SQVisitIndex, const TSharedRef<FChaosVDSolverDataSelectionHandle>& CurrentSelection, const TSharedRef<FChaosVDSolverDataSelectionHandle>& SQVisitSelectionHandle)
+{
+	if (SQVisitSelectionHandle->IsSelected())
+	{
+		if (FChaosVDSceneQuerySelectionContext* SelectionContext = CurrentSelection->GetContextData<FChaosVDSceneQuerySelectionContext>())
+		{
+			return SelectionContext->SQVisitIndex == SQVisitIndex;
+		}
+	}
+
+	return false;				
+}
+
 void FChaosVDSceneQueryDataComponentVisualizer::DrawSceneQuery(const UActorComponent* Component, const FSceneView* View, FPrimitiveDrawInterface* PDI, const TSharedPtr<FChaosVDScene>& CVDScene, const TSharedPtr<FChaosVDRecording>& CVDRecording, FChaosVDSceneQueryVisualizationDataContext& VisualizationContext, const TSharedPtr<FChaosVDQueryDataWrapper>& Query)
 {
 	// Reset query Specify context values
 	VisualizationContext.InputGeometry = nullptr;
-	VisualizationContext.DataSelectionHandle = FChaosVDSceneQuerySelectionHandle(nullptr, INDEX_NONE);
+	VisualizationContext.DataSelectionHandle = nullptr;
 
 	if (!Query)
 	{
@@ -367,8 +352,9 @@ void FChaosVDSceneQueryDataComponentVisualizer::DrawSceneQuery(const UActorCompo
 		VisualizationContext.InputGeometry = *InputShapePtrPtr;
 	}
 
-	VisualizationContext.DataSelectionHandle = FChaosVDSceneQuerySelectionHandle(Query, INDEX_NONE);
-	VisualizationContext.GenerateColor(Query->ID, Query->bIsSelectedInEditor);
+	VisualizationContext.DataSelectionHandle = VisualizationContext.SolverDataSelectionObject->MakeSelectionHandle(Query);
+	bool bIsSelected = VisualizationContext.DataSelectionHandle && VisualizationContext.DataSelectionHandle->IsSelected();
+	VisualizationContext.GenerateColor(Query->ID, bIsSelected);
 
 	switch (Query->Type)
 	{
diff --git a/Engine/Plugins/ChaosVD/Source/ChaosVD/Private/Visualizers/ChaosVDSceneQueryDataComponentVisualizer.h b/Engine/Plugins/ChaosVD/Source/ChaosVD/Private/Visualizers/ChaosVDSceneQueryDataComponentVisualizer.h
--- a/Engine/Plugins/ChaosVD/Source/ChaosVD/Private/Visualizers/ChaosVDSceneQueryDataComponentVisualizer.h
+++ b/Engine/Plugins/ChaosVD/Source/ChaosVD/Private/Visualizers/ChaosVDSceneQueryDataComponentVisualizer.h
@@ -2,21 +2,32 @@
 
 #pragma once
 
+#include "ChaosVDComponentVisualizerBase.h"
 #include "ComponentVisualizer.h"
-#include "IChaosVDParticleVisualizationDataProvider.h"
 #include "Chaos/ImplicitFwd.h"
 #include "Chaos/ImplicitObject.h"
 #include "Components/ChaosVDSceneQueryDataComponent.h"
 
+#include "ChaosVDSceneQueryDataComponentVisualizer.generated.h"
+
 struct FChaosVDRecording;
 class FChaosVDGeometryBuilder;
 struct FChaosVDVisualizationContext;
 struct FChaosVDQueryDataWrapper;
 
+/** Struct used to pass data about a specific query to other objects */
+USTRUCT()
+struct FChaosVDSceneQuerySelectionContext : public FChaosVDSelectionContext
+{
+	GENERATED_BODY()
+
+	int32 SQVisitIndex = INDEX_NONE;
+};
+
 /** Visualization context structure specific for Scene Queries visualizations */
 struct FChaosVDSceneQueryVisualizationDataContext : public FChaosVDVisualizationContext
 {
-	FChaosVDSceneQuerySelectionHandle DataSelectionHandle = FChaosVDSceneQuerySelectionHandle(nullptr, INDEX_NONE);
+	TSharedPtr<FChaosVDSolverDataSelectionHandle> DataSelectionHandle = nullptr;
 
 	/** Generates a random color based on the selection state and query ID, which will be used to debug draw the scene query */
 	void GenerateColor(int32 QueryID, bool bIsSelected)
@@ -36,36 +47,21 @@ struct FChaosVDSceneQueryVisualizationDataContext : public FChaosVDVisualization
 	TWeakPtr<FChaosVDGeometryBuilder> GeometryGenerator = nullptr;
 };
 
-/** Custom Hit Proxy for debug drawn scene queries */
-struct HChaosVDSceneQueryProxy : public HComponentVisProxy
-{
-	DECLARE_HIT_PROXY()
-	
-	HChaosVDSceneQueryProxy(const UActorComponent* Component, const FChaosVDSceneQuerySelectionHandle& InContactFinderData) : HComponentVisProxy(Component, HPP_UI), DataSelectionHandle(InContactFinderData)
-	{	
-	}
-
-	virtual EMouseCursor::Type GetMouseCursor() override
-	{
-		return EMouseCursor::Crosshairs;
-	}
-
-	FChaosVDSceneQuerySelectionHandle DataSelectionHandle;
-};
-
 /**
  * Component visualizer in charge of generating debug draw visualizations for scene queries in a ChaosVDSceneQueryDataComponent
  */
-class FChaosVDSceneQueryDataComponentVisualizer final : public FComponentVisualizer
+class FChaosVDSceneQueryDataComponentVisualizer final : public FChaosVDComponentVisualizerBase
 {
 
 public:
 	FChaosVDSceneQueryDataComponentVisualizer();
-	
-	void RegisterVisualizerMenus();
+
+	virtual void RegisterVisualizerMenus() override;
 
 	virtual void DrawVisualization(const UActorComponent* Component, const FSceneView* View, FPrimitiveDrawInterface* PDI) override;
-	virtual bool VisProxyHandleClick(FEditorViewportClient* InViewportClient, HComponentVisProxy* VisProxy, const FViewportClick& Click) override;
+
+
+	virtual bool CanHandleClick(const HChaosVDComponentVisProxy& VisProxy) override;
 
 private:
 
@@ -77,4 +73,6 @@ private:
 	void DrawSceneQuery(const UActorComponent* Component, const FSceneView* View, FPrimitiveDrawInterface* PDI, const TSharedPtr<FChaosVDScene>& CVDScene, const TSharedPtr<FChaosVDRecording>& CVDRecording, FChaosVDSceneQueryVisualizationDataContext& VisualizationContext, const TSharedPtr<FChaosVDQueryDataWrapper>& Query);
 
 	bool HasEndLocation(const FChaosVDQueryDataWrapper& SceneQueryData) const;
+
+	bool IsHitSelected(int32 SQVisitIndex, const TSharedRef<FChaosVDSolverDataSelectionHandle>& CurrentSelection, const TSharedRef<FChaosVDSolverDataSelectionHandle>& SQVisitSelectionHandle);
 };
diff --git a/Engine/Plugins/ChaosVD/Source/ChaosVD/Private/Visualizers/ChaosVDSolverCollisionDataComponentVisualizer.cpp b/Engine/Plugins/ChaosVD/Source/ChaosVD/Private/Visualizers/ChaosVDSolverCollisionDataComponentVisualizer.cpp
--- a/Engine/Plugins/ChaosVD/Source/ChaosVD/Private/Visualizers/ChaosVDSolverCollisionDataComponentVisualizer.cpp
+++ b/Engine/Plugins/ChaosVD/Source/ChaosVD/Private/Visualizers/ChaosVDSolverCollisionDataComponentVisualizer.cpp
@@ -2,11 +2,10 @@
 
 #include "ChaosVDSolverCollisionDataComponentVisualizer.h"
 
-#include "ChaosVDCollisionDataDetailsTab.h"
+#include "ChaosVDComponentVisualizerBase.h"
 #include "ChaosVDParticleActor.h"
 #include "ChaosVDScene.h"
 #include "ChaosVDTabsIDs.h"
-#include "EditorModeManager.h"
 #include "EditorViewportClient.h"
 #include "Actors/ChaosVDSolverInfoActor.h"
 #include "Components/ChaosVDSolverCollisionDataComponent.h"
@@ -15,31 +14,18 @@
 #include "ToolMenus.h"
 #include "ToolMenuEntry.h"
 #include "ToolMenuSection.h"
-#include "Utils/ChaosVDUserInterfaceUtils.h"
 #include "Visualizers/ChaosVDDebugDrawUtils.h"
 #include "Widgets/SChaosVDCollisionDataInspector.h"
-#include "Widgets/SChaosVDMainTab.h"
 #include "Widgets/SChaosVDViewportToolbar.h"
-#include "Widgets/SChaosVDEnumFlagsMenu.h"
 
-IMPLEMENT_HIT_PROXY(HChaosVDContactPointProxy, HComponentVisProxy)
 
 #define LOCTEXT_NAMESPACE "ChaosVisualDebugger"
 
-void FChaosVDCollisionDataFinder::SetIsSelected(bool bNewSelected)
-{
-	if (TSharedPtr<FChaosVDParticlePairMidPhase> MidPhaseData = OwningMidPhase.Pin())
-	{
-		if (OwningConstraint && OwningConstraint->ManifoldPoints.IsValidIndex(ContactIndex))
-		{
-			OwningConstraint->ManifoldPoints[ContactIndex].bIsSelectedInEditor = bNewSelected;
-		}
-	}
-}
-
 FChaosVDSolverCollisionDataComponentVisualizer::FChaosVDSolverCollisionDataComponentVisualizer()
 {
-	RegisterVisualizerMenus();
+	FChaosVDSolverCollisionDataComponentVisualizer::RegisterVisualizerMenus();
+
+	InspectorTabID = FChaosVDTabID::CollisionDataDetails;
 }
 
 FChaosVDSolverCollisionDataComponentVisualizer::~FChaosVDSolverCollisionDataComponentVisualizer()
@@ -48,33 +34,16 @@ FChaosVDSolverCollisionDataComponentVisualizer::~FChaosVDSolverCollisionDataComp
 
 void FChaosVDSolverCollisionDataComponentVisualizer::RegisterVisualizerMenus()
 {
-	UToolMenus* ToolMenus = UToolMenus::Get();
+	FName MenuSection("CollisionDataVisualization.Show");
+	FText MenuSectionLabel = LOCTEXT("CollisionDataVisualizationShowMenuLabel", "Collision Data Visualization");
+	FText FlagsMenuLabel = LOCTEXT("CollisionVisualizationFlagsMenuLabel", "Collision Data Flags");
+	FText FlagsMenuTooltip = LOCTEXT("CollisionVisualizationFlagsMenuToolTip", "Set of flags to enable/disable visibility of specific types of collisionData");
+	FSlateIcon FlagsMenuIcon = FSlateIcon(FSlateIcon(FAppStyle::Get().GetStyleSetName(), TEXT("EditorViewport.CollisionVisibility")));
+
+	FText SettingsMenuLabel = LOCTEXT("CollisionVisualizationSettingsMenuLabel", "Collision Visualization Settings");
+	FText SettingsMenuTooltip = LOCTEXT("CollisionVisualizationSettingsMenuToolTip", "Options to change how the recorded collision Data is debug drawn");
 	
-	if (!ensure(ToolMenus))
-	{
-		return;
-	}
-
-	if (UToolMenu* Menu = ToolMenus->ExtendMenu(SChaosVDViewportToolbar::ShowMenuName))
-	{
-		FToolMenuSection& Section = Menu->AddSection("CollisionDataVisualization.Show", LOCTEXT("CollisionDataVisualizationShowMenuLabel", "Collision Data Visualization"));
-		
-		Section.AddSubMenu(TEXT("CollisionVisualizationFlags"), LOCTEXT("CollisionVisualizationFlagsMenuLabel", "Collision Data Flags"), LOCTEXT("CollisionVisualizationFlagsMenuToolTip", "Set of flags to enable/disable visibility of specific types of collisionData"), FNewToolMenuDelegate::CreateLambda([](UToolMenu* Menu)
-		                   {
-			                   TSharedRef<SWidget> VisualizationFlagsWidget = SNew(SChaosVDEnumFlagsMenu<EChaosVDCollisionVisualizationFlags>)
-				                   .CurrentValue_Static(&UChaosVDCollisionDataVisualizationSettings::GetCollisionDataVisualizationFlags)
-				                   .OnEnumSelectionChanged_Static(&UChaosVDCollisionDataVisualizationSettings::SetCollisionDataVisualizationFlags);
-			
-
-			                   FToolMenuEntry FlagsMenuEntry = FToolMenuEntry::InitWidget("CollisionVisualizationFlagsFlags", VisualizationFlagsWidget,FText::GetEmpty());
-			                   Menu->AddMenuEntry(NAME_None, FlagsMenuEntry);
-		                   }),
-		                   false, FSlateIcon(FAppStyle::Get().GetStyleSetName(), TEXT("EditorViewport.CollisionVisibility")));
-
-		using namespace Chaos::VisualDebugger::Utils;
-		Section.AddSubMenu(TEXT("CollisionVisualizationSettings"), LOCTEXT("CollisionVisualizationSettingsMenuLabel", "Collision Visualization Settings"), LOCTEXT("CollisionVisualizationSettingsMenuToolTip", "Options to change how the recorded collision Data is debug drawn"), FNewToolMenuDelegate::CreateStatic(&CreateMenuEntryForDefaultObject<UChaosVDCollisionDataVisualizationSettings>, EChaosVDSaveSettingsOptions::ShowSaveButton),
-		                   false, FSlateIcon(FAppStyle::Get().GetStyleSetName(), TEXT("Icons.Toolbar.Settings")));
-	}
+	CreateGenericVisualizerMenu<UChaosVDCollisionDataVisualizationSettings, EChaosVDCollisionVisualizationFlags>(SChaosVDViewportToolbar::ShowMenuName, MenuSection, MenuSectionLabel, FlagsMenuLabel, FlagsMenuTooltip, FlagsMenuIcon, SettingsMenuLabel, SettingsMenuTooltip);
 }
 
 bool FChaosVDSolverCollisionDataComponentVisualizer::ShowWhenSelected()
@@ -105,7 +74,7 @@ void FChaosVDSolverCollisionDataComponentVisualizer::DrawVisualization(const UAc
 	VisualizationContext.SolverID = SolverInfoContainer->GetSolverID();
 	VisualizationContext.CVDScene = SolverInfoContainer->GetScene();
 	VisualizationContext.SpaceTransform = SolverInfoContainer->GetSimulationTransform();
-	VisualizationContext.VisualizationFlags = static_cast<uint32>(UChaosVDCollisionDataVisualizationSettings::GetCollisionDataVisualizationFlags());
+	VisualizationContext.VisualizationFlags = static_cast<uint32>(UChaosVDCollisionDataVisualizationSettings::GetDataVisualizationFlags());
 	VisualizationContext.DebugDrawSettings = GetDefault<UChaosVDCollisionDataVisualizationSettings>();
 
 	if (!EnumHasAnyFlags(EChaosVDCollisionVisualizationFlags::EnableDraw, static_cast<EChaosVDCollisionVisualizationFlags>(VisualizationContext.VisualizationFlags)))
@@ -137,58 +106,9 @@ void FChaosVDSolverCollisionDataComponentVisualizer::DrawVisualization(const UAc
 	}
 }
 
-bool FChaosVDSolverCollisionDataComponentVisualizer::VisProxyHandleClick(FEditorViewportClient* InViewportClient, HComponentVisProxy* VisProxy, const FViewportClick& Click)
+bool FChaosVDSolverCollisionDataComponentVisualizer::CanHandleClick(const HChaosVDComponentVisProxy& VisProxy)
 {
-	const HChaosVDContactPointProxy* ContactProxy = HitProxyCast<HChaosVDContactPointProxy>(VisProxy);
-	if (ContactProxy == nullptr)
-	{
-		ClearCurrentSelection();
-		return false;
-	}
-
-	const TSharedPtr<SChaosVDMainTab> MainTabToolkitHost = InViewportClient->GetModeTools() ? StaticCastSharedPtr<SChaosVDMainTab>(InViewportClient->GetModeTools()->GetToolkitHost()) : nullptr;
-	if (!MainTabToolkitHost.IsValid())
-	{
-		return false;
-	}
-
-	const TSharedPtr<FChaosVDCollisionDataDetailsTab> CollisionDataDetailsTab = MainTabToolkitHost->GetTabSpawnerInstance<FChaosVDCollisionDataDetailsTab>(FChaosVDTabID::CollisionDataDetails).Pin();
-	if (!CollisionDataDetailsTab)
-	{
-		return false;
-	}
-
-	const TSharedPtr<FTabManager> TabManager = MainTabToolkitHost->GetTabManager();
-	if (!TabManager)
-	{
-		return false;
-	}
-
-	TabManager->TryInvokeTab(FChaosVDTabID::CollisionDataDetails);
-	
-	const TSharedPtr<SChaosVDCollisionDataInspector> CollisionInspector = CollisionDataDetailsTab->GetCollisionInspectorInstance().Pin();
-	if (!CollisionInspector)
-	{
-		return false;
-	}
-
-	if (ContactProxy->ContactFinder.OwningMidPhase.Pin())
-	{
-		
-		CollisionInspector->SetSingleContactDataToInspect(ContactProxy->ContactFinder);
-		
-		ClearCurrentSelection();
-		CurrentSelectedContactData = ContactProxy->ContactFinder;
-		CurrentSelectedContactData.SetIsSelected(true);
-		return true;
-	}
-
-	return false;
-}
-
-void FChaosVDSolverCollisionDataComponentVisualizer::ClearCurrentSelection()
-{
-	CurrentSelectedContactData.SetIsSelected(false);
+	return VisProxy.DataSelectionHandle && VisProxy.DataSelectionHandle->IsA<FChaosVDParticlePairMidPhase>();
 }
 
 void FChaosVDSolverCollisionDataComponentVisualizer::DrawMidPhaseData(const UActorComponent* Component, const TSharedPtr<FChaosVDParticlePairMidPhase>& MidPhase, const FChaosVDVisualizationContext& VisualizationContext, const FSceneView* View, FPrimitiveDrawInterface* PDI)
@@ -213,14 +133,28 @@ void FChaosVDSolverCollisionDataComponentVisualizer::DrawMidPhaseData(const UAct
 			return;
 		}
 		
-		const FChaosVDParticleDataWrapper* ParticleDataViewer = ParticleActor->GetParticleData();
+		TSharedPtr<const FChaosVDParticleDataWrapper> ParticleDataViewer = ParticleActor->GetParticleData();
 		if (!ensure(ParticleDataViewer))
 		{
 			return;
 		}
 
-		EChaosVDCollisionVisualizationFlags VisualizationFlags = static_cast<EChaosVDCollisionVisualizationFlags>(VisualizationContext.VisualizationFlags);
+		TSharedPtr<FChaosVDSolverDataSelection> SolverDataSelectionObject = ScenePtr->GetSolverDataSelectionObject().Pin();
+		if (!ensure(SolverDataSelectionObject))
+		{
+			return;
+		}
+
+		TSharedPtr<FChaosVDSolverDataSelectionHandle> CurrentSelection = SolverDataSelectionObject->GetCurrentSelectionHandle();
+		FChaosVDCollisionDataSelectionContext* CurrentSelectionContext = nullptr;
 
+		if (CurrentSelection && CurrentSelection->IsA<FChaosVDParticlePairMidPhase>())
+		{
+			CurrentSelectionContext = CurrentSelection->GetContextData<FChaosVDCollisionDataSelectionContext>();
+		}
+
+		EChaosVDCollisionVisualizationFlags VisualizationFlags = static_cast<EChaosVDCollisionVisualizationFlags>(VisualizationContext.VisualizationFlags);
+		
 		for (FChaosVDConstraint& Constraint : MidPhase->Constraints)
 		{
 			int32 ContactIndex = INDEX_NONE;
@@ -228,10 +162,14 @@ void FChaosVDSolverCollisionDataComponentVisualizer::DrawMidPhaseData(const UAct
 			{
 				ContactIndex++;
 
-				FChaosVDCollisionDataFinder HitPRoxyDataFinder;
-				HitPRoxyDataFinder.OwningMidPhase = MidPhase;
-				HitPRoxyDataFinder.OwningConstraint = &Constraint;
-				HitPRoxyDataFinder.ContactIndex = ContactIndex;
+				TSharedPtr<FChaosVDSolverDataSelectionHandle> SelectionHandle = SolverDataSelectionObject->MakeSelectionHandle(MidPhase);
+
+				FChaosVDCollisionDataSelectionContext ContextData;
+				ContextData.MidPhase = MidPhase;
+				ContextData.ConstraintDataPtr = &Constraint;
+				ContextData.ContactDataIndex = ContactIndex;
+
+				SelectionHandle->SetHandleContext(MoveTemp(ContextData));
 
 				const bool bIsProbe = Constraint.bIsProbe;
 				const bool bIsActive = ManifoldPoint.bIsValid && (!ManifoldPoint.NetPushOut.IsNearlyZero() || !ManifoldPoint.NetImpulse.IsNearlyZero() || (!Constraint.bUseManifold && !Constraint.AccumulatedImpulse.IsNearlyZero()));
@@ -249,10 +187,10 @@ void FChaosVDSolverCollisionDataComponentVisualizer::DrawMidPhaseData(const UAct
 					continue;
 				}
 
-				PDI->SetHitProxy(new HChaosVDContactPointProxy(Component, HitPRoxyDataFinder));
+				PDI->SetHitProxy(new HChaosVDComponentVisProxy(Component, SelectionHandle));
 
-				const FChaosVDParticleDataWrapper* ParticleDataActor1 = CVDParticleActor1->GetParticleData();
-				const FChaosVDParticleDataWrapper* ParticleDataActor0 = CVDParticleActor0->GetParticleData();
+				TSharedPtr<const FChaosVDParticleDataWrapper> ParticleDataActor1 = CVDParticleActor1->GetParticleData();
+				TSharedPtr<const FChaosVDParticleDataWrapper> ParticleDataActor0 = CVDParticleActor0->GetParticleData();
 				const FTransform WorldActorTransform1 = ParticleDataActor1 && ParticleDataActor1->ParticlePositionRotation.HasValidData() ? FTransform(ParticleDataActor1->ParticlePositionRotation.MR, ParticleDataActor1->ParticlePositionRotation.MX) : FTransform();
 				const FTransform WorldActorTransform0 = ParticleDataActor0 && ParticleDataActor0->ParticlePositionRotation.HasValidData() ? FTransform(ParticleDataActor0->ParticlePositionRotation.MR, ParticleDataActor0->ParticlePositionRotation.MX) : FTransform();
 
@@ -407,7 +345,7 @@ void FChaosVDSolverCollisionDataComponentVisualizer::DrawMidPhaseData(const UAct
 
 				PDI->SetHitProxy(nullptr);
 
-				if (ManifoldPoint.bIsSelectedInEditor)
+				if (CurrentSelectionContext && CurrentSelectionContext->ConstraintDataPtr == &Constraint && CurrentSelectionContext->ContactDataIndex == ContactIndex)
 				{
 					// We don't have an easy way to show a contact is selected with debug draw
 					// but 3D box surrounding the contact is better than nothing
diff --git a/Engine/Plugins/ChaosVD/Source/ChaosVD/Private/Visualizers/ChaosVDSolverCollisionDataComponentVisualizer.h b/Engine/Plugins/ChaosVD/Source/ChaosVD/Private/Visualizers/ChaosVDSolverCollisionDataComponentVisualizer.h
--- a/Engine/Plugins/ChaosVD/Source/ChaosVD/Private/Visualizers/ChaosVDSolverCollisionDataComponentVisualizer.h
+++ b/Engine/Plugins/ChaosVD/Source/ChaosVD/Private/Visualizers/ChaosVDSolverCollisionDataComponentVisualizer.h
@@ -1,58 +1,43 @@
 // Copyright Epic Games, Inc. All Rights Reserved.
 
 #pragma once
+#include "ChaosVDComponentVisualizerBase.h"
+#include "ChaosVDSolverDataSelection.h"
 #include "ComponentVisualizer.h"
 #include "HitProxies.h"
-#include "DataWrappers/ChaosVDCollisionDataWrappers.h"
 
+#include "ChaosVDSolverCollisionDataComponentVisualizer.generated.h"
+
+class FChaosVDScene;
+
+struct FChaosVDConstraint;
 struct FChaosVDParticlePairMidPhase;
 struct FChaosVDVisualizationContext;
-class FChaosVDScene;
 
-struct FChaosVDCollisionDataFinder
+USTRUCT()
+struct FChaosVDCollisionDataSelectionContext : public FChaosVDSelectionContext
 {
-	TWeakPtr<FChaosVDParticlePairMidPhase> OwningMidPhase;
-	FChaosVDConstraint* OwningConstraint = nullptr;
-	int32 ContactIndex = INDEX_NONE;
+	GENERATED_BODY()
 
-	void SetIsSelected(bool bNewSelected);
+	TSharedPtr<FChaosVDParticlePairMidPhase> MidPhase;
+	FChaosVDConstraint* ConstraintDataPtr = nullptr;
+	int32 ContactDataIndex = INDEX_NONE;
 };
 
-struct HChaosVDContactPointProxy : public HComponentVisProxy
-{
-	DECLARE_HIT_PROXY()
-	
-	HChaosVDContactPointProxy(const UActorComponent* Component, const FChaosVDCollisionDataFinder& InContactFinderData) : HComponentVisProxy(Component, HPP_UI), ContactFinder(InContactFinderData)
-	{	
-	}
-
-	virtual EMouseCursor::Type GetMouseCursor() override
-	{
-		return EMouseCursor::Crosshairs;
-	}
-
-	FChaosVDCollisionDataFinder ContactFinder;
-};
-
-class FChaosVDSolverCollisionDataComponentVisualizer : public FComponentVisualizer
+class FChaosVDSolverCollisionDataComponentVisualizer : public FChaosVDComponentVisualizerBase
 {
 public:
 	
 	FChaosVDSolverCollisionDataComponentVisualizer();
 	virtual ~FChaosVDSolverCollisionDataComponentVisualizer() override;
-	
-	void RegisterVisualizerMenus();
+
+	virtual void RegisterVisualizerMenus() override;
 
 	virtual bool ShowWhenSelected() override;
 	virtual void DrawVisualization(const UActorComponent* Component, const FSceneView* View, FPrimitiveDrawInterface* PDI) override;
 
-	virtual bool VisProxyHandleClick(FEditorViewportClient* InViewportClient, HComponentVisProxy* VisProxy, const FViewportClick& Click) override;
+	virtual bool CanHandleClick(const HChaosVDComponentVisProxy& VisProxy) override;
 
 protected:
-	
-	void ClearCurrentSelection();
-
 	void DrawMidPhaseData(const UActorComponent* Component, const TSharedPtr<FChaosVDParticlePairMidPhase>& MidPhase, const FChaosVDVisualizationContext& VisualizationContext, const FSceneView* View, FPrimitiveDrawInterface* PDI);
-
-	FChaosVDCollisionDataFinder CurrentSelectedContactData;
 };
diff --git a/Engine/Plugins/ChaosVD/Source/ChaosVD/Private/Visualizers/IChaosVDParticleVisualizationDataProvider.h b/Engine/Plugins/ChaosVD/Source/ChaosVD/Private/Visualizers/IChaosVDParticleVisualizationDataProvider.h
--- a/Engine/Plugins/ChaosVD/Source/ChaosVD/Private/Visualizers/IChaosVDParticleVisualizationDataProvider.h
+++ b/Engine/Plugins/ChaosVD/Source/ChaosVD/Private/Visualizers/IChaosVDParticleVisualizationDataProvider.h
@@ -10,16 +10,6 @@ class UChaosVDCollisionDataVisualizationSettings;
 struct FChaosVDParticleDataWrapper;
 class FChaosVDScene;
 
-/** Context needed to be able to visualize data in the viewport */
-struct FChaosVDVisualizationContext
-{
-	FTransform SpaceTransform;
-	TWeakPtr<FChaosVDScene> CVDScene;
-	int32 SolverID = INDEX_NONE;
-	uint32 VisualizationFlags = 0;
-	const UChaosVDSettingsObjectBase* DebugDrawSettings = nullptr;
-};
-
 /** Interface to be used by any object that contains Particle Data that needs to be visualized.
  * @note As we are still not settled on if we will stick with having Actors to represent each particle, this interface allows us to abstract the Particle data access
  * instead of using directly AChaosVDParticleActor
@@ -29,6 +19,5 @@ class IChaosVDParticleVisualizationDataProvider
 public:
 	virtual ~IChaosVDParticleVisualizationDataProvider() = default;
 
-	virtual void GetVisualizationContext(FChaosVDVisualizationContext& OutVisualizationContext) {}
-	virtual const FChaosVDParticleDataWrapper* GetParticleData() { return nullptr; }
+	virtual TSharedPtr<const FChaosVDParticleDataWrapper> GetParticleData() { return nullptr; }
 };
diff --git a/Engine/Plugins/ChaosVD/Source/ChaosVD/Private/Widgets/SChaosVDCharacterGroundConstraintDataInspector.cpp b/Engine/Plugins/ChaosVD/Source/ChaosVD/Private/Widgets/SChaosVDCharacterGroundConstraintDataInspector.cpp
deleted file mode 100644
diff --git a/Engine/Plugins/ChaosVD/Source/ChaosVD/Private/Widgets/SChaosVDCharacterGroundConstraintDataInspector.h b/Engine/Plugins/ChaosVD/Source/ChaosVD/Private/Widgets/SChaosVDCharacterGroundConstraintDataInspector.h
deleted file mode 100644
diff --git a/Engine/Plugins/ChaosVD/Source/ChaosVD/Private/Widgets/SChaosVDCollisionDataInspector.cpp b/Engine/Plugins/ChaosVD/Source/ChaosVD/Private/Widgets/SChaosVDCollisionDataInspector.cpp
--- a/Engine/Plugins/ChaosVD/Source/ChaosVD/Private/Widgets/SChaosVDCollisionDataInspector.cpp
+++ b/Engine/Plugins/ChaosVD/Source/ChaosVD/Private/Widgets/SChaosVDCollisionDataInspector.cpp
@@ -10,102 +10,11 @@
 #include "Visualizers/ChaosVDSolverCollisionDataComponentVisualizer.h"
 #include "Widgets/Input/SButton.h"
 #include "Widgets/Layout/SScrollBox.h"
-#include "Widgets/Layout/SUniformGridPanel.h"
 #include "Widgets/SChaosVDNameListPicker.h"
-#include "Widgets/SChaosVDWarningMessageBox.h"
 #include "Widgets/Text/STextBlock.h"
 
 #define LOCTEXT_NAMESPACE "ChaosVisualDebugger"
 
-SChaosVDCollisionDataInspector::~SChaosVDCollisionDataInspector()
-{
-	if (TSharedPtr<FChaosVDScene> ScenePtr = SceneWeakPtr.Pin())
-	{
-		ScenePtr->OnSceneUpdated().RemoveAll(this);
-	}
-}
-
-void SChaosVDCollisionDataInspector::Construct(const FArguments& InArgs, const TWeakPtr<FChaosVDScene>& InScenePtr)
-{
-	SceneWeakPtr = InScenePtr;
-
-	if (TSharedPtr<FChaosVDScene> ScenePtr = SceneWeakPtr.Pin())
-	{
-		ScenePtr->OnSceneUpdated().AddRaw(this, &SChaosVDCollisionDataInspector::HandleSceneUpdated);
-	}
-
-	MainCollisionDataDetailsView = CreateCollisionDataDetailsView();
-	SecondaryCollisionDataDetailsPanel = CreateCollisionDataDetailsView();
-
-	ChildSlot
-	[
-		SNew(SVerticalBox)
-		+SVerticalBox::Slot()
-		.AutoHeight()
-		.Padding(FMargin(10.0f, 10.0f, 10.0f, 0.0f))
-		[
-			GenerateObjectNameRowWidget().ToSharedRef()
-		]
-		+SVerticalBox::Slot()
-		.Padding(FMargin(10.0f, 5.0f, 10.0f, 1.0f))
-		.AutoHeight()
-		[
-			SAssignNew(CollisionDataAvailableList, SChaosVDNameListPicker)
-			.OnNameSleceted_Raw(this, &SChaosVDCollisionDataInspector::HandleCollisionDataEntryNameSelected)
-		]
-		+SVerticalBox::Slot()
-		.Padding(10)
-		[
-			SNew(SScrollBox)
-			.Visibility_Raw(this, &SChaosVDCollisionDataInspector::GetDetailsSectionVisibility)
-			+SScrollBox::Slot()
-			[
-				SNew(SBox)
-				.Padding(4.0f,2.0f,2.0f,10.0f)
-				[
-					SNew(SChaosVDWarningMessageBox)
-					.WarningText(LOCTEXT("CollisionDataOutOfData", "Scene change detected!. Selected collision data is out of date..."))
-					.Visibility_Raw(this, &SChaosVDCollisionDataInspector::GetOutOfDateWarningVisibility)
-				]
-			]
-			+SScrollBox::Slot()
-			.Padding(0.0f,5.0f,0.0f,15.0f)
-			[
-				SNew(SUniformGridPanel)
-				.SlotPadding(FAppStyle::GetMargin("StandardDialog.SlotPadding"))
-				.MinDesiredSlotWidth(FAppStyle::GetFloat("StandardDialog.MinDesiredSlotWidth"))
-				.MinDesiredSlotHeight(FAppStyle::GetFloat("StandardDialog.MinDesiredSlotHeight"))
-				+SUniformGridPanel::Slot(0, 0)
-				[
-					SNew(SButton)
-					.HAlign(HAlign_Center)
-					.ContentPadding(FAppStyle::GetMargin("StandardDialog.ContentPadding"))
-					.Text(LOCTEXT("SelectParticle0", "Select Particle 0"))
-					.OnClicked(this, &SChaosVDCollisionDataInspector::SelectParticleForCurrentCollisionData, EChaosVDCollisionParticleSelector::Index_0)
-				]
-				+SUniformGridPanel::Slot(1, 0)
-				[
-					SNew(SButton)
-					.HAlign(HAlign_Center)
-					.ContentPadding(FAppStyle::GetMargin("StandardDialog.ContentPadding"))
-					.Text(LOCTEXT("SelectParticle1", "Select Particle 1"))
-					.OnClicked(this, &SChaosVDCollisionDataInspector::SelectParticleForCurrentCollisionData, EChaosVDCollisionParticleSelector::Index_1)
-				]				
-			]
-			+SScrollBox::Slot()
-			.Padding(15.0f,0.0f,15.0f,0.0f)
-			[
-				MainCollisionDataDetailsView->GetWidget().ToSharedRef()
-			]
-			+SScrollBox::Slot()
-			.Padding(15.0f,10.0f,15.0f,5.0f)
-			[
-				SecondaryCollisionDataDetailsPanel->GetWidget().ToSharedRef()
-			]
-		]
-	];
-}
-
 void SChaosVDCollisionDataInspector::SetCollisionDataProviderObjectToInspect(IChaosVDCollisionDataProviderInterface* CollisionDataProvider)
 {
 	ClearInspector();
@@ -114,51 +23,74 @@ void SChaosVDCollisionDataInspector::SetCollisionDataProviderObjectToInspect(ICh
 	{
 		return;
 	}
+	
+	TSharedPtr<FChaosVDScene> ScenePtr = SceneWeakPtr.Pin();
+	if (!ScenePtr.IsValid())
+	{
+		return;
+	}
+	
+	TSharedPtr<FChaosVDSolverDataSelection> SolverDataSelectionObject = ScenePtr->GetSolverDataSelectionObject().Pin();
 
-	CurrentObjectBeingInspectedName = CollisionDataProvider->GetProviderName();
+	if (!SolverDataSelectionObject)
+	{
+		return;
+	}
 
 	TArray<TSharedPtr<FName>> NewCollisionDataEntriesNameList;
-	TArray<TSharedPtr<FChaosVDCollisionDataFinder>> FoundCollisionData;
+	TConstArrayView<TSharedPtr<FChaosVDParticlePairMidPhase>> FoundCollisionData = CollisionDataProvider->GetCollisionData();
 
-	CollisionDataProvider->GetCollisionData(FoundCollisionData);
-
-	for (const TSharedPtr<FChaosVDCollisionDataFinder>& CollisionData : FoundCollisionData)
+	for (const TSharedPtr<FChaosVDParticlePairMidPhase>& CollisionData : FoundCollisionData)
 	{
-		if (TSharedPtr<FName> CollisionDataEntryName = GenerateNameForCollisionDataItem(*CollisionData))
+		TSharedPtr<FChaosVDSolverDataSelectionHandle> SelectionHandle = SolverDataSelectionObject->MakeSelectionHandle(CollisionData);
+		if (TSharedPtr<FName> CollisionDataEntryName = GenerateNameForCollisionDataItem(SelectionHandle))
 		{
-			CollisionDataByNameMap.Add(*CollisionDataEntryName, CollisionData);
+			CollisionDataByNameMap.Add(*CollisionDataEntryName, SelectionHandle);
 			NewCollisionDataEntriesNameList.Add(CollisionDataEntryName);
+
+			if (!CurrentDataSelectionHandle->IsValid())
+			{
+				CurrentSelectedName = CollisionDataEntryName;
+			}
 		}
 	}
 
 	CollisionDataAvailableList->UpdateNameList(MoveTemp(NewCollisionDataEntriesNameList));
+	CollisionDataAvailableList->SelectName(CurrentSelectedName, ESelectInfo::OnMouseClick);
 	bIsUpToDate = true;
 }
 
-void SChaosVDCollisionDataInspector::SetSingleContactDataToInspect(const FChaosVDCollisionDataFinder& InContactFinderData)
+void SChaosVDCollisionDataInspector::SetConstraintDataToInspect(const TSharedPtr<FChaosVDSolverDataSelectionHandle>& InDataSelectionHandle)
 {
 	ClearInspector();
 
-	TSharedPtr<FChaosVDCollisionDataFinder> ContactFinderDataPtr = MakeShared<FChaosVDCollisionDataFinder>();
-	*ContactFinderDataPtr = InContactFinderData;
-
-	CurrentSelectedName = GenerateNameForCollisionDataItem(*ContactFinderDataPtr);
+	if (!InDataSelectionHandle || !InDataSelectionHandle->IsA<FChaosVDParticlePairMidPhase>())
+	{
+		return;
+	}
 	
-	// TODO: We need a way to provide a name when we receive collision data directly.
-	// Currently this is only possible when clicking contacts in the viewport
-	CurrentObjectBeingInspectedName = TEXT("Contact");
+	CurrentDataSelectionHandle = InDataSelectionHandle.ToSharedRef();
 
-	CollisionDataByNameMap.Add(*CurrentSelectedName, ContactFinderDataPtr);
+	CurrentSelectedName = GenerateNameForCollisionDataItem(InDataSelectionHandle);
+
+	CollisionDataByNameMap.Add(*CurrentSelectedName, InDataSelectionHandle);
 
 	CollisionDataAvailableList->UpdateNameList({ CurrentSelectedName });
 	CollisionDataAvailableList->SelectName(CurrentSelectedName, ESelectInfo::OnMouseClick);
 }
 
+void SChaosVDCollisionDataInspector::SetupWidgets()
+{
+	SChaosVDConstraintDataInspector::SetupWidgets();
+
+	SecondaryCollisionDataDetailsPanel = CreateCollisionDataDetailsView();
+}
+
 void SChaosVDCollisionDataInspector::HandleSceneUpdated()
 {
-	if (TSharedPtr<FChaosVDCollisionDataFinder> DataFinderPtr = GetCurrentDataBeingInspected())
+	if (const TSharedPtr<FChaosVDSolverDataSelectionHandle>& SelectionHandle = GetCurrentDataBeingInspected())
 	{
-		if (DataFinderPtr.IsValid() && DataFinderPtr->OwningMidPhase.Pin())
+		if (SelectionHandle->IsSelected() && SelectionHandle->GetData<FChaosVDParticlePairMidPhase>())
 		{
 			bIsUpToDate = false;
 		}
@@ -171,142 +103,152 @@ void SChaosVDCollisionDataInspector::HandleSceneUpdated()
 
 void SChaosVDCollisionDataInspector::ClearInspector()
 {
-	MainCollisionDataDetailsView->SetStructureData(nullptr);
-	SecondaryCollisionDataDetailsPanel->SetStructureData(nullptr);
-
 	CollisionDataByNameMap.Reset();
 	CollisionDataAvailableList->UpdateNameList({});
 
 	CurrentSelectedName = nullptr;
 
-	CurrentObjectBeingInspectedName = FName();
-
 	bIsUpToDate = true;
+	
+	SecondaryCollisionDataDetailsPanel->SetStructureData(nullptr);
+
+	SChaosVDConstraintDataInspector::ClearInspector();
 }
 
-FText SChaosVDCollisionDataInspector::GetObjectBeingInspectedName() const
+TSharedRef<SWidget> SChaosVDCollisionDataInspector::GenerateHeaderWidget(FMargin Margin)
 {
-	return FText::AsCultureInvariant(CurrentObjectBeingInspectedName.ToString());
+	return SNew(SVerticalBox)
+			+SVerticalBox::Slot()
+			.Padding(Margin)
+			.AutoHeight()
+			[
+				SAssignNew(CollisionDataAvailableList, SChaosVDNameListPicker)
+				.OnNameSleceted_Raw(this, &SChaosVDCollisionDataInspector::HandleCollisionDataEntryNameSelected)
+			];
 }
 
-TSharedPtr<SWidget> SChaosVDCollisionDataInspector::GenerateObjectNameRowWidget()
+TSharedRef<SWidget> SChaosVDCollisionDataInspector::GenerateDetailsViewWidget(FMargin Margin)
 {
-	TSharedPtr<SWidget> RowWidget = SNew(SBorder)
-		.BorderImage(FAppStyle::Get().GetBrush("DetailsView.CategoryMiddle"))
-		.BorderBackgroundColor(FAppStyle::Get().GetSlateColor("Colors.Panel"))
-		.Padding(0)
-		[
-			SNew(SBox)
-			.Padding(2.0f)
+	return SNew(SScrollBox)
+			+SScrollBox::Slot()
+			.Padding(Margin)
 			[
-				SNew(SSplitter)
-				.Style(FAppStyle::Get(), "DetailsView.Splitter")
-				.PhysicalSplitterHandleSize(1.5f)
-				.HitDetectionSplitterHandleSize(5.0f)
-				+SSplitter::Slot()
-				.Value(0.2)
-				[
-					SNew(SBox)
-					.VAlign(VAlign_Center)
-					.Padding(1.0)
-					.VAlign(VAlign_Center)
-					[
-						SNew(STextBlock)
-						.Text(LOCTEXT("CollisionDataInspectorObject","Object"))
-						.Font(FAppStyle::GetFontStyle("PropertyWindow.NormalFont"))
-					]	
-				]
-				+SSplitter::Slot()
-				[
-					SNew(SBox)
-					.VAlign(VAlign_Center)
-					.Padding(1.0)
-					[
-						SNew(STextBlock)
-						.Text_Raw(this, &SChaosVDCollisionDataInspector::GetObjectBeingInspectedName)
-						.Font(FAppStyle::GetFontStyle("PropertyWindow.NormalFont"))
-					]
-				]
+				MainDataDetailsView->GetWidget().ToSharedRef()
 			]
-		];
-
-	return RowWidget;
+			+SScrollBox::Slot()
+			.Padding(Margin)
+			[
+				SecondaryCollisionDataDetailsPanel->GetWidget().ToSharedRef()
+			];
 }
 
-
-EVisibility SChaosVDCollisionDataInspector::GetOutOfDateWarningVisibility() const
+FText SChaosVDCollisionDataInspector::GetParticleName(EChaosVDParticlePairIndex ParticleSlot, const TSharedPtr<FChaosVDSolverDataSelectionHandle>& InSelectionHandle) const
 {
-	return bIsUpToDate ? EVisibility::Collapsed : EVisibility::Visible;
+	int32 OutSolverID = INDEX_NONE;
+	int32 OutParticleID = INDEX_NONE;
+	GetParticleIDForSelectedData(InSelectionHandle, ParticleSlot, OutSolverID, OutParticleID);
+
+	if (const TSharedPtr<FChaosVDScene> ScenePtr = SceneWeakPtr.Pin())
+	{
+		if (AChaosVDParticleActor* ParticleActor = ScenePtr->GetParticleActor(OutSolverID, OutParticleID))
+		{
+			if (TSharedPtr<const FChaosVDParticleDataWrapper> ParticleData = ParticleActor->GetParticleData())
+			{
+				return FText::AsCultureInvariant(ParticleData->DebugName);
+			}
+		}
+	}
+
+	return FText::GetEmpty();
 }
 
-void SChaosVDCollisionDataInspector::HandleCollisionDataEntryNameSelected(TSharedPtr<FName> SelectedName)
-{
-	CurrentSelectedName = SelectedName;
 
-	if (!CurrentSelectedName.IsValid())
+void SChaosVDCollisionDataInspector::GetParticleIDForSelectedData(const TSharedPtr<FChaosVDSolverDataSelectionHandle>& InSelectionHandle, EChaosVDParticlePairIndex ParticleSlot, int32& OutSolverID, int32& OutParticleID) const
+{
+	if (!InSelectionHandle)
 	{
 		return;
 	}
 
-	if (TSharedPtr<FChaosVDCollisionDataFinder> DataFinderPtr = GetCurrentDataBeingInspected())
+	if (const FChaosVDParticlePairMidPhase* MidPhasePtr = InSelectionHandle->GetData<FChaosVDParticlePairMidPhase>())
 	{
-		if (DataFinderPtr->OwningConstraint)
+		bool bHasConstraintData = false;
+		if (FChaosVDCollisionDataSelectionContext* SelectionContext = InSelectionHandle->GetContextData<FChaosVDCollisionDataSelectionContext>())
 		{
-			FChaosVDConstraint* MutableConstraint = const_cast<FChaosVDConstraint*>(DataFinderPtr->OwningConstraint);
-			TSharedPtr<FStructOnScope> ConstraintView = MakeShared<FStructOnScope>(FChaosVDConstraint::StaticStruct(), reinterpret_cast<uint8*>(MutableConstraint));
-			SecondaryCollisionDataDetailsPanel->SetStructureData(ConstraintView); 
-
-			if (DataFinderPtr->OwningConstraint->ManifoldPoints.IsValidIndex(DataFinderPtr->ContactIndex))
-			{
-				const FChaosVDManifoldPoint* ContactPointData = &DataFinderPtr->OwningConstraint->ManifoldPoints[DataFinderPtr->ContactIndex];
-				TSharedPtr<FStructOnScope> ContactViewView = MakeShared<FStructOnScope>(FChaosVDManifoldPoint::StaticStruct(), reinterpret_cast<uint8*>(const_cast<FChaosVDManifoldPoint*>(ContactPointData)));
-				MainCollisionDataDetailsView->SetStructureData(ContactViewView);
-			}
-			else
+			if (SelectionContext->ConstraintDataPtr)
 			{
-				MainCollisionDataDetailsView->SetStructureData(nullptr); 
+				bHasConstraintData = true;
+				OutParticleID = ParticleSlot == EChaosVDParticlePairIndex::Index_0 ? SelectionContext->ConstraintDataPtr->Particle0Index : SelectionContext->ConstraintDataPtr->Particle1Index;
 			}
 		}
-		else if (TSharedPtr<FChaosVDParticlePairMidPhase> MidPhaseData = DataFinderPtr->OwningMidPhase.Pin())
+		
+		if (!bHasConstraintData)
 		{
-			// If we have a recorded midphase with not contact data, show that instead
-			const TSharedPtr<FStructOnScope> MidPhaseView = MakeShared<FStructOnScope>(FChaosVDParticlePairMidPhase::StaticStruct(), reinterpret_cast<uint8*>(MidPhaseData.Get()));
-			MainCollisionDataDetailsView->SetStructureData(MidPhaseView); 
+			OutParticleID = ParticleSlot == EChaosVDParticlePairIndex::Index_0 ? MidPhasePtr->Particle0Idx : MidPhasePtr->Particle1Idx;
 		}
+
+		OutSolverID = MidPhasePtr->SolverID;
 	}
 }
 
-TSharedPtr<FName> SChaosVDCollisionDataInspector::GenerateNameForCollisionDataItem(const FChaosVDCollisionDataFinder& InContactFinderData)
+void SChaosVDCollisionDataInspector::HandleCollisionDataEntryNameSelected(TSharedPtr<FName> SelectedName)
 {
-	int32 Particle0Index = INDEX_NONE;
-	int32 Particle1Index = INDEX_NONE;
+	CurrentSelectedName = SelectedName;
 
-	bool bHasConstraintData = false;
+	if (!CurrentSelectedName.IsValid())
+	{
+		return;
+	}
 
-	if (const TSharedPtr<FChaosVDParticlePairMidPhase> MidPhaseData = InContactFinderData.OwningMidPhase.Pin())
+	CurrentDataSelectionHandle = GetSelectionNameForName(SelectedName).ToSharedRef();
+
+	if (CurrentDataSelectionHandle->IsValid())
 	{
-		if (InContactFinderData.OwningConstraint)
-		{
-			Particle0Index = InContactFinderData.OwningConstraint->Particle0Index;
-			Particle1Index = InContactFinderData.OwningConstraint->Particle1Index;
-			bHasConstraintData = true;
-		}
-		else
+		if (FChaosVDParticlePairMidPhase* MidPhasePtr = CurrentDataSelectionHandle->GetData<FChaosVDParticlePairMidPhase>())
 		{
-			Particle0Index = MidPhaseData->Particle0Idx;
-			Particle1Index = MidPhaseData->Particle1Idx;	
+			bool bHasConstraintData = false;
+			if (FChaosVDCollisionDataSelectionContext* SelectionContext = CurrentDataSelectionHandle->GetContextData<FChaosVDCollisionDataSelectionContext>())
+			{
+				if (SelectionContext->ConstraintDataPtr)
+				{
+					bHasConstraintData = true;
+
+					TSharedPtr<FStructOnScope> ConstraintView = MakeShared<FStructOnScope>(FChaosVDConstraint::StaticStruct(), reinterpret_cast<uint8*>(SelectionContext->ConstraintDataPtr));
+					SecondaryCollisionDataDetailsPanel->SetStructureData(ConstraintView); 
+
+					if (SelectionContext->ConstraintDataPtr->ManifoldPoints.IsValidIndex(SelectionContext->ContactDataIndex))
+					{
+						const FChaosVDManifoldPoint* ContactPointData = &SelectionContext->ConstraintDataPtr->ManifoldPoints[SelectionContext->ContactDataIndex];
+						TSharedPtr<FStructOnScope> ContactViewView = MakeShared<FStructOnScope>(FChaosVDManifoldPoint::StaticStruct(), reinterpret_cast<uint8*>(const_cast<FChaosVDManifoldPoint*>(ContactPointData)));
+						MainDataDetailsView->SetStructureData(ContactViewView);
+					}
+					else
+					{
+						MainDataDetailsView->SetStructureData(nullptr); 
+					}
+				}
+			}
+
+			if (!bHasConstraintData)
+			{
+				// If we have a recorded midphase with not contact data, show that instead
+				const TSharedPtr<FStructOnScope> MidPhaseView = MakeShared<FStructOnScope>(FChaosVDParticlePairMidPhase::StaticStruct(), reinterpret_cast<uint8*>(MidPhasePtr));
+				MainDataDetailsView->SetStructureData(MidPhaseView); 
+			}
 		}
 	}
+}
 
-	if (Particle0Index != INDEX_NONE || Particle1Index != INDEX_NONE)
+TSharedRef<FName> SChaosVDCollisionDataInspector::GenerateNameForCollisionDataItem(const TSharedPtr<FChaosVDSolverDataSelectionHandle>& InDataSelectionHandle) const
+{
+	static TSharedPtr<FName> InvalidName = MakeShared<FName>(NAME_None);
+	if (!InDataSelectionHandle)
 	{
-		static FText MidPhaseOnlyText = LOCTEXT("CollisionItemDataTitleMidPhaseOnly", "[MidPhase Only]");
-		static FText ConstraintText = LOCTEXT("CollisionItemDataTitleConstaint", "[Constraint]");
-		const FText GeneratedName = FText::Format(LOCTEXT("CollisionItemDataTitle", "Particle Pair | Index0 [ID {0}] <-> Index1 [{1}] | Type {2}"), Particle0Index, Particle1Index, bHasConstraintData ? ConstraintText: MidPhaseOnlyText);
-		return MakeShared<FName>(GeneratedName.ToString());
+		return InvalidName.ToSharedRef();
 	}
 
-	return nullptr;
+	const FText GeneratedName = FText::Format(LOCTEXT("CollisionItemDataTitle", "Particle Pair | [ {0} ] <-> [ {1} ] "), GetParticleName(EChaosVDParticlePairIndex::Index_0, InDataSelectionHandle), GetParticleName(EChaosVDParticlePairIndex::Index_1, InDataSelectionHandle));
+	return MakeShared<FName>(GeneratedName.ToString());
 }
 
 TSharedPtr<IStructureDetailsView> SChaosVDCollisionDataInspector::CreateCollisionDataDetailsView()
@@ -328,54 +270,41 @@ EVisibility SChaosVDCollisionDataInspector::GetDetailsSectionVisibility() const
 	return CurrentSelectedName ? EVisibility::Visible : EVisibility::Collapsed;
 }
 
-FReply SChaosVDCollisionDataInspector::SelectParticleForCurrentCollisionData(EChaosVDCollisionParticleSelector ParticleSelector)
+FReply SChaosVDCollisionDataInspector::SelectParticleForCurrentSelectedData(EChaosVDParticlePairIndex ParticleSlot)
 {
-	TSharedPtr<FChaosVDScene> ScenePtr = SceneWeakPtr.Pin();
-	if (!ScenePtr.IsValid())
-	{
-		return FReply::Handled();
-	}
-	
-	if (const TSharedPtr<FChaosVDCollisionDataFinder> DataFinderPtr = GetCurrentDataBeingInspected())
-	{
-		if (const FChaosVDMidPhasePtr MidPhasePtr = DataFinderPtr->OwningMidPhase.Pin())
-		{
-			int32 ParticleIndex = INDEX_NONE;
-			if (DataFinderPtr->OwningConstraint)
-			{
-				ParticleIndex = ParticleSelector == EChaosVDCollisionParticleSelector::Index_0 ? DataFinderPtr->OwningConstraint->Particle0Index : DataFinderPtr->OwningConstraint->Particle1Index;
-			}
-			else
-			{
-				ParticleIndex = ParticleSelector == EChaosVDCollisionParticleSelector::Index_0 ? MidPhasePtr->Particle0Idx : MidPhasePtr->Particle1Idx;
-			}
+	int32 OutParticleID = INDEX_NONE;
+	int32 OutSolverID = INDEX_NONE;
 
-			if (AChaosVDParticleActor* ParticleActor = ScenePtr->GetParticleActor(MidPhasePtr->SolverID, ParticleIndex))
-			{
-				ScenePtr->SetSelectedObject(ParticleActor);
-			}
-		}
-	}
+	GetParticleIDForSelectedData(GetCurrentDataBeingInspected(), ParticleSlot,OutSolverID,OutParticleID);
+	
+	SelectParticle(OutSolverID, OutParticleID);
 
 	return FReply::Handled();
 }
 
-TSharedPtr<FChaosVDCollisionDataFinder> SChaosVDCollisionDataInspector::GetCurrentDataBeingInspected()
+
+TSharedPtr<FChaosVDSolverDataSelectionHandle> SChaosVDCollisionDataInspector::GetSelectionNameForName(const TSharedPtr<FName>& InName)
 {
-	if (!CurrentSelectedName.IsValid())
+	if (!InName.IsValid())
 	{
-		return nullptr;
+		return MakeShared<FChaosVDSolverDataSelectionHandle>();
 	}
 
-	if (TSharedPtr<FChaosVDCollisionDataFinder>* DataFinderPtrPtr = CollisionDataByNameMap.Find(*CurrentSelectedName))
+	if (const TSharedPtr<FChaosVDSolverDataSelectionHandle>* SelectionHandlePtrPtr = CollisionDataByNameMap.Find(*InName))
 	{
-		if (TSharedPtr<FChaosVDCollisionDataFinder> DataFinderPtr = *DataFinderPtrPtr)
+		if (const TSharedPtr<FChaosVDSolverDataSelectionHandle>& SelectionHandlePtr = *SelectionHandlePtrPtr)
 		{
-			return DataFinderPtr;
+			return SelectionHandlePtr;
 		}
 	}
 
-	return nullptr;
+	return MakeShared<FChaosVDSolverDataSelectionHandle>();
+}
+
+
+const TSharedRef<FChaosVDSolverDataSelectionHandle>& SChaosVDCollisionDataInspector::GetCurrentDataBeingInspected() const
+{
+	return CurrentDataSelectionHandle;
 }
 
 #undef LOCTEXT_NAMESPACE
diff --git a/Engine/Plugins/ChaosVD/Source/ChaosVD/Private/Widgets/SChaosVDCollisionDataInspector.h b/Engine/Plugins/ChaosVD/Source/ChaosVD/Private/Widgets/SChaosVDCollisionDataInspector.h
--- a/Engine/Plugins/ChaosVD/Source/ChaosVD/Private/Widgets/SChaosVDCollisionDataInspector.h
+++ b/Engine/Plugins/ChaosVD/Source/ChaosVD/Private/Widgets/SChaosVDCollisionDataInspector.h
@@ -2,8 +2,10 @@
 
 #pragma once
 
+#include "SChaosVDConstraintDataInspector.h"
 #include "Widgets/SCompoundWidget.h"
 
+struct FChaosVDSolverDataSelectionHandle;
 class AChaosVDParticleActor;
 class FChaosVDScene;
 class IStructureDetailsView;
@@ -13,68 +15,46 @@ class SChaosVDNameListPicker;
 struct FChaosVDCollisionDataFinder;
 struct FChaosVDParticlePairMidPhase;
 
-typedef TSharedPtr<FChaosVDParticlePairMidPhase> FChaosVDMidPhasePtr;
-
-enum class EChaosVDCollisionParticleSelector
-{
-	Index_0,
-	Index_1
-};
-
-class SChaosVDCollisionDataInspector : public SCompoundWidget
+class SChaosVDCollisionDataInspector : public SChaosVDConstraintDataInspector
 {
 public:
 
-	SLATE_BEGIN_ARGS(SChaosVDCollisionDataInspector)
-	{
-	}
-
-	SLATE_END_ARGS()
-
-	virtual ~SChaosVDCollisionDataInspector() override;
-
-	/** Constructs this widget with InArgs */
-	void Construct(const FArguments& InArgs, const TWeakPtr<FChaosVDScene>& InScenePtr);
-
 	void SetCollisionDataProviderObjectToInspect(IChaosVDCollisionDataProviderInterface* CollisionDataProvider);
-	void SetSingleContactDataToInspect(const FChaosVDCollisionDataFinder& InContactFinderData);
+
+	virtual void SetConstraintDataToInspect(const TSharedPtr<FChaosVDSolverDataSelectionHandle>& InDataSelectionHandle) override;
 
 protected:
+	virtual void SetupWidgets() override;
+	virtual void HandleSceneUpdated() override;
 
-	void HandleSceneUpdated();
+	virtual void ClearInspector() override;
 
-	void ClearInspector();
+	virtual TSharedRef<SWidget> GenerateHeaderWidget(FMargin Margin) override;
+	virtual TSharedRef<SWidget> GenerateDetailsViewWidget(FMargin Margin) override;
 
-	FText GetObjectBeingInspectedName() const;
-	TSharedPtr<SWidget> GenerateObjectNameRowWidget();
+	virtual FText GetParticleName(EChaosVDParticlePairIndex ParticleSlot, const TSharedPtr<FChaosVDSolverDataSelectionHandle>& InSelectionHandle) const override;
 
-	TSharedPtr<FChaosVDCollisionDataFinder> GetCurrentDataBeingInspected();
+	void GetParticleIDForSelectedData(const TSharedPtr<FChaosVDSolverDataSelectionHandle>& InSelectionHandle, EChaosVDParticlePairIndex ParticleSlot, int32& OutSolverID, int32& OutParticleID) const;
 
-	EVisibility GetOutOfDateWarningVisibility() const;
+	virtual const TSharedRef<FChaosVDSolverDataSelectionHandle>& GetCurrentDataBeingInspected() const override;
 
 	void HandleCollisionDataEntryNameSelected(TSharedPtr<FName> SelectedName);
 
-	TSharedPtr<FName> GenerateNameForCollisionDataItem(const FChaosVDCollisionDataFinder& InContactFinderData);
+	TSharedRef<FName> GenerateNameForCollisionDataItem(const TSharedPtr<FChaosVDSolverDataSelectionHandle>& InDataSelectionHandle) const;
 
 	TSharedPtr<IStructureDetailsView> CreateCollisionDataDetailsView();
 
-	EVisibility GetDetailsSectionVisibility() const;
+	virtual EVisibility GetDetailsSectionVisibility() const override;
 
-	FReply SelectParticleForCurrentCollisionData(EChaosVDCollisionParticleSelector ParticleSelector);
+	virtual FReply SelectParticleForCurrentSelectedData(EChaosVDParticlePairIndex ParticleSlot) override;
 
-	TWeakPtr<FChaosVDScene> SceneWeakPtr;
+	TSharedPtr<FChaosVDSolverDataSelectionHandle> GetSelectionNameForName(const TSharedPtr<FName>& InName);
 
 	TSharedPtr<SChaosVDNameListPicker> CollisionDataAvailableList;
 	
-	TMap<FName, TSharedPtr<FChaosVDCollisionDataFinder>> CollisionDataByNameMap;
+	TMap<FName, TSharedPtr<FChaosVDSolverDataSelectionHandle>> CollisionDataByNameMap;
 
 	TSharedPtr<FName> CurrentSelectedName;
 
-	TSharedPtr<IStructureDetailsView> MainCollisionDataDetailsView;
-	
 	TSharedPtr<IStructureDetailsView> SecondaryCollisionDataDetailsPanel;
-
-	FName CurrentObjectBeingInspectedName;
-
-	bool bIsUpToDate = true;
 };
diff --git a/Engine/Plugins/ChaosVD/Source/ChaosVD/Private/Widgets/SChaosVDConstraintDataInspector.cpp b/Engine/Plugins/ChaosVD/Source/ChaosVD/Private/Widgets/SChaosVDConstraintDataInspector.cpp
--- a/Engine/Plugins/ChaosVD/Source/ChaosVD/Private/Widgets/SChaosVDConstraintDataInspector.cpp
+++ b/Engine/Plugins/ChaosVD/Source/ChaosVD/Private/Widgets/SChaosVDConstraintDataInspector.cpp
@@ -7,8 +7,6 @@
 #include "PropertyEditorModule.h"
 #include "Actors/ChaosVDSolverInfoActor.h"
 #include "Widgets/SChaosVDWarningMessageBox.h"
-#include "Components/ChaosVDSolverJointConstraintDataComponent.h"
-#include "DataWrappers/ChaosVDJointDataWrappers.h"
 #include "Modules/ModuleManager.h"
 #include "Widgets/Input/SButton.h"
 #include "Widgets/Layout/SBox.h"
@@ -19,6 +17,46 @@
 
 #define LOCTEXT_NAMESPACE "ChaosVisualDebugger"
 
+FReadOnlyCopyStructOnScope::FReadOnlyCopyStructOnScope(const FStructOnScope& StructToCopy)
+{
+	if (const UStruct* ScriptStructPtr = StructToCopy.GetStruct())
+	{
+		ScriptStruct = ScriptStructPtr;
+		SampleStructMemory = (uint8*)FMemory::Malloc(ScriptStructPtr->GetStructureSize() ? ScriptStructPtr->GetStructureSize() : 1, ScriptStructPtr->GetMinAlignment());
+		ScriptStructPtr->InitializeStruct(SampleStructMemory);
+		OwnsMemory = true;
+
+		if (ensure(SampleStructMemory))
+		{
+			UScriptStruct* AsScript = Cast<UScriptStruct>(const_cast<UStruct*>(ScriptStructPtr));
+			if (ensure(AsScript))
+			{
+				AsScript->CopyScriptStruct(SampleStructMemory,StructToCopy.GetStructMemory());
+			}
+		}
+	}
+}
+
+void FReadOnlyCopyStructOnScope::UpdateFromOther(const FStructOnScope& StructToCopy)
+{
+	const UStruct* StructPtr = StructToCopy.GetStruct();
+	if (StructPtr && ensure(ScriptStruct == StructToCopy.GetStruct()))
+	{
+		if (ensure(SampleStructMemory))
+		{
+			UScriptStruct* AsScript = Cast<UScriptStruct>(const_cast<UStruct*>(StructToCopy.GetStruct()));
+			if (ensure(AsScript))
+			{
+				AsScript->CopyScriptStruct(SampleStructMemory,StructToCopy.GetStructMemory());
+			}
+		}
+	}
+}
+
+SChaosVDConstraintDataInspector::SChaosVDConstraintDataInspector() : CurrentDataSelectionHandle(MakeShared<FChaosVDSolverDataSelectionHandle>())
+{ 
+}
+
 SChaosVDConstraintDataInspector::~SChaosVDConstraintDataInspector()
 {
 	UnregisterSceneEvents();
@@ -30,7 +68,7 @@ void SChaosVDConstraintDataInspector::Construct(const FArguments& InArgs, const
 
 	RegisterSceneEvents();
 
-	ConstraintDataDetailsView = CreateDataDetailsView();
+	SetupWidgets();
 
 	constexpr float NoPadding = 0.0f;
 	constexpr float OuterBoxPadding = 2.0f;
@@ -46,6 +84,12 @@ void SChaosVDConstraintDataInspector::Construct(const FArguments& InArgs, const
 		+SVerticalBox::Slot()
 		.AutoHeight()
 		.Padding(OuterInnerPadding)
+		[
+			GenerateHeaderWidget({})
+		]
+		+SVerticalBox::Slot()
+		.AutoHeight()
+		.Padding(OuterInnerPadding)
 		[
 			SNew(SBox)
 			.Visibility_Raw(this, &SChaosVDConstraintDataInspector::GetOutOfDateWarningVisibility)
@@ -69,101 +113,107 @@ void SChaosVDConstraintDataInspector::Construct(const FArguments& InArgs, const
 			.Visibility_Raw(this, &SChaosVDConstraintDataInspector::GetNothingSelectedMessageVisibility)
 			.Justification(ETextJustify::Center)
 			.TextStyle(FAppStyle::Get(), "DetailsView.BPMessageTextStyle")
-			.Text(LOCTEXT("ConstraintDataNoSelectedMessage", "Select a Joint Constraint in the viewport to see its details..."))
+			.Text(LOCTEXT("ConstraintDataNoSelectedMessage", "Select a Constraint in the viewport to see its details..."))
 			.AutoWrapText(true)
 		]
 		+SVerticalBox::Slot()
 		.Padding(OuterInnerPadding)
 		[
-			SNew(SScrollBox)
-			.Visibility_Raw(this, &SChaosVDConstraintDataInspector::GetDetailsSectionVisibility)
-			+SScrollBox::Slot()
-			.Padding(InnerDetailsPanelsHorizontalPadding,NoPadding,InnerDetailsPanelsHorizontalPadding,InnerDetailsPanelsVerticalPadding)
-			[
-				ConstraintDataDetailsView->GetWidget().ToSharedRef()
-			]
+			GenerateDetailsViewWidget({InnerDetailsPanelsHorizontalPadding,NoPadding,InnerDetailsPanelsHorizontalPadding,InnerDetailsPanelsVerticalPadding})
 		]
 	];
 }
 
-void SChaosVDConstraintDataInspector::RegisterSelectionEventsForSolver(AChaosVDSolverInfoActor* SolverInfo)
-{
-	if (UChaosVDSolverJointConstraintDataComponent* JointDataComponent = SolverInfo ? SolverInfo->GetJointsDataComponent() : nullptr)
-	{
-		JointDataComponent->OnSelectionChanged().AddRaw(this, &SChaosVDConstraintDataInspector::SetConstraintDataToInspect);
-	}
-}
-
-void SChaosVDConstraintDataInspector::UnregisterSelectionEventsForSolver(AChaosVDSolverInfoActor* SolverInfo)
-{
-	if (UChaosVDSolverJointConstraintDataComponent* JointDataComponent = SolverInfo ? SolverInfo->GetJointsDataComponent() : nullptr)
-	{
-		JointDataComponent->OnSelectionChanged().RemoveAll(this);
-	}
-}
-
 void SChaosVDConstraintDataInspector::RegisterSceneEvents()
 {
 	if (const TSharedPtr<FChaosVDScene> ScenePtr = SceneWeakPtr.Pin())
 	{
 		ScenePtr->OnSceneUpdated().AddRaw(this, &SChaosVDConstraintDataInspector::HandleSceneUpdated);
-		ScenePtr->OnSolverInfoActorCreated().AddRaw(this, &SChaosVDConstraintDataInspector::RegisterSelectionEventsForSolver);
-
-		for (const TPair<int32, AChaosVDSolverInfoActor*> SolverInfos : ScenePtr->GetSolverInfoActorsMap())
+		if (TSharedPtr<FChaosVDSolverDataSelection> SelectionObject = ScenePtr->GetSolverDataSelectionObject().Pin())
 		{
-			RegisterSelectionEventsForSolver(SolverInfos.Value);
+			SelectionObject->GetDataSelectionChangedDelegate().AddRaw(this, &SChaosVDConstraintDataInspector::SetConstraintDataToInspect);
 		}
 	}
 }
 
-void SChaosVDConstraintDataInspector::UnregisterSceneEvents()
+void SChaosVDConstraintDataInspector::UnregisterSceneEvents() const
 {
 	if (const TSharedPtr<FChaosVDScene> ScenePtr = SceneWeakPtr.Pin())
 	{
 		ScenePtr->OnSceneUpdated().RemoveAll(this);
-		ScenePtr->OnSolverInfoActorCreated().RemoveAll(this);
 
-		for (const TPair<int32, AChaosVDSolverInfoActor*> SolverInfos : ScenePtr->GetSolverInfoActorsMap())
+		if (TSharedPtr<FChaosVDSolverDataSelection> SelectionObject = ScenePtr->GetSolverDataSelectionObject().Pin())
 		{
-			UnregisterSelectionEventsForSolver(SolverInfos.Value);
+			SelectionObject->GetDataSelectionChangedDelegate().RemoveAll(this);
 		}
 	}
 }
 
-void SChaosVDConstraintDataInspector::SetConstraintDataToInspect(const FChaosVDJointConstraintSelectionHandle& InDataSelectionHandle)
+void SChaosVDConstraintDataInspector::SetConstraintDataToInspect(const TSharedPtr<FChaosVDSolverDataSelectionHandle>& InDataSelectionHandle)
 {
-	ClearInspector();
+	if (InDataSelectionHandle && InDataSelectionHandle->IsA<FChaosVDConstraintDataWrapperBase>())
+	{
+		CurrentDataSelectionHandle = InDataSelectionHandle.ToSharedRef();
 
-	if (const TSharedPtr<FChaosVDJointConstraint> QueryDataToInspect = InDataSelectionHandle.GetData().Pin())
+		if (HasCompatibleStructScopeView(InDataSelectionHandle.ToSharedRef()))
+		{
+			DataBeingInspectedCopy->UpdateFromOther(*InDataSelectionHandle->GetDataAsStructScope());
+		}
+		else
+		{
+			DataBeingInspectedCopy = MakeShared<FReadOnlyCopyStructOnScope>(*InDataSelectionHandle->GetDataAsStructScope());
+			MainDataDetailsView->SetStructureData(DataBeingInspectedCopy);
+		}
+	}
+	else
 	{
-		CurrentDataSelectionHandle = InDataSelectionHandle;
-		const TSharedPtr<FStructOnScope> QueryDataView = MakeShared<FStructOnScope>(FChaosVDJointConstraint::StaticStruct(), reinterpret_cast<uint8*>(QueryDataToInspect.Get()));
-		ConstraintDataDetailsView->SetStructureData(QueryDataView);
+		ClearInspector();
 	}
 
 	bIsUpToDate = true;
 }
 
-FText SChaosVDConstraintDataInspector::GetParticleName(EChaosVDParticleSelector PairIndex) const
+void SChaosVDConstraintDataInspector::SetupWidgets()
+{
+	MainDataDetailsView = CreateDataDetailsView();
+}
+
+TSharedRef<SWidget> SChaosVDConstraintDataInspector::GenerateDetailsViewWidget(FMargin Margin)
+{
+	return  SNew(SScrollBox)
+	.Visibility_Raw(this, &SChaosVDConstraintDataInspector::GetDetailsSectionVisibility)
+	+SScrollBox::Slot()
+	.Padding(Margin.Left, Margin.Top, Margin.Right,Margin.Bottom)
+	[
+		MainDataDetailsView->GetWidget().ToSharedRef()
+	];
+}
+
+TSharedRef<SWidget> SChaosVDConstraintDataInspector::GenerateHeaderWidget(FMargin Margin)
 {
-	const TSharedPtr<FChaosVDJointConstraint> JointConstraintPtr = CurrentDataSelectionHandle.GetData().Pin();
-	if (!JointConstraintPtr)
+	return SNew(SBox)
+			.Visibility(EVisibility::Collapsed);
+}
+
+FText SChaosVDConstraintDataInspector::GetParticleName(EChaosVDParticlePairIndex ParticleSlot) const
+{
+	return GetParticleName(ParticleSlot, GetCurrentDataBeingInspected());
+}
+
+FText SChaosVDConstraintDataInspector::GetParticleName(const EChaosVDParticlePairIndex ParticleSlot, const TSharedPtr<FChaosVDSolverDataSelectionHandle>& InSelectionHandle) const
+{
+	bool bHasValidSelection = InSelectionHandle && InSelectionHandle->IsValid();
+	if (!bHasValidSelection)
 	{
-		return FText::AsCultureInvariant("UnnamedParticle");
+		return FText::GetEmpty();
 	}
 
-	if (const TSharedPtr<FChaosVDScene> ScenePtr = SceneWeakPtr.Pin())
+	if (FChaosVDConstraintDataWrapperBase* ConstraintData = GetCurrentDataBeingInspected()->GetData<FChaosVDConstraintDataWrapperBase>())
 	{
-		if (AChaosVDParticleActor* ParticleActor = ScenePtr->GetParticleActor(JointConstraintPtr->SolverID, PairIndex == EChaosVDParticleSelector::Index_0 ? JointConstraintPtr->ParticleParIndexes[0] : JointConstraintPtr->ParticleParIndexes[1]))
-		{
-			if (const FChaosVDParticleDataWrapper* ParticleData = ParticleActor->GetParticleData())
-			{
-				return FText::AsCultureInvariant(ParticleData->DebugName);
-			}
-		}
+		return GetParticleName_Internal(ConstraintData->GetSolverID(), ConstraintData->GetParticleIDAtSlot(ParticleSlot));	
 	}
 
-	return FText::AsCultureInvariant("UnnamedParticle");
+	return FText::GetEmpty();
 }
 
 TSharedRef<SWidget> SChaosVDConstraintDataInspector::GenerateParticleSelectorButtons()
@@ -179,8 +229,8 @@ TSharedRef<SWidget> SChaosVDConstraintDataInspector::GenerateParticleSelectorBut
 				.HAlign(HAlign_Center)
 				.ContentPadding(FAppStyle::GetMargin("StandardDialog.ContentPadding"))
 				.Text(LOCTEXT("SelectParticle0", "Select Particle 0"))
-				.ToolTipText_Raw(this, &SChaosVDConstraintDataInspector::GetParticleName, EChaosVDParticleSelector::Index_0)
-				.OnClicked(this, &SChaosVDConstraintDataInspector::SelectParticleForCurrentSelectedData, EChaosVDParticleSelector::Index_0)
+				.ToolTipText_Raw(this, &SChaosVDConstraintDataInspector::GetParticleName, EChaosVDParticlePairIndex::Index_0)
+				.OnClicked(this, &SChaosVDConstraintDataInspector::SelectParticleForCurrentSelectedData, EChaosVDParticlePairIndex::Index_0)
 			]
 			+SUniformGridPanel::Slot(1, 0)
 			[
@@ -188,43 +238,39 @@ TSharedRef<SWidget> SChaosVDConstraintDataInspector::GenerateParticleSelectorBut
 				.HAlign(HAlign_Center)
 				.ContentPadding(FAppStyle::GetMargin("StandardDialog.ContentPadding"))
 				.Text(LOCTEXT("SelectParticle1", "Select Particle 1"))
-				.ToolTipText_Raw(this, &SChaosVDConstraintDataInspector::GetParticleName, EChaosVDParticleSelector::Index_1)
-				.OnClicked(this, &SChaosVDConstraintDataInspector::SelectParticleForCurrentSelectedData, EChaosVDParticleSelector::Index_1)
+				.ToolTipText_Raw(this, &SChaosVDConstraintDataInspector::GetParticleName, EChaosVDParticlePairIndex::Index_1)
+				.OnClicked(this, &SChaosVDConstraintDataInspector::SelectParticleForCurrentSelectedData, EChaosVDParticlePairIndex::Index_1)
 			];
 }
 
-FReply SChaosVDConstraintDataInspector::SelectParticleForCurrentSelectedData(EChaosVDParticleSelector ParticlePairIndex)
+FReply SChaosVDConstraintDataInspector::SelectParticleForCurrentSelectedData(EChaosVDParticlePairIndex ParticleSlot)
 {
-	const TSharedPtr<FChaosVDJointConstraint> JointConstraintPtr = CurrentDataSelectionHandle.GetData().Pin();
-	if (!JointConstraintPtr)
+	if (FChaosVDConstraintDataWrapperBase* ConstraintData = GetConstraintDataFromSelectionHandle())
 	{
-		return FReply::Handled();
-	}
-
-	if (const TSharedPtr<FChaosVDScene> ScenePtr = SceneWeakPtr.Pin())
-	{
-		if (AChaosVDParticleActor* ParticleActor = ScenePtr->GetParticleActor(JointConstraintPtr->SolverID, ParticlePairIndex == EChaosVDParticleSelector::Index_0 ? JointConstraintPtr->ParticleParIndexes[0] : JointConstraintPtr->ParticleParIndexes[1]))
-		{
-			ScenePtr->SetSelectedObject(ParticleActor);
-		}
+		SelectParticle(ConstraintData->GetSolverID(), ConstraintData->GetParticleIDAtSlot(ParticleSlot));
 	}
 	
 	return FReply::Handled();
 }
 
+bool SChaosVDConstraintDataInspector::HasCompatibleStructScopeView(const TSharedRef<FChaosVDSolverDataSelectionHandle>& InSelectionHandle) const
+{
+	return DataBeingInspectedCopy && InSelectionHandle->GetDataAsStructScope() && DataBeingInspectedCopy->GetStruct() == InSelectionHandle->GetDataAsStructScope()->GetStruct();
+}
+
 EVisibility SChaosVDConstraintDataInspector::GetOutOfDateWarningVisibility() const
 {
-	return !bIsUpToDate && CurrentDataSelectionHandle.GetData().IsValid() ? EVisibility::Visible : EVisibility::Collapsed;
+	return !bIsUpToDate && GetCurrentDataBeingInspected()->IsValid() ? EVisibility::Visible : EVisibility::Collapsed;
 }
 
 EVisibility SChaosVDConstraintDataInspector::GetDetailsSectionVisibility() const
 {
-	return CurrentDataSelectionHandle.GetData().IsValid() ? EVisibility::Visible : EVisibility::Collapsed;
+	return GetCurrentDataBeingInspected()->IsValid() ? EVisibility::Visible : EVisibility::Collapsed;
 }
 
 EVisibility SChaosVDConstraintDataInspector::GetNothingSelectedMessageVisibility() const
 {
-	return !CurrentDataSelectionHandle.GetData().IsValid() ? EVisibility::Visible : EVisibility::Collapsed;
+	return GetCurrentDataBeingInspected()->IsValid() ? EVisibility::Collapsed : EVisibility::Visible;
 }
 
 void SChaosVDConstraintDataInspector::HandleSceneUpdated()
@@ -232,10 +278,9 @@ void SChaosVDConstraintDataInspector::HandleSceneUpdated()
 	// TODO: Disabling the "out of date" message system, it is a little bit confusing with joint constraints as the debug draw position will be updated between frames (because the constraint only has particles IDs)
 	// but the data itself is not updated.
 	// We are clearing out the selection altogether instead for now
-	//bIsUpToDate = false;
-	
-	CurrentDataSelectionHandle.SetIsSelected(false);
-	CurrentDataSelectionHandle = FChaosVDJointConstraintSelectionHandle();
+	bIsUpToDate = false;
+
+	//ClearInspector();
 
 	// TODO: To Keep a selection up to date we need a persistent ID for the constraint
 	// We could hash the pointer for that or add an ID to the Constraint Handle only compiled in when CVD is enabled
@@ -243,8 +288,52 @@ void SChaosVDConstraintDataInspector::HandleSceneUpdated()
 
 void SChaosVDConstraintDataInspector::ClearInspector()
 {
-	ConstraintDataDetailsView->SetStructureData(nullptr);
-	CurrentDataSelectionHandle = FChaosVDJointConstraintSelectionHandle();
+	DataBeingInspectedCopy = nullptr;
+	MainDataDetailsView->SetStructureData(nullptr);
+	CurrentDataSelectionHandle = MakeShared<FChaosVDSolverDataSelectionHandle>();
+}
+
+FText SChaosVDConstraintDataInspector::GetParticleName_Internal(int32 SolverID, int32 ParticleID) const
+{
+	if (const TSharedPtr<FChaosVDScene> ScenePtr = SceneWeakPtr.Pin())
+	{
+		if (AChaosVDParticleActor* ParticleActor = ScenePtr->GetParticleActor(SolverID, ParticleID))
+		{
+			if (TSharedPtr<const FChaosVDParticleDataWrapper> ParticleData = ParticleActor->GetParticleData())
+			{
+				return FText::AsCultureInvariant(ParticleData->DebugName);
+			}
+		}
+	}
+
+	return FText::GetEmpty();
+}
+
+void SChaosVDConstraintDataInspector::SelectParticle(int32 SolverID, int32 ParticleID) const
+{
+	if (const TSharedPtr<FChaosVDScene> ScenePtr = SceneWeakPtr.Pin())
+	{
+		if (AChaosVDParticleActor* ParticleActor = ScenePtr->GetParticleActor(SolverID, ParticleID))
+		{
+			ScenePtr->SetSelectedObject(ParticleActor);
+		}
+	}
+}
+
+FChaosVDConstraintDataWrapperBase* SChaosVDConstraintDataInspector::GetConstraintDataFromSelectionHandle() const
+{
+	bool bHasValidSelection = GetCurrentDataBeingInspected()->IsValid();
+	if (!bHasValidSelection)
+	{
+		return nullptr;
+	}
+
+	return GetCurrentDataBeingInspected()->GetData<FChaosVDConstraintDataWrapperBase>();
+}
+
+const TSharedRef<FChaosVDSolverDataSelectionHandle>& SChaosVDConstraintDataInspector::GetCurrentDataBeingInspected() const
+{
+	return CurrentDataSelectionHandle;
 }
 
 TSharedPtr<IStructureDetailsView> SChaosVDConstraintDataInspector::CreateDataDetailsView()
diff --git a/Engine/Plugins/ChaosVD/Source/ChaosVD/Private/Widgets/SChaosVDConstraintDataInspector.h b/Engine/Plugins/ChaosVD/Source/ChaosVD/Private/Widgets/SChaosVDConstraintDataInspector.h
--- a/Engine/Plugins/ChaosVD/Source/ChaosVD/Private/Widgets/SChaosVDConstraintDataInspector.h
+++ b/Engine/Plugins/ChaosVD/Source/ChaosVD/Private/Widgets/SChaosVDConstraintDataInspector.h
@@ -5,23 +5,33 @@
 #include "Components/ChaosVDSolverJointConstraintDataComponent.h"
 #include "Widgets/SCompoundWidget.h"
 
+struct FChaosVDSolverDataSelectionHandle;
 struct FChaosVDJointConstraintSelectionHandle;
 class FChaosVDScene;
 class IStructureDetailsView;
 class SChaosVDNameListPicker;
 
-enum class EChaosVDParticleSelector : int32
+struct FChaosVDConstraintDataWrapperBase;
+
+/** Version of FStructOnScope that will take another FStructOnScope, and copy its data over.
+ * This allows us update a details panel without making a full rebuild when we want to inspect another struct that is of the same type
+ * As long we don't mind the copy and not being able to edit the source struct (which is 99% of the use cases in CVD)
+ */
+class FReadOnlyCopyStructOnScope : public FStructOnScope
 {
-	Index_0,
-	Index_1
+public:
+	explicit FReadOnlyCopyStructOnScope(const FStructOnScope& StructToCopy);
+
+	void UpdateFromOther(const FStructOnScope& StructToCopy);
 };
 
 class SChaosVDConstraintDataInspector : public SCompoundWidget
 {
 public:
+	SChaosVDConstraintDataInspector();
 
 	SLATE_BEGIN_ARGS(SChaosVDConstraintDataInspector)
-	{
+		{
 	}
 
 	SLATE_END_ARGS()
@@ -30,38 +40,55 @@ public:
 
 	/** Constructs this widget with InArgs */
 	void Construct(const FArguments& InArgs, const TWeakPtr<FChaosVDScene>& InScenePtr);
+
 	/** Sets a new query data to be inspected */
-	void SetConstraintDataToInspect(const FChaosVDJointConstraintSelectionHandle& InDataSelectionHandle);
+	virtual void SetConstraintDataToInspect(const TSharedPtr<FChaosVDSolverDataSelectionHandle>& InDataSelectionHandle);
 
 protected:
 
-	FText GetParticleName(EChaosVDParticleSelector PairIndex) const;
+	virtual void SetupWidgets();
+
+	virtual TSharedRef<SWidget> GenerateDetailsViewWidget(FMargin Margin);
+	virtual TSharedRef<SWidget> GenerateHeaderWidget(FMargin Margin);
+
+	virtual FText GetParticleName(EChaosVDParticlePairIndex ParticleSlot) const;
+	virtual FText GetParticleName(const EChaosVDParticlePairIndex ParticleSlot, const TSharedPtr<FChaosVDSolverDataSelectionHandle>& InSelectionHandle) const;
 
 	TSharedPtr<IStructureDetailsView> CreateDataDetailsView();
 
 	TSharedRef<SWidget> GenerateParticleSelectorButtons();
 	
-	FReply SelectParticleForCurrentSelectedData(EChaosVDParticleSelector ParticlePairIndex);
+	virtual FReply SelectParticleForCurrentSelectedData(EChaosVDParticlePairIndex ParticleSlot);
+
+	bool HasCompatibleStructScopeView(const TSharedRef<FChaosVDSolverDataSelectionHandle>& InSelectionHandle) const;
 
 	EVisibility GetOutOfDateWarningVisibility() const;
-	EVisibility GetDetailsSectionVisibility() const;
+	virtual EVisibility GetDetailsSectionVisibility() const;
 	EVisibility GetNothingSelectedMessageVisibility() const;
 
-	void RegisterSelectionEventsForSolver(AChaosVDSolverInfoActor* SolverInfo);
-	void UnregisterSelectionEventsForSolver(AChaosVDSolverInfoActor* SolverInfo);
-
 	void RegisterSceneEvents();
-	void UnregisterSceneEvents();
+	void UnregisterSceneEvents() const;
+
+	virtual void HandleSceneUpdated();
+
+	virtual void ClearInspector();
 
-	void HandleSceneUpdated();
+	FText GetParticleName_Internal(int32 SolverID, int32 ParticleID) const;
 
-	void ClearInspector();
+	void SelectParticle(int32 SolverID, int32 ParticleID) const;
 
-	TSharedPtr<IStructureDetailsView> ConstraintDataDetailsView;
+	FChaosVDConstraintDataWrapperBase* GetConstraintDataFromSelectionHandle() const;
+	
+	virtual const TSharedRef<FChaosVDSolverDataSelectionHandle>& GetCurrentDataBeingInspected() const;
+
+	TSharedPtr<IStructureDetailsView> MainDataDetailsView;
+	TSharedPtr<IStructureDetailsView> ConstraintSecondaryDataDetailsView;
 	
 	TWeakPtr<FChaosVDScene> SceneWeakPtr;
 
-	FChaosVDJointConstraintSelectionHandle CurrentDataSelectionHandle;
+	mutable TSharedRef<FChaosVDSolverDataSelectionHandle> CurrentDataSelectionHandle;
+
+	TSharedPtr<FReadOnlyCopyStructOnScope> DataBeingInspectedCopy;
 
-	bool bIsUpToDate = true;
+	bool bIsUpToDate = true;;
 };
diff --git a/Engine/Plugins/ChaosVD/Source/ChaosVD/Private/Widgets/SChaosVDMainTab.cpp b/Engine/Plugins/ChaosVD/Source/ChaosVD/Private/Widgets/SChaosVDMainTab.cpp
--- a/Engine/Plugins/ChaosVD/Source/ChaosVD/Private/Widgets/SChaosVDMainTab.cpp
+++ b/Engine/Plugins/ChaosVD/Source/ChaosVD/Private/Widgets/SChaosVDMainTab.cpp
@@ -11,7 +11,6 @@
 #include "ChaosVDPlaybackViewportTab.h"
 #include "ChaosVDScene.h"
 #include "ChaosVDSolversTracksTab.h"
-#include "ChaosVDCharacterGroundConstraintDataInspectorTab.h"
 #include "ChaosVDCollisionDataDetailsTab.h"
 #include "ChaosVDConstraintDataInspectorTab.h"
 #include "ChaosVDSceneQueryDataInspectorTab.h"
@@ -75,8 +74,7 @@ void SChaosVDMainTab::Construct(const FArguments& InArgs, TSharedPtr<FChaosVDEng
 	RegisterTabSpawner<FChaosVDSolversTracksTab>(FChaosVDTabID::SolversTrack);
 	RegisterTabSpawner<FChaosVDCollisionDataDetailsTab>(FChaosVDTabID::CollisionDataDetails);
 	RegisterTabSpawner<FChaosVDSceneQueryDataInspectorTab>(FChaosVDTabID::SceneQueryDataDetails);
-	RegisterTabSpawner<FChaosVDConstraintDataInspectorTab>(FChaosVDTabID::JointsDataDetails);
-	RegisterTabSpawner<FChaosVDCharacterGroundConstraintDataInspectorTab>(FChaosVDTabID::CharacterGroundConstraintDataDetails);
+	RegisterTabSpawner<FChaosVDConstraintDataInspectorTab>(FChaosVDTabID::ConstraintsInspector);
 
 	StatusBarID = FName(FChaosVDTabID::StatusBar.ToString() + InChaosVDEngine->GetInstanceGuid().ToString());
 	
@@ -284,8 +282,7 @@ TSharedRef<FTabManager::FLayout> SChaosVDMainTab::GenerateMainLayout()
 					->AddTab(FChaosVDTabID::DetailsPanel, ETabState::OpenedTab)
 					->AddTab(FChaosVDTabID::CollisionDataDetails, ETabState::OpenedTab)
 					->AddTab(FChaosVDTabID::SceneQueryDataDetails, ETabState::OpenedTab)
-					->AddTab(FChaosVDTabID::JointsDataDetails, ETabState::ClosedTab)
-					->AddTab(FChaosVDTabID::CharacterGroundConstraintDataDetails, ETabState::ClosedTab)
+					->AddTab(FChaosVDTabID::ConstraintsInspector, ETabState::ClosedTab)
 				)
 			)
 		);
diff --git a/Engine/Plugins/ChaosVD/Source/ChaosVD/Private/Widgets/SChaosVDSceneQueryDataInspector.cpp b/Engine/Plugins/ChaosVD/Source/ChaosVD/Private/Widgets/SChaosVDSceneQueryDataInspector.cpp
--- a/Engine/Plugins/ChaosVD/Source/ChaosVD/Private/Widgets/SChaosVDSceneQueryDataInspector.cpp
+++ b/Engine/Plugins/ChaosVD/Source/ChaosVD/Private/Widgets/SChaosVDSceneQueryDataInspector.cpp
@@ -26,43 +26,51 @@
 
 #define LOCTEXT_NAMESPACE "ChaosVisualDebugger"
 
-SChaosVDSceneQueryDataInspector::SChaosVDSceneQueryDataInspector() : CurrentSceneQueryBeingInspectedHandle({nullptr, INDEX_NONE})
+SChaosVDSceneQueryDataInspector::SChaosVDSceneQueryDataInspector() : CurrentSceneQueryBeingInspectedHandle(MakeShared<FChaosVDSolverDataSelectionHandle>())
 {
 	
 }
 
-SChaosVDSceneQueryDataInspector::~SChaosVDSceneQueryDataInspector()
+void SChaosVDSceneQueryDataInspector::RegisterSceneEvents()
 {
 	if (const TSharedPtr<FChaosVDScene> ScenePtr = SceneWeakPtr.Pin())
 	{
-		ScenePtr->OnSceneUpdated().RemoveAll(this);
-		
-		if (UChaosVDSceneQueryDataComponent* SQDataComponent = ScenePtr->GetSceneQueryDataContainerComponent())
+		ScenePtr->OnSceneUpdated().AddRaw(this, &SChaosVDSceneQueryDataInspector::HandleSceneUpdated);
+		if (TSharedPtr<FChaosVDSolverDataSelection> SelectionObject = ScenePtr->GetSolverDataSelectionObject().Pin())
 		{
-			SQDataComponent->GetOnSelectionChangeDelegate().RemoveAll(this);
+			SelectionObject->GetDataSelectionChangedDelegate().AddRaw(this, &SChaosVDSceneQueryDataInspector::SetQueryDataToInspect);
 		}
 	}
 }
 
-void SChaosVDSceneQueryDataInspector::Construct(const FArguments& InArgs, const TWeakPtr<FChaosVDScene>& InScenePtr, const TWeakPtr<FEditorModeTools>& InEditorModeTools)
+void SChaosVDSceneQueryDataInspector::UnregisterSceneEvents()
 {
-	SceneWeakPtr = InScenePtr;
-	EditorModeToolsWeakPtr = InEditorModeTools;
-
 	if (const TSharedPtr<FChaosVDScene> ScenePtr = SceneWeakPtr.Pin())
 	{
-		ScenePtr->OnSceneUpdated().AddRaw(this, &SChaosVDSceneQueryDataInspector::HandleSceneUpdated);
+		ScenePtr->OnSceneUpdated().RemoveAll(this);
 
-		if (UChaosVDSceneQueryDataComponent* SQDataComponent = ScenePtr->GetSceneQueryDataContainerComponent())
+		if (TSharedPtr<FChaosVDSolverDataSelection> SelectionObject = ScenePtr->GetSolverDataSelectionObject().Pin())
 		{
-			SQDataComponent->GetOnSelectionChangeDelegate().AddRaw(this, &SChaosVDSceneQueryDataInspector::SetQueryDataToInspect);
+			SelectionObject->GetDataSelectionChangedDelegate().RemoveAll(this);
 		}
 	}
+}
+
+SChaosVDSceneQueryDataInspector::~SChaosVDSceneQueryDataInspector()
+{
+	UnregisterSceneEvents();
+}
+
+void SChaosVDSceneQueryDataInspector::Construct(const FArguments& InArgs, const TWeakPtr<FChaosVDScene>& InScenePtr, const TWeakPtr<FEditorModeTools>& InEditorModeTools)
+{
+	SceneWeakPtr = InScenePtr;
+	EditorModeToolsWeakPtr = InEditorModeTools;
+
+	RegisterSceneEvents();
 
 	SceneQueryDataDetailsView = CreateDataDetailsView();
 	SceneQueryHitDataDetailsView = CreateDataDetailsView();
 
-	
 	constexpr float NoPadding = 0.0f;
 	constexpr float OuterBoxPadding = 2.0f;
 	constexpr float OuterInnerPadding = 5.0f;
@@ -280,22 +288,30 @@ TSharedRef<SWidget> SChaosVDSceneQueryDataInspector::GenerateQueryTagInfoRow()
 			];
 }
 
-void SChaosVDSceneQueryDataInspector::SetQueryDataToInspect(const FChaosVDSceneQuerySelectionHandle& InQueryDataSelectionHandle)
+void SChaosVDSceneQueryDataInspector::SetQueryDataToInspect(const TSharedPtr<FChaosVDSolverDataSelectionHandle>& InDataSelectionHandle)
 {
 	ClearInspector();
 
-	if (const TSharedPtr<FChaosVDQueryDataWrapper> QueryDataToInspect = InQueryDataSelectionHandle.GetQueryData().Pin())
+	if (!InDataSelectionHandle)
+	{
+		return;
+	}
+
+	if (const TSharedPtr<FChaosVDQueryDataWrapper> QueryDataToInspect = InDataSelectionHandle->GetDataAsShared<FChaosVDQueryDataWrapper>())
 	{
-		CurrentSceneQueryBeingInspectedHandle = InQueryDataSelectionHandle;
+		CurrentSceneQueryBeingInspectedHandle = InDataSelectionHandle.ToSharedRef();
 
 		const TSharedPtr<FStructOnScope> QueryDataView = MakeShared<FStructOnScope>(FChaosVDQueryDataWrapper::StaticStruct(), reinterpret_cast<uint8*>(QueryDataToInspect.Get()));
 		SceneQueryDataDetailsView->SetStructureData(QueryDataView);
 
-		if (QueryDataToInspect->SQVisitData.IsValidIndex(InQueryDataSelectionHandle.GetSQVisitIndex()))
+		if(FChaosVDSceneQuerySelectionContext* SelectionContext = InDataSelectionHandle->GetContextData<FChaosVDSceneQuerySelectionContext>())
 		{
-			const TSharedPtr<FStructOnScope> SQVisitDataDataView = MakeShared<FStructOnScope>(FChaosVDQueryVisitStep::StaticStruct(), reinterpret_cast<uint8*>(&QueryDataToInspect->SQVisitData[InQueryDataSelectionHandle.GetSQVisitIndex()]));
-			SceneQueryHitDataDetailsView->SetStructureData(SQVisitDataDataView);
-			QueryDataToInspect->CurrentVisitIndex = InQueryDataSelectionHandle.GetSQVisitIndex();
+			if (QueryDataToInspect->SQVisitData.IsValidIndex(SelectionContext->SQVisitIndex))
+			{
+				const TSharedPtr<FStructOnScope> SQVisitDataDataView = MakeShared<FStructOnScope>(FChaosVDQueryVisitStep::StaticStruct(), reinterpret_cast<uint8*>(&QueryDataToInspect->SQVisitData[SelectionContext->SQVisitIndex]));
+				SceneQueryHitDataDetailsView->SetStructureData(SQVisitDataDataView);
+				QueryDataToInspect->CurrentVisitIndex = SelectionContext->SQVisitIndex;
+			}
 		}
 
 		if (QueryDataToInspect->SubQueriesIDs.Num() > 0)
@@ -332,7 +348,7 @@ void SChaosVDSceneQueryDataInspector::HandleQueryStepSelectionUpdated(int32 NewS
 		return;
 	}
 
-	const TSharedPtr<FChaosVDQueryDataWrapper> QueryDataBeingInspected = CurrentSceneQueryBeingInspectedHandle.GetQueryData().Pin();
+	const TSharedPtr<FChaosVDQueryDataWrapper> QueryDataBeingInspected = CurrentSceneQueryBeingInspectedHandle->GetDataAsShared<FChaosVDQueryDataWrapper>();
 	if (!QueryDataBeingInspected)
 	{
 		ClearInspector();
@@ -346,6 +362,14 @@ void SChaosVDSceneQueryDataInspector::HandleQueryStepSelectionUpdated(int32 NewS
 		return;
 	}
 
+	TSharedPtr<FChaosVDSolverDataSelection> SelectionObject = ScenePtr->GetSolverDataSelectionObject().Pin();
+
+	if (!SelectionObject)
+	{
+		ClearInspector();
+		return;
+	}
+
 	UChaosVDSceneQueryDataComponent* SQDataComponent = ScenePtr->GetSceneQueryDataContainerComponent();
 	if (!SQDataComponent)
 	{
@@ -360,13 +384,15 @@ void SChaosVDSceneQueryDataInspector::HandleQueryStepSelectionUpdated(int32 NewS
 		return;
 	}
 
-	const FChaosVDSceneQuerySelectionHandle NewSelection(CurrentSceneQueryBeingInspectedHandle.GetQueryData(), NewStepIndex);
+	TSharedPtr<FChaosVDSolverDataSelectionHandle> NewSelection = SelectionObject->MakeSelectionHandle(QueryDataBeingInspected);
+	FChaosVDSceneQuerySelectionContext ContextData;
+	ContextData.SQVisitIndex = NewStepIndex;
+	NewSelection->SetHandleContext(MoveTemp(ContextData));
+
 	QueryDataBeingInspected->CurrentVisitIndex = NewStepIndex;
 
 	FScopedSQInspectorSilencedSelectionEvents IgnoreSelectionEventsScope(*this);
-	SQDataComponent->SelectQuery(NewSelection);
-
-	SetQueryDataToInspect(NewSelection);
+	SelectionObject->SelectData(NewSelection);
 
 	if (const TSharedPtr<FEditorModeTools> EditorModeToolsPtr = EditorModeToolsWeakPtr.Pin())
 	{
@@ -380,7 +406,7 @@ void SChaosVDSceneQueryDataInspector::HandleQueryStepSelectionUpdated(int32 NewS
 
 FText SChaosVDSceneQueryDataInspector::GetQueryBeingInspectedTag() const
 {
-	const TSharedPtr<FChaosVDQueryDataWrapper> QueryDataBeingInspected = CurrentSceneQueryBeingInspectedHandle.GetQueryData().Pin();
+	const TSharedPtr<FChaosVDQueryDataWrapper> QueryDataBeingInspected = CurrentSceneQueryBeingInspectedHandle->GetDataAsShared<FChaosVDQueryDataWrapper>();
 	return FText::AsCultureInvariant(QueryDataBeingInspected ? QueryDataBeingInspected->CollisionQueryParams.TraceTag.ToString() : TEXT("None"));
 }
 
@@ -391,7 +417,7 @@ FText SChaosVDSceneQueryDataInspector::GetSelectParticleText() const
 
 FText SChaosVDSceneQueryDataInspector::GetSQVisitsStepsText() const
 {
-	const TSharedPtr<FChaosVDQueryDataWrapper> QueryDataBeingInspected = CurrentSceneQueryBeingInspectedHandle.GetQueryData().Pin();
+	const TSharedPtr<FChaosVDQueryDataWrapper> QueryDataBeingInspected = CurrentSceneQueryBeingInspectedHandle->GetDataAsShared<FChaosVDQueryDataWrapper>();
 	const int32 AvailableSQVisitDataNum = QueryDataBeingInspected ? QueryDataBeingInspected->SQVisitData.Num() : 0;
 	return FText::Format(FTextFormat(LOCTEXT("SQVisitStepsPlaybackControlsTitle", "Recorded SQ Visits Available {0}")), AvailableSQVisitDataNum);
 }
@@ -404,7 +430,7 @@ FReply SChaosVDSceneQueryDataInspector::SelectParticleForCurrentQueryData() cons
 		return FReply::Handled();
 	}
 
-	const TSharedPtr<FChaosVDQueryDataWrapper> QueryDataBeingInspected = CurrentSceneQueryBeingInspectedHandle.GetQueryData().Pin();
+	const TSharedPtr<FChaosVDQueryDataWrapper> QueryDataBeingInspected = CurrentSceneQueryBeingInspectedHandle->GetDataAsShared<FChaosVDQueryDataWrapper>();
 
 	if (!QueryDataBeingInspected || !QueryDataBeingInspected->SQVisitData.IsValidIndex(QueryDataBeingInspected->CurrentVisitIndex))
 	{
@@ -448,12 +474,16 @@ FReply SChaosVDSceneQueryDataInspector::SelectQueryToInspectByID(int32 QueryID)
 	{
 		return FReply::Handled();
 	}
+	
+	TSharedPtr<FChaosVDSolverDataSelection> SelectionObject = ScenePtr->GetSolverDataSelectionObject().Pin();
+	if (!SelectionObject)
+	{
+		return FReply::Handled();
+	}
 
 	if (UChaosVDSceneQueryDataComponent* SQDataComponent = ScenePtr->GetSceneQueryDataContainerComponent())
-	{
-		SQDataComponent->SelectQuery(QueryID);
-							
-		SetQueryDataToInspect(SQDataComponent->GetSelectedQueryHandle());
+	{						
+		SelectionObject->SelectData(SelectionObject->MakeSelectionHandle(SQDataComponent->GetQueryByID(QueryID)));
 	}
 	else
 	{
@@ -467,7 +497,7 @@ FReply SChaosVDSceneQueryDataInspector::SelectParentQuery()
 {
 	const TSharedPtr<FChaosVDScene> ScenePtr = SceneWeakPtr.Pin();
 	const UChaosVDSceneQueryDataComponent* SQDataComponent = ScenePtr ? ScenePtr->GetSceneQueryDataContainerComponent() : nullptr;
-	if (const TSharedPtr<FChaosVDQueryDataWrapper> SelectedQuery = SQDataComponent ? SQDataComponent->GetSelectedQueryHandle().GetQueryData().Pin() : nullptr)
+	if (const TSharedPtr<FChaosVDQueryDataWrapper> SelectedQuery = SQDataComponent ? CurrentSceneQueryBeingInspectedHandle->GetDataAsShared<FChaosVDQueryDataWrapper>() : nullptr)
 	{
 		SelectQueryToInspectByID(SelectedQuery->ParentQueryID);
 	}
@@ -543,7 +573,7 @@ void SChaosVDSceneQueryDataInspector::ClearInspector()
 
 	CurrentSubQueriesByName.Empty();
 
-	CurrentSceneQueryBeingInspectedHandle = FChaosVDSceneQuerySelectionHandle(nullptr, INDEX_NONE);
+	CurrentSceneQueryBeingInspectedHandle = MakeShared<FChaosVDSolverDataSelectionHandle>();
 
 	bIsUpToDate = true;
 }
@@ -555,44 +585,44 @@ EVisibility SChaosVDSceneQueryDataInspector::GetOutOfDateWarningVisibility() con
 
 EVisibility SChaosVDSceneQueryDataInspector::GetQueryDetailsSectionVisibility() const
 {
-	const TSharedPtr<FChaosVDQueryDataWrapper> QueryDataBeingInspected = CurrentSceneQueryBeingInspectedHandle.GetQueryData().Pin();
+	const TSharedPtr<FChaosVDQueryDataWrapper> QueryDataBeingInspected = CurrentSceneQueryBeingInspectedHandle->GetDataAsShared<FChaosVDQueryDataWrapper>();
 	return QueryDataBeingInspected ? EVisibility::Visible : EVisibility::Collapsed;
 }
 
 EVisibility SChaosVDSceneQueryDataInspector::GetQueryStepPlaybackControlsVisibility() const
 {
-	const TSharedPtr<FChaosVDQueryDataWrapper> QueryDataBeingInspected = CurrentSceneQueryBeingInspectedHandle.GetQueryData().Pin();
+	const TSharedPtr<FChaosVDQueryDataWrapper> QueryDataBeingInspected = CurrentSceneQueryBeingInspectedHandle->GetDataAsShared<FChaosVDQueryDataWrapper>();
 	// If the this inspector no longer reflects data represented in the viewport, we can't offer playback so we need to hide the controls
 	return bIsUpToDate && QueryDataBeingInspected ? EVisibility::Visible : EVisibility::Collapsed;
 }
 
 EVisibility SChaosVDSceneQueryDataInspector::GetSQVisitDetailsSectionVisibility() const
 {
-	const TSharedPtr<FChaosVDQueryDataWrapper> QueryDataBeingInspected = CurrentSceneQueryBeingInspectedHandle.GetQueryData().Pin();
+	const TSharedPtr<FChaosVDQueryDataWrapper> QueryDataBeingInspected = CurrentSceneQueryBeingInspectedHandle->GetDataAsShared<FChaosVDQueryDataWrapper>();
 	return QueryDataBeingInspected && QueryDataBeingInspected->SQVisitData.Num() > 0 ? EVisibility::Visible : EVisibility::Collapsed;
 }
 
 EVisibility SChaosVDSceneQueryDataInspector::GetNothingSelectedMessageVisibility() const
 {
-	const TSharedPtr<FChaosVDQueryDataWrapper> QueryDataBeingInspected = CurrentSceneQueryBeingInspectedHandle.GetQueryData().Pin();
+	const TSharedPtr<FChaosVDQueryDataWrapper> QueryDataBeingInspected = CurrentSceneQueryBeingInspectedHandle->GetDataAsShared<FChaosVDQueryDataWrapper>();
 	return QueryDataBeingInspected ? EVisibility::Collapsed : EVisibility::Visible;
 }
 
 EVisibility SChaosVDSceneQueryDataInspector::GetSubQuerySelectorVisibility() const
 {
-	const TSharedPtr<FChaosVDQueryDataWrapper> QueryDataBeingInspected = CurrentSceneQueryBeingInspectedHandle.GetQueryData().Pin();
+	const TSharedPtr<FChaosVDQueryDataWrapper> QueryDataBeingInspected = CurrentSceneQueryBeingInspectedHandle->GetDataAsShared<FChaosVDQueryDataWrapper>();
 	return QueryDataBeingInspected && QueryDataBeingInspected->SubQueriesIDs.Num() > 0 ? EVisibility::Visible : EVisibility::Collapsed;
 }
 
 EVisibility SChaosVDSceneQueryDataInspector::GetParentQuerySelectorVisibility() const
 {
-	const TSharedPtr<FChaosVDQueryDataWrapper> QueryDataBeingInspected = CurrentSceneQueryBeingInspectedHandle.GetQueryData().Pin();
+	const TSharedPtr<FChaosVDQueryDataWrapper> QueryDataBeingInspected = CurrentSceneQueryBeingInspectedHandle->GetDataAsShared<FChaosVDQueryDataWrapper>();
 	return QueryDataBeingInspected && QueryDataBeingInspected->ParentQueryID != INDEX_NONE ? EVisibility::Visible : EVisibility::Collapsed;
 }
 
 bool SChaosVDSceneQueryDataInspector::GetSelectParticleHitStateEnable() const
 {
-	const TSharedPtr<FChaosVDQueryDataWrapper> QueryDataBeingInspected = CurrentSceneQueryBeingInspectedHandle.GetQueryData().Pin();
+	const TSharedPtr<FChaosVDQueryDataWrapper> QueryDataBeingInspected = CurrentSceneQueryBeingInspectedHandle->GetDataAsShared<FChaosVDQueryDataWrapper>();
 	if (QueryDataBeingInspected && QueryDataBeingInspected->SQVisitData.IsValidIndex(QueryDataBeingInspected->CurrentVisitIndex))
 	{
 		return QueryDataBeingInspected->SQVisitData[QueryDataBeingInspected->CurrentVisitIndex].ParticleIndex != INDEX_NONE;
@@ -603,13 +633,13 @@ bool SChaosVDSceneQueryDataInspector::GetSelectParticleHitStateEnable() const
 
 bool SChaosVDSceneQueryDataInspector::GetSQVisitStepsEnabled() const
 {
-	const TSharedPtr<FChaosVDQueryDataWrapper> QueryDataBeingInspected = CurrentSceneQueryBeingInspectedHandle.GetQueryData().Pin();
+	const TSharedPtr<FChaosVDQueryDataWrapper> QueryDataBeingInspected = CurrentSceneQueryBeingInspectedHandle->GetDataAsShared<FChaosVDQueryDataWrapper>();
 	return QueryDataBeingInspected && QueryDataBeingInspected->SQVisitData.Num() > 0;
 }
 
 TSharedPtr<FChaosVDQueryDataWrapper> SChaosVDSceneQueryDataInspector::GetCurrentDataBeingInspected() const
 {
-	const TSharedPtr<FChaosVDQueryDataWrapper> QueryDataBeingInspected = CurrentSceneQueryBeingInspectedHandle.GetQueryData().Pin();
+	const TSharedPtr<FChaosVDQueryDataWrapper> QueryDataBeingInspected = CurrentSceneQueryBeingInspectedHandle->GetDataAsShared<FChaosVDQueryDataWrapper>();
 	return QueryDataBeingInspected;
 }
 
diff --git a/Engine/Plugins/ChaosVD/Source/ChaosVD/Private/Widgets/SChaosVDSceneQueryDataInspector.h b/Engine/Plugins/ChaosVD/Source/ChaosVD/Private/Widgets/SChaosVDSceneQueryDataInspector.h
--- a/Engine/Plugins/ChaosVD/Source/ChaosVD/Private/Widgets/SChaosVDSceneQueryDataInspector.h
+++ b/Engine/Plugins/ChaosVD/Source/ChaosVD/Private/Widgets/SChaosVDSceneQueryDataInspector.h
@@ -23,6 +23,8 @@ class SChaosVDSceneQueryDataInspector : public SCompoundWidget
 {
 public:
 	SChaosVDSceneQueryDataInspector();
+	void RegisterSceneEvents();
+	void UnregisterSceneEvents();
 
 	SLATE_BEGIN_ARGS(SChaosVDSceneQueryDataInspector)
 		{
@@ -35,7 +37,7 @@ public:
 	void Construct(const FArguments& InArgs, const TWeakPtr<FChaosVDScene>& InScenePtr, const TWeakPtr<FEditorModeTools>& InEditorModeTools);
 
 	/** Sets a new query data to be inspected */
-	void SetQueryDataToInspect(const FChaosVDSceneQuerySelectionHandle& InQueryDataSelectionHandle);
+	void SetQueryDataToInspect(const TSharedPtr<FChaosVDSolverDataSelectionHandle>& InDataSelectionHandle);
 
 protected:
 
@@ -88,7 +90,7 @@ protected:
 
 	TMap<TSharedPtr<FName>, int32> CurrentSubQueriesByName;
 	
-	FChaosVDSceneQuerySelectionHandle CurrentSceneQueryBeingInspectedHandle;
+	TSharedRef<FChaosVDSolverDataSelectionHandle> CurrentSceneQueryBeingInspectedHandle;
 
 	bool bIsUpToDate = true;
 
diff --git a/Engine/Source/Runtime/Experimental/ChaosVisualDebugger/Private/DataWrappers/ChaosVDCharacterGroundConstraintDataWrappers.cpp b/Engine/Source/Runtime/Experimental/ChaosVisualDebugger/Private/DataWrappers/ChaosVDCharacterGroundConstraintDataWrappers.cpp
--- a/Engine/Source/Runtime/Experimental/ChaosVisualDebugger/Private/DataWrappers/ChaosVDCharacterGroundConstraintDataWrappers.cpp
+++ b/Engine/Source/Runtime/Experimental/ChaosVisualDebugger/Private/DataWrappers/ChaosVDCharacterGroundConstraintDataWrappers.cpp
@@ -84,3 +84,16 @@ bool FChaosVDCharacterGroundConstraint::Serialize(FArchive& Ar)
 
 	return !Ar.IsError();
 }
+
+int32 FChaosVDCharacterGroundConstraint::GetParticleIDAtSlot(EChaosVDParticlePairIndex IndexSlot) const
+{
+	switch (IndexSlot)
+	{
+		case EChaosVDParticlePairIndex::Index_0:
+			return CharacterParticleIndex;
+		case EChaosVDParticlePairIndex::Index_1:
+			return GroundParticleIndex;
+		default:
+			return INDEX_NONE;
+	}
+}
diff --git a/Engine/Source/Runtime/Experimental/ChaosVisualDebugger/Public/DataWrappers/ChaosVDCharacterGroundConstraintDataWrappers.h b/Engine/Source/Runtime/Experimental/ChaosVisualDebugger/Public/DataWrappers/ChaosVDCharacterGroundConstraintDataWrappers.h
--- a/Engine/Source/Runtime/Experimental/ChaosVisualDebugger/Public/DataWrappers/ChaosVDCharacterGroundConstraintDataWrappers.h
+++ b/Engine/Source/Runtime/Experimental/ChaosVisualDebugger/Public/DataWrappers/ChaosVDCharacterGroundConstraintDataWrappers.h
@@ -115,7 +115,7 @@ CVD_IMPLEMENT_SERIALIZER(FChaosVDCharacterGroundConstraintDataDataWrapper)
 
 
 USTRUCT()
-struct FChaosVDCharacterGroundConstraint : public FChaosVDWrapperDataBase
+struct FChaosVDCharacterGroundConstraint : public FChaosVDConstraintDataWrapperBase
 {
 	GENERATED_BODY()
 public:
@@ -141,7 +141,9 @@ public:
 	UPROPERTY(VisibleAnywhere, Category=Data)
 	FChaosVDCharacterGroundConstraintDataDataWrapper Data;
 
-	bool bIsSelectedInEditor = false;
+	CHAOSVDRUNTIME_API virtual int32 GetSolverID() const override { return SolverID; }
+	CHAOSVDRUNTIME_API virtual int32 GetParticleIDAtSlot(EChaosVDParticlePairIndex IndexSlot) const override;
+	CHAOSVDRUNTIME_API virtual int32 GetConstraintIndex () const override {return ConstraintIndex; }
 };
 
 CVD_IMPLEMENT_SERIALIZER(FChaosVDCharacterGroundConstraint)
diff --git a/Engine/Source/Runtime/Experimental/ChaosVisualDebugger/Public/DataWrappers/ChaosVDCollisionDataWrappers.h b/Engine/Source/Runtime/Experimental/ChaosVisualDebugger/Public/DataWrappers/ChaosVDCollisionDataWrappers.h
--- a/Engine/Source/Runtime/Experimental/ChaosVisualDebugger/Public/DataWrappers/ChaosVDCollisionDataWrappers.h
+++ b/Engine/Source/Runtime/Experimental/ChaosVisualDebugger/Public/DataWrappers/ChaosVDCollisionDataWrappers.h
@@ -421,7 +421,7 @@ struct CHAOSVDRUNTIME_API FChaosVDParticlePairMidPhase
 	UPROPERTY(VisibleAnywhere, Category=Particle)
 	int32 Particle1Idx = 0;
 
-	UPROPERTY()
+	UPROPERTY(VisibleAnywhere, Category=Constraints)
 	TArray<FChaosVDConstraint> Constraints;
 
 	bool Serialize(FArchive& Ar);
diff --git a/Engine/Source/Runtime/Experimental/ChaosVisualDebugger/Public/DataWrappers/ChaosVDJointDataWrappers.h b/Engine/Source/Runtime/Experimental/ChaosVisualDebugger/Public/DataWrappers/ChaosVDJointDataWrappers.h
--- a/Engine/Source/Runtime/Experimental/ChaosVisualDebugger/Public/DataWrappers/ChaosVDJointDataWrappers.h
+++ b/Engine/Source/Runtime/Experimental/ChaosVisualDebugger/Public/DataWrappers/ChaosVDJointDataWrappers.h
@@ -378,10 +378,9 @@ public:
 CVD_IMPLEMENT_SERIALIZER(FChaosVDJointSettingsDataWrapper)
 
 USTRUCT()
-struct FChaosVDJointConstraint : public FChaosVDWrapperDataBase
+struct FChaosVDJointConstraint : public FChaosVDConstraintDataWrapperBase
 {
 	GENERATED_BODY()
-public:
 
 	inline static FStringView WrapperTypeName = TEXT("FChaosVDJointConstraint");
 
@@ -399,8 +398,10 @@ public:
 	
 	UPROPERTY(VisibleAnywhere, Category=JointSettings)
 	FChaosVDJointSettingsDataWrapper JointSettings;
-
-	bool bIsSelectedInEditor = false;
+	
+	CHAOSVDRUNTIME_API virtual int32 GetSolverID() const override { return SolverID; }
+	CHAOSVDRUNTIME_API virtual int32 GetParticleIDAtSlot(EChaosVDParticlePairIndex IndexSlot) const override { return ParticleParIndexes[static_cast<int32>(IndexSlot)]; }
+	CHAOSVDRUNTIME_API virtual int32 GetConstraintIndex () const override {return ConstraintIndex; }
 };
 
 CVD_IMPLEMENT_SERIALIZER(FChaosVDJointConstraint)
diff --git a/Engine/Source/Runtime/Experimental/ChaosVisualDebugger/Public/DataWrappers/ChaosVDParticleDataWrapper.h b/Engine/Source/Runtime/Experimental/ChaosVisualDebugger/Public/DataWrappers/ChaosVDParticleDataWrapper.h
--- a/Engine/Source/Runtime/Experimental/ChaosVisualDebugger/Public/DataWrappers/ChaosVDParticleDataWrapper.h
+++ b/Engine/Source/Runtime/Experimental/ChaosVisualDebugger/Public/DataWrappers/ChaosVDParticleDataWrapper.h
@@ -2,9 +2,9 @@
 
 #pragma once
 
+#include "ChaosVDCollisionDataWrappers.h"
 #include "UObject/ObjectMacros.h"
 #include "HAL/Platform.h"
-#include "DataWrappers/ChaosVDCollisionDataWrappers.h"
 
 #include "ChaosVDParticleDataWrapper.generated.h"
 
@@ -56,6 +56,24 @@ protected:
 	bool bHasValidData = false;
 };
 
+enum class EChaosVDParticlePairIndex : uint8
+{
+	Index_0,
+	Index_1
+};
+
+/** Base struct that declares the interface to be used for any Constraint data to be visualized */
+USTRUCT()
+struct FChaosVDConstraintDataWrapperBase : public FChaosVDWrapperDataBase
+{
+	GENERATED_BODY()
+	virtual ~FChaosVDConstraintDataWrapperBase() override = default;
+
+	CHAOSVDRUNTIME_API virtual int32 GetSolverID() const  { return INDEX_NONE; }
+	CHAOSVDRUNTIME_API virtual int32 GetParticleIDAtSlot (EChaosVDParticlePairIndex IndexSlot) const { return INDEX_NONE;}
+	CHAOSVDRUNTIME_API virtual int32 GetConstraintIndex () const {  return INDEX_NONE; }
+};
+
 USTRUCT()
 struct FChaosVDFRigidParticleControlFlags : public FChaosVDWrapperDataBase 
 {
diff --git a/Engine/Source/Runtime/Experimental/ChaosVisualDebugger/Public/DataWrappers/ChaosVDQueryDataWrappers.h b/Engine/Source/Runtime/Experimental/ChaosVisualDebugger/Public/DataWrappers/ChaosVDQueryDataWrappers.h
--- a/Engine/Source/Runtime/Experimental/ChaosVisualDebugger/Public/DataWrappers/ChaosVDQueryDataWrappers.h
+++ b/Engine/Source/Runtime/Experimental/ChaosVisualDebugger/Public/DataWrappers/ChaosVDQueryDataWrappers.h
@@ -362,10 +362,6 @@ struct FChaosVDQueryVisitStep : public FChaosVDWrapperDataBase
 
 	UPROPERTY(VisibleAnywhere, Category="SQ Visit Data")
 	FChaosVDQueryHitData HitData;
-
-	// Editor only properties
-
-	bool bIsSelectedInEditor = false;
 };
 
 inline FArchive& operator<<(FArchive& Ar, FChaosVDQueryVisitStep& Data)
@@ -440,8 +436,6 @@ struct FChaosVDQueryDataWrapper
 
 	TArray<int32> SubQueriesIDs;
 
-	bool bIsSelectedInEditor = false;
-
 	int32 CurrentVisitIndex = 0;
 };
 
