[ChaosVD] Changed default geometry visualization flags

By default we were not showing Query only and Complex geometry to reduce visual noise, but this could be confusing for a new user as things will appear missing. Also, it is not the best default for most projects.

This CL enables Complex Geometry and Query Only geometry visualization by default.

In addition, this change makes saving all visualization flags save-able so you don't need to change them each time.

Lastly, now each change gets auto-saved automatically, but a button was added to revert any setting category to its defaults if needed.

diff --git a/Engine/Plugins/ChaosVD/Source/ChaosVD/Private/Actors/ChaosVDSolverInfoActor.cpp b/Engine/Plugins/ChaosVD/Source/ChaosVD/Private/Actors/ChaosVDSolverInfoActor.cpp
--- a/Engine/Plugins/ChaosVD/Source/ChaosVD/Private/Actors/ChaosVDSolverInfoActor.cpp
+++ b/Engine/Plugins/ChaosVD/Source/ChaosVD/Private/Actors/ChaosVDSolverInfoActor.cpp
@@ -5,6 +5,7 @@
 #include "ChaosVDModule.h"
 #include "ChaosVDParticleActor.h"
 #include "ChaosVDScene.h"
+#include "ChaosVDSettingsManager.h"
 #include "EditorActorFolders.h"
 #include "Components/ChaosVDParticleDataComponent.h"
 #include "Components/ChaosVDSolverCharacterGroundConstraintDataComponent.h"
@@ -25,16 +26,15 @@ AChaosVDSolverInfoActor::AChaosVDSolverInfoActor(const FObjectInitializer& Objec
 	CharacterGroundConstraintDataComponent = CreateDefaultSubobject<UChaosVDSolverCharacterGroundConstraintDataComponent>(TEXT("CharacterGroundConstraintDataComponent"));
 	bIsServer = false;
 
-	if (UChaosVDParticleVisualizationSettings* ParticleVisualizationSettings = GetMutableDefault<UChaosVDParticleVisualizationSettings>())
+	if (UChaosVDParticleVisualizationSettings* ParticleVisualizationSettings = FChaosVDSettingsManager::Get().GetSettingsObject<UChaosVDParticleVisualizationSettings>())
 	{
 		ParticleVisualizationSettings->OnSettingsChanged().AddUObject(this, &AChaosVDSolverInfoActor::HandleVisibilitySettingsUpdated);
 	}
 	
-	if (UChaosVDParticleVisualizationColorSettings* ColorVisualizationSettings = GetMutableDefault<UChaosVDParticleVisualizationColorSettings>())
+	if (UChaosVDParticleVisualizationColorSettings* ColorVisualizationSettings = FChaosVDSettingsManager::Get().GetSettingsObject<UChaosVDParticleVisualizationColorSettings>())
 	{
 		ColorVisualizationSettings->OnSettingsChanged().AddUObject(this, &AChaosVDSolverInfoActor::HandleColorsSettingsUpdated);
 	}
-
 }
 
 void AChaosVDSolverInfoActor::SetSolverName(const FName& InSolverName)
@@ -231,12 +231,12 @@ void AChaosVDSolverInfoActor::Destroyed()
 		return;
 	}
 
-	if (UChaosVDParticleVisualizationSettings* ParticleVisualizationSettings = GetMutableDefault<UChaosVDParticleVisualizationSettings>())
+	if (UChaosVDParticleVisualizationSettings* ParticleVisualizationSettings = FChaosVDSettingsManager::Get().GetSettingsObject<UChaosVDParticleVisualizationSettings>())
 	{
 		ParticleVisualizationSettings->OnSettingsChanged().RemoveAll(this);
 	}
 	
-	if (UChaosVDParticleVisualizationColorSettings* ColorVisualizationSettings = GetMutableDefault<UChaosVDParticleVisualizationColorSettings>())
+	if (UChaosVDParticleVisualizationColorSettings* ColorVisualizationSettings = FChaosVDSettingsManager::Get().GetSettingsObject<UChaosVDParticleVisualizationColorSettings>())
 	{
 		ColorVisualizationSettings->OnSettingsChanged().RemoveAll(this);
 	}
diff --git a/Engine/Plugins/ChaosVD/Source/ChaosVD/Private/ChaosVDGeometryDataComponent.cpp b/Engine/Plugins/ChaosVD/Source/ChaosVD/Private/ChaosVDGeometryDataComponent.cpp
--- a/Engine/Plugins/ChaosVD/Source/ChaosVD/Private/ChaosVDGeometryDataComponent.cpp
+++ b/Engine/Plugins/ChaosVD/Source/ChaosVD/Private/ChaosVDGeometryDataComponent.cpp
@@ -6,6 +6,7 @@
 #include "ChaosVDGeometryBuilder.h"
 #include "ChaosVDModule.h"
 #include "ChaosVDParticleActor.h"
+#include "ChaosVDSettingsManager.h"
 #include "Components/MeshComponent.h"
 #include "Materials/Material.h"
 #include "Materials/MaterialInterface.h"
@@ -274,7 +275,7 @@ void FChaosVDGeometryComponentUtils::UpdateMeshVisibility(const TSharedPtr<FChao
 		return;
 	}
 
-	if (const UChaosVDParticleVisualizationSettings* ParticleVisualizationSettings = GetDefault<UChaosVDParticleVisualizationSettings>())
+	if (const UChaosVDParticleVisualizationSettings* ParticleVisualizationSettings = FChaosVDSettingsManager::Get().GetSettingsObject<UChaosVDParticleVisualizationSettings>())
 	{
 		const EChaosVDGeometryVisibilityFlags CurrentVisibilityFlags = ParticleVisualizationSettings->GetGeometryVisualizationFlags();
 		
@@ -330,7 +331,7 @@ FLinearColor FChaosVDGeometryComponentUtils::GetGeometryParticleColor(const TSha
 		return ColorToApply;
 	}
 
-	const UChaosVDParticleVisualizationColorSettings* VisualizationSettings = GetDefault<UChaosVDParticleVisualizationColorSettings>();
+	const UChaosVDParticleVisualizationColorSettings* VisualizationSettings = FChaosVDSettingsManager::Get().GetSettingsObject<UChaosVDParticleVisualizationColorSettings>();
 	if (!VisualizationSettings)
 	{
 		return ColorToApply;
@@ -379,7 +380,7 @@ FLinearColor FChaosVDGeometryComponentUtils::GetGeometryParticleColor(const TSha
 
 UMaterialInterface* FChaosVDGeometryComponentUtils::GetBaseMaterialForType(EChaosVDMaterialType Type)
 {
-	const UChaosVDCoreSettings* EditorSettings = GetDefault<UChaosVDCoreSettings>();
+	const UChaosVDCoreSettings* EditorSettings = FChaosVDSettingsManager::Get().GetSettingsObject<UChaosVDCoreSettings>();
 	if (!EditorSettings)
 	{
 		return nullptr;
diff --git a/Engine/Plugins/ChaosVD/Source/ChaosVD/Private/ChaosVDModule.cpp b/Engine/Plugins/ChaosVD/Source/ChaosVD/Private/ChaosVDModule.cpp
--- a/Engine/Plugins/ChaosVD/Source/ChaosVD/Private/ChaosVDModule.cpp
+++ b/Engine/Plugins/ChaosVD/Source/ChaosVD/Private/ChaosVDModule.cpp
@@ -5,6 +5,7 @@
 #include "ChaosVDCommands.h"
 #include "ChaosVDEngine.h"
 #include "ChaosVDParticleActorCustomization.h"
+#include "ChaosVDSettingsManager.h"
 #include "ChaosVDStyle.h"
 #include "ChaosVDTabsIDs.h"
 #include "DetailsCustomizations/ChaosVDGeometryComponentCustomization.h"
@@ -87,6 +88,8 @@ void FChaosVDModule::ShutdownModule()
 	FCoreDelegates::OnEnginePreExit.RemoveAll(this);
 
 	CloseActiveInstances();
+	
+	FChaosVDSettingsManager::TearDown();
 }
 
 void FChaosVDModule::RegisterClassesCustomDetails() const
diff --git a/Engine/Plugins/ChaosVD/Source/ChaosVD/Private/ChaosVDScene.cpp b/Engine/Plugins/ChaosVD/Source/ChaosVD/Private/ChaosVDScene.cpp
--- a/Engine/Plugins/ChaosVD/Source/ChaosVD/Private/ChaosVDScene.cpp
+++ b/Engine/Plugins/ChaosVD/Source/ChaosVD/Private/ChaosVDScene.cpp
@@ -10,6 +10,7 @@
 #include "Chaos/ImplicitObject.h"
 #include "ChaosVDRecording.h"
 #include "ChaosVDSelectionCustomization.h"
+#include "ChaosVDSettingsManager.h"
 #include "ChaosVDSkySphereInterface.h"
 #include "Components/ChaosVDSceneQueryDataComponent.h"
 #include "Components/ChaosVDSolverCharacterGroundConstraintDataComponent.h"
@@ -66,7 +67,7 @@ void FChaosVDScene::Initialize()
 	
 	StreamableManager = MakeShared<FStreamableManager>();
 
-	if (UChaosVDCoreSettings* Settings = GetMutableDefault<UChaosVDCoreSettings>())
+	if (UChaosVDCoreSettings* Settings = FChaosVDSettingsManager::Get().GetSettingsObject<UChaosVDCoreSettings>())
 	{
 		// TODO: Do an async load instead, and prepare a loading screen or notification popup
 		// Jira for tracking UE-191639
@@ -521,7 +522,7 @@ void FChaosVDScene::CreateBaseLights(UWorld* TargetWorld) const
 
 	const FVector SpawnPosition(0.0, 0.0, 2000.0);
 	
-	if (const UChaosVDCoreSettings* Settings = GetDefault<UChaosVDCoreSettings>())
+	if (const UChaosVDCoreSettings* Settings = FChaosVDSettingsManager::Get().GetSettingsObject<UChaosVDCoreSettings>())
 	{
 		if (ADirectionalLight* DirectionalLightActor = TargetWorld->SpawnActor<ADirectionalLight>())
 		{
diff --git a/Engine/Plugins/ChaosVD/Source/ChaosVD/Private/ChaosVDSettingsManager.cpp b/Engine/Plugins/ChaosVD/Source/ChaosVD/Private/ChaosVDSettingsManager.cpp
new file mode 100644
--- /dev/null
+++ b/Engine/Plugins/ChaosVD/Source/ChaosVD/Private/ChaosVDSettingsManager.cpp
@@ -0,0 +1,88 @@
+// Copyright Epic Games, Inc. All Rights Reserved.
+
+#include "ChaosVDSettingsManager.h"
+
+#include "Misc/ConfigContext.h"
+#include "Misc/LazySingleton.h"
+#include "Settings/ChaosVDCoreSettings.h"
+
+FChaosVDSettingsManager::FChaosVDSettingsManager()
+{
+	SettingsOuter = NewObject<UChaosVDSettingsObjectsOuter>();
+}
+
+void FChaosVDSettingsManager::AddReferencedObjects(FReferenceCollector& Collector)
+{
+	Collector.AddReferencedObject(SettingsOuter);
+	Collector.AddStableReferenceMap(AvailableSettingsObject);
+}
+
+FString FChaosVDSettingsManager::GetReferencerName() const
+{
+	return TEXT("FChaosVDSettingsManager");
+}
+
+FChaosVDSettingsManager& FChaosVDSettingsManager::Get()
+{
+	return TLazySingleton<FChaosVDSettingsManager>::Get();
+}
+
+void FChaosVDSettingsManager::TearDown()
+{
+	TLazySingleton<FChaosVDSettingsManager>::TearDown();
+}
+
+UChaosVDSettingsObjectBase* FChaosVDSettingsManager::GetSettingsObject(UClass* SettingsClass)
+{
+	if (TObjectPtr<UChaosVDSettingsObjectBase>* ObjectPtrPtr = AvailableSettingsObject.Find(SettingsClass))
+	{
+		return ObjectPtrPtr->Get();
+	}
+
+	UChaosVDSettingsObjectBase* NewSettingsObject = NewObject<UChaosVDSettingsObjectBase>(SettingsOuter, SettingsClass);
+
+	AvailableSettingsObject.Add(SettingsClass, NewSettingsObject);
+
+	return NewSettingsObject;
+}
+
+void FChaosVDSettingsManager::ResetSettings(UClass* SettingsClass)
+{
+	UChaosVDSettingsObjectBase* SettingsObject = GetSettingsObject(SettingsClass);
+	if (ensure(GConfig && SettingsClass && SettingsObject))
+	{
+		FString ConfigFile = SettingsClass->GetConfigName();
+		FStringView Section = SettingsObject->GetConfigSectionName();
+
+		GConfig->EmptySection(Section.GetData(), ConfigFile);
+		GConfig->Flush(false);
+
+		FConfigContext::ForceReloadIntoGConfig().Load(*FPaths::GetBaseFilename(ConfigFile));
+	
+		RestoreConfigPropertiesValuesFromCDO(SettingsObject);
+
+		SettingsObject->BroadcastSettingsChanged();
+	}
+}
+
+void FChaosVDSettingsManager::RestoreConfigPropertiesValuesFromCDO(UChaosVDSettingsObjectBase* TargetSettingsObject)
+{
+	const UChaosVDSettingsObjectBase* CDOSettingsObject = TargetSettingsObject ? GetDefault<UChaosVDSettingsObjectBase>(TargetSettingsObject->GetClass()) : nullptr;
+	if (!ensure(CDOSettingsObject && TargetSettingsObject))
+	{
+		return;
+	}
+
+	for (FProperty* Property = CDOSettingsObject->GetClass()->PropertyLink; Property; Property = Property->PropertyLinkNext)
+	{
+		if (!Property->HasAnyPropertyFlags(CPF_Config) )
+		{
+			continue;
+		}
+
+		const void* CDOPropertyAddress = Property->ContainerPtrToValuePtr<void>(CDOSettingsObject);
+		void* TargetPropertyAddr = Property->ContainerPtrToValuePtr<void>(TargetSettingsObject);
+
+		Property->CopyCompleteValue(TargetPropertyAddr, CDOPropertyAddress);
+	}
+}
diff --git a/Engine/Plugins/ChaosVD/Source/ChaosVD/Private/ChaosVDSettingsManager.h b/Engine/Plugins/ChaosVD/Source/ChaosVD/Private/ChaosVDSettingsManager.h
new file mode 100644
--- /dev/null
+++ b/Engine/Plugins/ChaosVD/Source/ChaosVD/Private/ChaosVDSettingsManager.h
@@ -0,0 +1,54 @@
+// Copyright Epic Games, Inc. All Rights Reserved.
+
+#pragma once
+
+#include "Settings/ChaosVDCoreSettings.h"
+
+class UChaosVDSettingsObjectBase;
+
+/** Manager class that handles all available CVD settings objects*/
+class FChaosVDSettingsManager : public FGCObject
+{
+public:
+	FChaosVDSettingsManager();
+
+	virtual void AddReferencedObjects(FReferenceCollector& Collector) override;
+	virtual FString GetReferencerName() const override;
+
+	static FChaosVDSettingsManager& Get();
+	static void TearDown();
+
+	/** Returns the setting object of the specified class - If the object does not exist yet, it will be created during this call */
+	template<typename SettingsClass>
+	SettingsClass* GetSettingsObject();
+
+	/** Returns the setting object of the specified class - If the object does not exist yet, it will be created during this call */
+	UChaosVDSettingsObjectBase* GetSettingsObject(UClass* SettingsClass);
+
+	/** Deletes any saved config for the setting object of the specified class, and restores its values to be the same as the CDO */
+	template<typename SettingsClass>
+	void ResetSettings();
+
+	/** Deletes any saved config for the setting object of the specified class, and restores its values to be the same as the CDO */
+	void ResetSettings(UClass* SettingsClass);
+
+protected:
+
+	void RestoreConfigPropertiesValuesFromCDO(UChaosVDSettingsObjectBase* TargetSettingsObject);
+	
+	TMap<TObjectPtr<UClass>, TObjectPtr<UChaosVDSettingsObjectBase>> AvailableSettingsObject;
+	TObjectPtr<UChaosVDSettingsObjectsOuter> SettingsOuter;
+};
+
+template <typename SettingsClass>
+SettingsClass* FChaosVDSettingsManager::GetSettingsObject()
+{
+	return Cast<SettingsClass>(GetSettingsObject(SettingsClass::StaticClass()));
+}
+
+template <typename SettingsClass>
+void FChaosVDSettingsManager::ResetSettings()
+{
+	UClass* ObjectClass = SettingsClass::StaticClass();
+	ResetSettings(ObjectClass);
+}
diff --git a/Engine/Plugins/ChaosVD/Source/ChaosVD/Private/Components/ChaosVDSolverCharacterGroundConstraintDataComponent.cpp b/Engine/Plugins/ChaosVD/Source/ChaosVD/Private/Components/ChaosVDSolverCharacterGroundConstraintDataComponent.cpp
--- a/Engine/Plugins/ChaosVD/Source/ChaosVD/Private/Components/ChaosVDSolverCharacterGroundConstraintDataComponent.cpp
+++ b/Engine/Plugins/ChaosVD/Source/ChaosVD/Private/Components/ChaosVDSolverCharacterGroundConstraintDataComponent.cpp
@@ -4,12 +4,13 @@
 
 #include "ChaosVDCharacterGroundConstraintDataProviderInterface.h"
 #include "ChaosVDScene.h"
+#include "ChaosVDSettingsManager.h"
 #include "Selection.h"
 #include "Settings/ChaosVDCharacterConstraintsVisualizationSettings.h"
 
 void UChaosVDSolverCharacterGroundConstraintDataComponent::HandleSceneUpdated()
 {
-	if (const UChaosVDCharacterConstraintsVisualizationSettings* CharacterConstraintsVisualizationSettings = GetDefault<UChaosVDCharacterConstraintsVisualizationSettings>())
+	if (const UChaosVDCharacterConstraintsVisualizationSettings* CharacterConstraintsVisualizationSettings = FChaosVDSettingsManager::Get().GetSettingsObject<UChaosVDCharacterConstraintsVisualizationSettings>())
 	{
 		if (!CharacterConstraintsVisualizationSettings->bAutoSelectConstraintFromSelectedParticle)
 		{
diff --git a/Engine/Plugins/ChaosVD/Source/ChaosVD/Private/Settings/ChaosVDCharacterConstraintsVisualizationSettings.cpp b/Engine/Plugins/ChaosVD/Source/ChaosVD/Private/Settings/ChaosVDCharacterConstraintsVisualizationSettings.cpp
--- a/Engine/Plugins/ChaosVD/Source/ChaosVD/Private/Settings/ChaosVDCharacterConstraintsVisualizationSettings.cpp
+++ b/Engine/Plugins/ChaosVD/Source/ChaosVD/Private/Settings/ChaosVDCharacterConstraintsVisualizationSettings.cpp
@@ -2,20 +2,23 @@
 
 #include "ChaosVDCharacterConstraintsVisualizationSettings.h"
 
+#include "ChaosVDSettingsManager.h"
+
+
 void UChaosVDCharacterConstraintsVisualizationSettings::SetDataVisualizationFlags(EChaosVDCharacterGroundConstraintDataVisualizationFlags NewFlags)
 {
-	if (UChaosVDCharacterConstraintsVisualizationSettings* Settings = GetMutableDefault<UChaosVDCharacterConstraintsVisualizationSettings>())
+	if (UChaosVDCharacterConstraintsVisualizationSettings* Settings = FChaosVDSettingsManager::Get().GetSettingsObject<UChaosVDCharacterConstraintsVisualizationSettings>())
 	{
-		Settings->GlobalCharacterGroundConstraintDataVisualizationFlags = NewFlags;
+		Settings->GlobalCharacterGroundConstraintDataVisualizationFlags = static_cast<uint32>(NewFlags);
 		Settings->BroadcastSettingsChanged();
 	}
 }
 
 EChaosVDCharacterGroundConstraintDataVisualizationFlags UChaosVDCharacterConstraintsVisualizationSettings::GetDataVisualizationFlags()
 {
-	if (UChaosVDCharacterConstraintsVisualizationSettings* Settings = GetMutableDefault<UChaosVDCharacterConstraintsVisualizationSettings>())
+	if (UChaosVDCharacterConstraintsVisualizationSettings* Settings = FChaosVDSettingsManager::Get().GetSettingsObject<UChaosVDCharacterConstraintsVisualizationSettings>())
 	{
-		return Settings->GlobalCharacterGroundConstraintDataVisualizationFlags;
+		return static_cast<EChaosVDCharacterGroundConstraintDataVisualizationFlags>(Settings->GlobalCharacterGroundConstraintDataVisualizationFlags);
 	}
 
 	return EChaosVDCharacterGroundConstraintDataVisualizationFlags::None;
diff --git a/Engine/Plugins/ChaosVD/Source/ChaosVD/Private/Settings/ChaosVDCharacterConstraintsVisualizationSettings.h b/Engine/Plugins/ChaosVD/Source/ChaosVD/Private/Settings/ChaosVDCharacterConstraintsVisualizationSettings.h
--- a/Engine/Plugins/ChaosVD/Source/ChaosVD/Private/Settings/ChaosVDCharacterConstraintsVisualizationSettings.h
+++ b/Engine/Plugins/ChaosVD/Source/ChaosVD/Private/Settings/ChaosVDCharacterConstraintsVisualizationSettings.h
@@ -33,7 +33,7 @@ enum class EChaosVDCharacterGroundConstraintDataVisualizationFlags : uint32
 };
 ENUM_CLASS_FLAGS(EChaosVDCharacterGroundConstraintDataVisualizationFlags);
 
-UCLASS(config=ChaosVD)
+UCLASS(config=ChaosVD, PerObjectConfig)
 class UChaosVDCharacterConstraintsVisualizationSettings : public UChaosVDVisualizationSettingsObjectBase
 {
 	GENERATED_BODY()
@@ -78,6 +78,6 @@ public:
 
 private:
 	/** Set of flags to enable/disable visualization of specific character ground constraint data as debug draw */
-	EChaosVDCharacterGroundConstraintDataVisualizationFlags GlobalCharacterGroundConstraintDataVisualizationFlags = EChaosVDCharacterGroundConstraintDataVisualizationFlags::GroundQueryDistance | EChaosVDCharacterGroundConstraintDataVisualizationFlags::GroundQueryNormal | EChaosVDCharacterGroundConstraintDataVisualizationFlags::TargetDeltaPosition;
-
+	UPROPERTY(config, meta = (Bitmask, BitmaskEnum = "/Script/ChaosVD.EChaosVDCharacterGroundConstraintDataVisualizationFlags"))
+	uint32 GlobalCharacterGroundConstraintDataVisualizationFlags = static_cast<uint32>(EChaosVDCharacterGroundConstraintDataVisualizationFlags::GroundQueryDistance | EChaosVDCharacterGroundConstraintDataVisualizationFlags::GroundQueryNormal | EChaosVDCharacterGroundConstraintDataVisualizationFlags::TargetDeltaPosition);
 };
diff --git a/Engine/Plugins/ChaosVD/Source/ChaosVD/Private/Settings/ChaosVDCollisionVisualizationSettings.cpp b/Engine/Plugins/ChaosVD/Source/ChaosVD/Private/Settings/ChaosVDCollisionVisualizationSettings.cpp
--- a/Engine/Plugins/ChaosVD/Source/ChaosVD/Private/Settings/ChaosVDCollisionVisualizationSettings.cpp
+++ b/Engine/Plugins/ChaosVD/Source/ChaosVD/Private/Settings/ChaosVDCollisionVisualizationSettings.cpp
@@ -2,20 +2,23 @@
 
 #include "ChaosVDCollisionVisualizationSettings.h"
 
+#include "ChaosVDSettingsManager.h"
+
+
 void UChaosVDCollisionDataVisualizationSettings::SetDataVisualizationFlags(EChaosVDCollisionVisualizationFlags NewFlags)
 {
-	if (UChaosVDCollisionDataVisualizationSettings* Settings = GetMutableDefault<UChaosVDCollisionDataVisualizationSettings>())
+	if (UChaosVDCollisionDataVisualizationSettings* Settings = FChaosVDSettingsManager::Get().GetSettingsObject<UChaosVDCollisionDataVisualizationSettings>())
 	{
-		Settings->CollisionDataVisualizationFlags = NewFlags;
+		Settings->CollisionDataVisualizationFlags = static_cast<uint32>(NewFlags);
 		Settings->BroadcastSettingsChanged();
 	}
 }
 
 EChaosVDCollisionVisualizationFlags UChaosVDCollisionDataVisualizationSettings::GetDataVisualizationFlags()
 {
-	if (UChaosVDCollisionDataVisualizationSettings* Settings = GetMutableDefault<UChaosVDCollisionDataVisualizationSettings>())
+	if (UChaosVDCollisionDataVisualizationSettings* Settings = FChaosVDSettingsManager::Get().GetSettingsObject<UChaosVDCollisionDataVisualizationSettings>())
 	{
-		return Settings->CollisionDataVisualizationFlags ;
+		return static_cast<EChaosVDCollisionVisualizationFlags>(Settings->CollisionDataVisualizationFlags) ;
 	}
 
 	return EChaosVDCollisionVisualizationFlags::None;
diff --git a/Engine/Plugins/ChaosVD/Source/ChaosVD/Private/Settings/ChaosVDCollisionVisualizationSettings.h b/Engine/Plugins/ChaosVD/Source/ChaosVD/Private/Settings/ChaosVDCollisionVisualizationSettings.h
--- a/Engine/Plugins/ChaosVD/Source/ChaosVD/Private/Settings/ChaosVDCollisionVisualizationSettings.h
+++ b/Engine/Plugins/ChaosVD/Source/ChaosVD/Private/Settings/ChaosVDCollisionVisualizationSettings.h
@@ -21,7 +21,7 @@ enum class EChaosVDCollisionVisualizationFlags: uint32
 };
 ENUM_CLASS_FLAGS(EChaosVDCollisionVisualizationFlags);
 
-UCLASS(config=ChaosVD)
+UCLASS(config=ChaosVD, PerObjectConfig)
 class UChaosVDCollisionDataVisualizationSettings : public UChaosVDVisualizationSettingsObjectBase
 {
 	GENERATED_BODY()
@@ -47,7 +47,8 @@ public:
 	static EChaosVDCollisionVisualizationFlags GetDataVisualizationFlags();
 
 private:
+	UPROPERTY(config, meta = (Bitmask, BitmaskEnum = "/Script/ChaosVD.EChaosVDCollisionVisualizationFlags"))
 	/** Set of flags to enable/disable visualization of specific collision data as debug draw */
-	EChaosVDCollisionVisualizationFlags CollisionDataVisualizationFlags = EChaosVDCollisionVisualizationFlags::ContactInfo | EChaosVDCollisionVisualizationFlags::ContactPoints | EChaosVDCollisionVisualizationFlags::ContactNormal;
+	uint32 CollisionDataVisualizationFlags = static_cast<uint32>(EChaosVDCollisionVisualizationFlags::ContactInfo | EChaosVDCollisionVisualizationFlags::ContactPoints | EChaosVDCollisionVisualizationFlags::ContactNormal);
 };
 
diff --git a/Engine/Plugins/ChaosVD/Source/ChaosVD/Private/Settings/ChaosVDCoreSettings.cpp b/Engine/Plugins/ChaosVD/Source/ChaosVD/Private/Settings/ChaosVDCoreSettings.cpp
--- a/Engine/Plugins/ChaosVD/Source/ChaosVD/Private/Settings/ChaosVDCoreSettings.cpp
+++ b/Engine/Plugins/ChaosVD/Source/ChaosVD/Private/Settings/ChaosVDCoreSettings.cpp
@@ -2,9 +2,15 @@
 
 #include "ChaosVDCoreSettings.h"
 
+#include "Misc/ConfigContext.h"
 #include "Widgets/SChaosVDPlaybackViewport.h"
 
 
+UChaosVDSettingsObjectBase::UChaosVDSettingsObjectBase()
+{
+	
+}
+
 void UChaosVDSettingsObjectBase::PostEditChangeProperty(FPropertyChangedEvent& PropertyChangedEvent)
 {
 	UObject::PostEditChangeProperty(PropertyChangedEvent);
@@ -18,9 +24,22 @@ void UChaosVDSettingsObjectBase::PostEditUndo()
 	BroadcastSettingsChanged();
 }
 
+void UChaosVDSettingsObjectBase::OverridePerObjectConfigSection(FString& SectionName)
+{
+	if (OverrideConfigSectionName.IsEmpty())
+	{
+		OverrideConfigSectionName = GetClass()->GetPathName() + TEXT(" Instance");
+	}
+
+	SectionName = OverrideConfigSectionName;
+}
+
 void UChaosVDSettingsObjectBase::BroadcastSettingsChanged()
 {
 	SettingsChangedDelegate.Broadcast(this);
+
+	constexpr bool bAllowCopyToDefaultObject = false;
+	SaveConfig(CPF_Config,nullptr, GConfig, bAllowCopyToDefaultObject);
 }
 
 void UChaosVDVisualizationSettingsObjectBase::BroadcastSettingsChanged()
diff --git a/Engine/Plugins/ChaosVD/Source/ChaosVD/Private/Settings/ChaosVDCoreSettings.h b/Engine/Plugins/ChaosVD/Source/ChaosVD/Private/Settings/ChaosVDCoreSettings.h
--- a/Engine/Plugins/ChaosVD/Source/ChaosVD/Private/Settings/ChaosVDCoreSettings.h
+++ b/Engine/Plugins/ChaosVD/Source/ChaosVD/Private/Settings/ChaosVDCoreSettings.h
@@ -14,23 +14,46 @@ class UMaterial;
 
 DECLARE_MULTICAST_DELEGATE_OneParam(FChaosVDSettingChanged, UObject* SettingsObject)
 
+UCLASS()
+class UChaosVDSettingsObjectsOuter : public UObject
+{
+	GENERATED_BODY()
+};
+
 UCLASS(config = ChaosVD)
 class UChaosVDSettingsObjectBase : public UObject
 {
+public:
+	UChaosVDSettingsObjectBase();
+
+private:
 	GENERATED_BODY()
+
 public:
 
+
 	virtual void PostEditChangeProperty(FPropertyChangedEvent& PropertyChangedEvent) override;
 
 	FChaosVDSettingChanged& OnSettingsChanged() { return SettingsChangedDelegate; }
 
 	virtual void PostEditUndo() override;
 
+	virtual void OverridePerObjectConfigSection(FString& SectionName) override;
+
+	FStringView GetConfigSectionName()
+	{
+		return OverrideConfigSectionName;
+	}
+
 protected:
 	virtual void BroadcastSettingsChanged();
 	
 private:
+
+	FString OverrideConfigSectionName;
 	FChaosVDSettingChanged SettingsChangedDelegate;
+
+	friend class FChaosVDSettingsManager;
 };
 
 UCLASS(config = ChaosVD)
diff --git a/Engine/Plugins/ChaosVD/Source/ChaosVD/Private/Settings/ChaosVDJointConstraintVisualizationSettings.cpp b/Engine/Plugins/ChaosVD/Source/ChaosVD/Private/Settings/ChaosVDJointConstraintVisualizationSettings.cpp
--- a/Engine/Plugins/ChaosVD/Source/ChaosVD/Private/Settings/ChaosVDJointConstraintVisualizationSettings.cpp
+++ b/Engine/Plugins/ChaosVD/Source/ChaosVD/Private/Settings/ChaosVDJointConstraintVisualizationSettings.cpp
@@ -2,20 +2,23 @@
 
 #include "ChaosVDJointConstraintVisualizationSettings.h"
 
+#include "ChaosVDSettingsManager.h"
+
+
 void UChaosVDJointConstraintsVisualizationSettings::SetDataVisualizationFlags(EChaosVDJointsDataVisualizationFlags NewFlags)
 {
-	if (UChaosVDJointConstraintsVisualizationSettings* Settings = GetMutableDefault<UChaosVDJointConstraintsVisualizationSettings>())
+	if (UChaosVDJointConstraintsVisualizationSettings* Settings = FChaosVDSettingsManager::Get().GetSettingsObject<UChaosVDJointConstraintsVisualizationSettings>())
 	{
-		Settings->GlobalJointsDataVisualizationFlags = NewFlags;
+		Settings->GlobalJointsDataVisualizationFlags = static_cast<uint32>(NewFlags);
 		Settings->BroadcastSettingsChanged();
 	}
 }
 
 EChaosVDJointsDataVisualizationFlags UChaosVDJointConstraintsVisualizationSettings::GetDataVisualizationFlags()
 {
-	if (UChaosVDJointConstraintsVisualizationSettings* Settings = GetMutableDefault<UChaosVDJointConstraintsVisualizationSettings>())
+	if (UChaosVDJointConstraintsVisualizationSettings* Settings = FChaosVDSettingsManager::Get().GetSettingsObject<UChaosVDJointConstraintsVisualizationSettings>())
 	{
-		return Settings->GlobalJointsDataVisualizationFlags;
+		return static_cast<EChaosVDJointsDataVisualizationFlags>(Settings->GlobalJointsDataVisualizationFlags);
 	}
 
 	return EChaosVDJointsDataVisualizationFlags::None;
diff --git a/Engine/Plugins/ChaosVD/Source/ChaosVD/Private/Settings/ChaosVDJointConstraintVisualizationSettings.h b/Engine/Plugins/ChaosVD/Source/ChaosVD/Private/Settings/ChaosVDJointConstraintVisualizationSettings.h
--- a/Engine/Plugins/ChaosVD/Source/ChaosVD/Private/Settings/ChaosVDJointConstraintVisualizationSettings.h
+++ b/Engine/Plugins/ChaosVD/Source/ChaosVD/Private/Settings/ChaosVDJointConstraintVisualizationSettings.h
@@ -28,7 +28,7 @@ enum class EChaosVDJointsDataVisualizationFlags : uint32
 };
 ENUM_CLASS_FLAGS(EChaosVDJointsDataVisualizationFlags);
 
-UCLASS(config=ChaosVD)
+UCLASS(config=ChaosVD, PerObjectConfig)
 class UChaosVDJointConstraintsVisualizationSettings : public UChaosVDVisualizationSettingsObjectBase
 {
 	GENERATED_BODY()
@@ -72,5 +72,6 @@ public:
 
 private:
 	/** Set of flags to enable/disable visualization of specific joint constraints data as debug draw */
-	EChaosVDJointsDataVisualizationFlags GlobalJointsDataVisualizationFlags = EChaosVDJointsDataVisualizationFlags::ActorConnector | EChaosVDJointsDataVisualizationFlags::DrawKinematic;
+	UPROPERTY(config, meta = (Bitmask, BitmaskEnum = "/Script/ChaosVD.EChaosVDJointsDataVisualizationFlags"))
+	uint32 GlobalJointsDataVisualizationFlags = static_cast<uint32>(EChaosVDJointsDataVisualizationFlags::ActorConnector | EChaosVDJointsDataVisualizationFlags::DrawKinematic);
 };
\ No newline at end of file
diff --git a/Engine/Plugins/ChaosVD/Source/ChaosVD/Private/Settings/ChaosVDParticleVisualizationSettings.cpp b/Engine/Plugins/ChaosVD/Source/ChaosVD/Private/Settings/ChaosVDParticleVisualizationSettings.cpp
--- a/Engine/Plugins/ChaosVD/Source/ChaosVD/Private/Settings/ChaosVDParticleVisualizationSettings.cpp
+++ b/Engine/Plugins/ChaosVD/Source/ChaosVD/Private/Settings/ChaosVDParticleVisualizationSettings.cpp
@@ -2,6 +2,7 @@
 
 #include "ChaosVDParticleVisualizationSettings.h"
 
+#include "ChaosVDSettingsManager.h"
 #include "DataWrappers/ChaosVDParticleDataWrapper.h"
 
 
@@ -86,18 +87,18 @@ float UChaosVDParticleVisualizationDebugDrawSettings::GetScaleFortDataID(EChaosV
 
 void UChaosVDParticleVisualizationDebugDrawSettings::SetDataDebugDrawVisualizationFlags(EChaosVDParticleDataVisualizationFlags Flags)
 {
-	if (UChaosVDParticleVisualizationDebugDrawSettings* Settings = GetMutableDefault<UChaosVDParticleVisualizationDebugDrawSettings>())
+	if (UChaosVDParticleVisualizationDebugDrawSettings* Settings = FChaosVDSettingsManager::Get().GetSettingsObject<UChaosVDParticleVisualizationDebugDrawSettings>())
 	{
-		Settings->ParticleDataVisualizationFlags = Flags;
+		Settings->ParticleDataVisualizationFlags = static_cast<uint32>(Flags);
 		Settings->BroadcastSettingsChanged();
 	}
 }
 
 EChaosVDParticleDataVisualizationFlags UChaosVDParticleVisualizationDebugDrawSettings::GetDataDebugDrawVisualizationFlags()
 {
-	if (UChaosVDParticleVisualizationDebugDrawSettings* Settings = GetMutableDefault<UChaosVDParticleVisualizationDebugDrawSettings>())
+	if (UChaosVDParticleVisualizationDebugDrawSettings* Settings = FChaosVDSettingsManager::Get().GetSettingsObject<UChaosVDParticleVisualizationDebugDrawSettings>())
 	{
-		return Settings->ParticleDataVisualizationFlags;
+		return static_cast<EChaosVDParticleDataVisualizationFlags>(Settings->ParticleDataVisualizationFlags);
 	}
 
 	return EChaosVDParticleDataVisualizationFlags::None;
@@ -105,18 +106,18 @@ EChaosVDParticleDataVisualizationFlags UChaosVDParticleVisualizationDebugDrawSet
 
 void UChaosVDParticleVisualizationSettings::SetGeometryVisualizationFlags(EChaosVDGeometryVisibilityFlags Flags)
 {
-	if (UChaosVDParticleVisualizationSettings* Settings = GetMutableDefault<UChaosVDParticleVisualizationSettings>())
+	if (UChaosVDParticleVisualizationSettings* Settings = FChaosVDSettingsManager::Get().GetSettingsObject<UChaosVDParticleVisualizationSettings>())
 	{
-		Settings->GeometryVisibilityFlags = Flags;
+		Settings->GeometryVisibilityFlags = static_cast<uint32>(Flags);
 		Settings->BroadcastSettingsChanged();
 	}
 }
 
 EChaosVDGeometryVisibilityFlags UChaosVDParticleVisualizationSettings::GetGeometryVisualizationFlags()
 {
-	if (UChaosVDParticleVisualizationSettings* Settings = GetMutableDefault<UChaosVDParticleVisualizationSettings>())
+	if (UChaosVDParticleVisualizationSettings* Settings = FChaosVDSettingsManager::Get().GetSettingsObject<UChaosVDParticleVisualizationSettings>())
 	{
-		return Settings->GeometryVisibilityFlags;
+		return static_cast<EChaosVDGeometryVisibilityFlags>(Settings->GeometryVisibilityFlags);
 	}
 
 	return EChaosVDGeometryVisibilityFlags::None;
diff --git a/Engine/Plugins/ChaosVD/Source/ChaosVD/Private/Settings/ChaosVDParticleVisualizationSettings.h b/Engine/Plugins/ChaosVD/Source/ChaosVD/Private/Settings/ChaosVDParticleVisualizationSettings.h
--- a/Engine/Plugins/ChaosVD/Source/ChaosVD/Private/Settings/ChaosVDParticleVisualizationSettings.h
+++ b/Engine/Plugins/ChaosVD/Source/ChaosVD/Private/Settings/ChaosVDParticleVisualizationSettings.h
@@ -223,7 +223,7 @@ namespace Chaos::VisualDebugger::ParticleDataUnitsStrings
 	}
 }
 
-UCLASS(config=ChaosVD)
+UCLASS(config=ChaosVD, PerObjectConfig)
 class UChaosVDParticleVisualizationColorSettings : public UChaosVDVisualizationSettingsObjectBase
 {
 	GENERATED_BODY()
@@ -242,7 +242,7 @@ public:
 	FChaosDebugDrawColorsByClientServer ColorsByClientServer;
 };
 
-UCLASS(config=ChaosVD)
+UCLASS(config=ChaosVD, PerObjectConfig)
 class UChaosVDParticleVisualizationDebugDrawSettings : public UChaosVDVisualizationSettingsObjectBase
 {
 	GENERATED_BODY()
@@ -294,11 +294,11 @@ public:
 
 private:
 	/** Set of flags to enable/disable visualization of specific particle data as debug draw */
-	UPROPERTY()
-	EChaosVDParticleDataVisualizationFlags ParticleDataVisualizationFlags = EChaosVDParticleDataVisualizationFlags::Velocity | EChaosVDParticleDataVisualizationFlags::AngularVelocity;
+	UPROPERTY(config, meta = (Bitmask, BitmaskEnum = "/Script/ChaosVD.EChaosVDParticleDataVisualizationFlags"))
+	uint32 ParticleDataVisualizationFlags = static_cast<uint32>(EChaosVDParticleDataVisualizationFlags::Velocity | EChaosVDParticleDataVisualizationFlags::AngularVelocity);
 };
 
-UCLASS()
+UCLASS(config=ChaosVD, PerObjectConfig)
 class UChaosVDParticleVisualizationSettings : public UChaosVDVisualizationSettingsObjectBase
 {
 	GENERATED_BODY()
@@ -310,5 +310,6 @@ public:
 private:
 
 	/** Set of flags to enable/disable visibility of specific types of geometry/particles */
-	EChaosVDGeometryVisibilityFlags GeometryVisibilityFlags = EChaosVDGeometryVisibilityFlags::Simulated | EChaosVDGeometryVisibilityFlags::Simple |  EChaosVDGeometryVisibilityFlags::ShowHeightfields;
+	UPROPERTY(config, meta = (Bitmask, BitmaskEnum = "/Script/ChaosVD.EChaosVDGeometryVisibilityFlags"))
+	uint32 GeometryVisibilityFlags = static_cast<uint32>(EChaosVDGeometryVisibilityFlags::Simulated | EChaosVDGeometryVisibilityFlags::Simple | EChaosVDGeometryVisibilityFlags::ShowHeightfields | EChaosVDGeometryVisibilityFlags::Complex | EChaosVDGeometryVisibilityFlags::Query);
 };
\ No newline at end of file
diff --git a/Engine/Plugins/ChaosVD/Source/ChaosVD/Private/Settings/ChaosVDSceneQueryVisualizationSettings.cpp b/Engine/Plugins/ChaosVD/Source/ChaosVD/Private/Settings/ChaosVDSceneQueryVisualizationSettings.cpp
--- a/Engine/Plugins/ChaosVD/Source/ChaosVD/Private/Settings/ChaosVDSceneQueryVisualizationSettings.cpp
+++ b/Engine/Plugins/ChaosVD/Source/ChaosVD/Private/Settings/ChaosVDSceneQueryVisualizationSettings.cpp
@@ -2,20 +2,23 @@
 
 #include "ChaosVDSceneQueryVisualizationSettings.h"
 
+#include "ChaosVDSettingsManager.h"
+
+
 void UChaosVDSceneQueriesVisualizationSettings::SetDataVisualizationFlags(EChaosVDSceneQueryVisualizationFlags NewFlags)
 {
-	if (UChaosVDSceneQueriesVisualizationSettings* Settings = GetMutableDefault<UChaosVDSceneQueriesVisualizationSettings>())
+	if (UChaosVDSceneQueriesVisualizationSettings* Settings = FChaosVDSettingsManager::Get().GetSettingsObject<UChaosVDSceneQueriesVisualizationSettings>())
 	{
-		Settings->GlobalSceneQueriesVisualizationFlags = NewFlags;
+		Settings->GlobalSceneQueriesVisualizationFlags = static_cast<uint32>(NewFlags);
 		Settings->BroadcastSettingsChanged();
 	}
 }
 
 EChaosVDSceneQueryVisualizationFlags UChaosVDSceneQueriesVisualizationSettings::GetDataVisualizationFlags()
 {
-	if (UChaosVDSceneQueriesVisualizationSettings* Settings = GetMutableDefault<UChaosVDSceneQueriesVisualizationSettings>())
+	if (UChaosVDSceneQueriesVisualizationSettings* Settings = FChaosVDSettingsManager::Get().GetSettingsObject<UChaosVDSceneQueriesVisualizationSettings>())
 	{
-		return Settings->GlobalSceneQueriesVisualizationFlags;
+		return static_cast<EChaosVDSceneQueryVisualizationFlags>(Settings->GlobalSceneQueriesVisualizationFlags);
 	}
 
 	return EChaosVDSceneQueryVisualizationFlags::None;
diff --git a/Engine/Plugins/ChaosVD/Source/ChaosVD/Private/Settings/ChaosVDSceneQueryVisualizationSettings.h b/Engine/Plugins/ChaosVD/Source/ChaosVD/Private/Settings/ChaosVDSceneQueryVisualizationSettings.h
--- a/Engine/Plugins/ChaosVD/Source/ChaosVD/Private/Settings/ChaosVDSceneQueryVisualizationSettings.h
+++ b/Engine/Plugins/ChaosVD/Source/ChaosVD/Private/Settings/ChaosVDSceneQueryVisualizationSettings.h
@@ -21,7 +21,7 @@ enum class EChaosVDSceneQueryVisualizationFlags: uint32
 };
 ENUM_CLASS_FLAGS(EChaosVDSceneQueryVisualizationFlags);
 
-UCLASS(config=ChaosVD)
+UCLASS(config=ChaosVD, PerObjectConfig)
 class UChaosVDSceneQueriesVisualizationSettings : public UChaosVDVisualizationSettingsObjectBase
 {
 	GENERATED_BODY()
@@ -37,5 +37,6 @@ public:
 
 private:
 	/** Set of flags to enable/disable visualization of specific scene queries data as debug draw */
-	EChaosVDSceneQueryVisualizationFlags GlobalSceneQueriesVisualizationFlags = EChaosVDSceneQueryVisualizationFlags::DrawHits | EChaosVDSceneQueryVisualizationFlags::DrawLineTraceQueries;
+	UPROPERTY(config, meta = (Bitmask, BitmaskEnum = "/Script/ChaosVD.EChaosVDSceneQueryVisualizationFlags"))
+	uint32 GlobalSceneQueriesVisualizationFlags = static_cast<uint32>(EChaosVDSceneQueryVisualizationFlags::DrawHits | EChaosVDSceneQueryVisualizationFlags::DrawLineTraceQueries);
 };
diff --git a/Engine/Plugins/ChaosVD/Source/ChaosVD/Private/Utils/ChaosVDUserInterfaceUtils.cpp b/Engine/Plugins/ChaosVD/Source/ChaosVD/Private/Utils/ChaosVDUserInterfaceUtils.cpp
--- a/Engine/Plugins/ChaosVD/Source/ChaosVD/Private/Utils/ChaosVDUserInterfaceUtils.cpp
+++ b/Engine/Plugins/ChaosVD/Source/ChaosVD/Private/Utils/ChaosVDUserInterfaceUtils.cpp
@@ -6,6 +6,7 @@
 #include "ToolMenu.h"
 #include "ToolMenuEntry.h"
 #include "ToolMenuSection.h"
+#include "Settings/ChaosVDCoreSettings.h"
 #include "Widgets/SBoxPanel.h"
 #include "Widgets/Text/STextBlock.h"
 
@@ -37,6 +38,7 @@ TSharedRef<IDetailsView> Chaos::VisualDebugger::Utils::MakeObjectDetailsViewForM
 	DetailsViewArgs.bShowObjectLabel = false;
 	DetailsViewArgs.bCustomNameAreaLocation = true;
 	DetailsViewArgs.ColumnWidth = 0.45f;
+	DetailsViewArgs.bShowModifiedPropertiesOption = true;
 
 	return PropertyEditorModule.CreateDetailView(DetailsViewArgs);
 }
@@ -67,25 +69,49 @@ void Chaos::VisualDebugger::Utils::CreateMenuEntryForObject(UToolMenu* Menu, UOb
 	FToolMenuEntry MenuEntry = FToolMenuEntry::InitWidget(Object->GetFName(), DetailsView, FText::GetEmpty());
 	Menu->AddMenuEntry(NAME_None, MenuEntry);
 
-	if (EnumHasAnyFlags(MenuEntryOptions, EChaosVDSaveSettingsOptions::ShowSaveButton))
+	if (EnumHasAnyFlags(MenuEntryOptions, EChaosVDSaveSettingsOptions::ShowSaveButton | EChaosVDSaveSettingsOptions::ShowResetButton))
 	{
 		Menu->AddMenuEntry(NAME_None,FToolMenuEntry::InitSeparator(NAME_None));
+
+		if (EnumHasAnyFlags(MenuEntryOptions, EChaosVDSaveSettingsOptions::ShowSaveButton))
+		{
+			const FString SaveObjectButtonName = Object->GetName() + TEXT("SaveButton");
+			
+			FToolMenuEntry SaveMenuEntry = FToolMenuEntry::InitMenuEntry(FName(SaveObjectButtonName),
+				NSLOCTEXT("ChaosVisualDebugger","CreateMenuEntryForObjectSaveButtonLabel", "Save Settings"),
+				NSLOCTEXT("ChaosVisualDebugger","CreateMenuEntryForObjectSaveButtonToolTip", "Saves the current settings into the Editor's configuration file"),
+				FSlateIcon(FAppStyle::Get().GetStyleSetName(), TEXT("LevelEditor.Save")),
+				FUIAction(
+					FExecuteAction::CreateLambda([ObjectWeakPtr = TWeakObjectPtr(Object)]()
+					{
+						if (UObject* Object = ObjectWeakPtr.Get())
+						{
+							constexpr bool bAllowCopyToDefaultObject = false;
+							Object->SaveConfig(CPF_Config,nullptr, GConfig, bAllowCopyToDefaultObject);
+						}
+					})));
+
+			Menu->AddMenuEntry(NAME_None, SaveMenuEntry);
+		}
 		
-		const FString SaveObjectButtonName = Object->GetName() + TEXT("SaveButton");
-		
-		FToolMenuEntry SaveMenuEntry = FToolMenuEntry::InitMenuEntry(FName(SaveObjectButtonName),
-			NSLOCTEXT("ChaosVisualDebugger","CreateMenuEntryForObjectSaveButtonLabel", "Save Settings"),
-			NSLOCTEXT("ChaosVisualDebugger","CreateMenuEntryForObjectSaveButtonToolTip", "Saves the current settings into the Editor's confiuration file"),
-			FSlateIcon(FAppStyle::Get().GetStyleSetName(), TEXT("LevelEditor.Save")),
-			FUIAction(
+		if (EnumHasAnyFlags(MenuEntryOptions, EChaosVDSaveSettingsOptions::ShowResetButton))
+		{
+			const FString ResetObjectButtonName = Object->GetName() + TEXT("Reset");
+
+			FToolMenuEntry ResetMenuEntry = FToolMenuEntry::InitMenuEntry(FName(ResetObjectButtonName),
+				NSLOCTEXT("ChaosVisualDebugger","CreateMenuEntryForObjectResetButtonLabel", "Reset to defaults"),
+				NSLOCTEXT("ChaosVisualDebugger","CreateMenuEntryForObjectResetButtonToolTip", "Reset this settings section to its defaults values and save it to the Editor's configuration file"),
+				FSlateIcon(FAppStyle::Get().GetStyleSetName(), TEXT("PropertyWindow.DiffersFromDefault")),
+				FUIAction(
 				FExecuteAction::CreateLambda([ObjectWeakPtr = TWeakObjectPtr(Object)]()
 				{
-					if (UObject* Object = ObjectWeakPtr.Get())
+					if (UObject* Object =ObjectWeakPtr.Get())
 					{
-						Object->SaveConfig();
+						FChaosVDSettingsManager::Get().ResetSettings(Object->GetClass());
 					}
 				})));
-			
-		Menu->AddMenuEntry(NAME_None, SaveMenuEntry);
+
+			Menu->AddMenuEntry(NAME_None, ResetMenuEntry);
+		}
 	}
 }
diff --git a/Engine/Plugins/ChaosVD/Source/ChaosVD/Private/Utils/ChaosVDUserInterfaceUtils.h b/Engine/Plugins/ChaosVD/Source/ChaosVD/Private/Utils/ChaosVDUserInterfaceUtils.h
--- a/Engine/Plugins/ChaosVD/Source/ChaosVD/Private/Utils/ChaosVDUserInterfaceUtils.h
+++ b/Engine/Plugins/ChaosVD/Source/ChaosVD/Private/Utils/ChaosVDUserInterfaceUtils.h
@@ -1,6 +1,7 @@
 // Copyright Epic Games, Inc. All Rights Reserved.
 
 #pragma once
+#include "ChaosVDSettingsManager.h"
 #include "IStructureDetailsView.h"
 
 class UToolMenu;
@@ -15,16 +16,17 @@ namespace Chaos::VisualDebugger::Utils
 	enum class EChaosVDSaveSettingsOptions
 	{
 		None = 0,
-		ShowSaveButton = 1 << 0
+		ShowSaveButton = 1 << 0,
+		ShowResetButton = 1 << 1
 	};
 	ENUM_CLASS_FLAGS(EChaosVDSaveSettingsOptions)
 
 	void CreateMenuEntryForObject(UToolMenu* Menu, UObject* Object, EChaosVDSaveSettingsOptions MenuEntryOptions = EChaosVDSaveSettingsOptions::None);
 
-	template <typename TObject>
-	void CreateMenuEntryForDefaultObject(UToolMenu* Menu, EChaosVDSaveSettingsOptions MenuEntryOptions = EChaosVDSaveSettingsOptions::None)
+	template <typename Object>
+	void CreateMenuEntryForSettingsObject(UToolMenu* Menu, EChaosVDSaveSettingsOptions MenuEntryOptions = EChaosVDSaveSettingsOptions::None)
 	{
-		CreateMenuEntryForObject(Menu, GetMutableDefault<TObject>(), MenuEntryOptions);
+		CreateMenuEntryForObject(Menu, FChaosVDSettingsManager::Get().GetSettingsObject<Object>(), MenuEntryOptions);
 	}
 
 	template <typename TStruct>
diff --git a/Engine/Plugins/ChaosVD/Source/ChaosVD/Private/Visualizers/ChaosVDCharacterGroundConstraintsDataComponentVisualizer.cpp b/Engine/Plugins/ChaosVD/Source/ChaosVD/Private/Visualizers/ChaosVDCharacterGroundConstraintsDataComponentVisualizer.cpp
--- a/Engine/Plugins/ChaosVD/Source/ChaosVD/Private/Visualizers/ChaosVDCharacterGroundConstraintsDataComponentVisualizer.cpp
+++ b/Engine/Plugins/ChaosVD/Source/ChaosVD/Private/Visualizers/ChaosVDCharacterGroundConstraintsDataComponentVisualizer.cpp
@@ -3,6 +3,7 @@
 #include "Visualizers/ChaosVDCharacterGroundConstraintsDataComponentVisualizer.h"
 
 #include "ChaosVDScene.h"
+#include "ChaosVDSettingsManager.h"
 #include "ChaosVDTabsIDs.h"
 #include "EditorViewportClient.h"
 #include "SceneManagement.h"
@@ -82,7 +83,7 @@ void FChaosVDCharacterGroundConstraintDataComponentVisualizer::DrawVisualization
 	VisualizationContext.SolverInfoActor = SolverInfoActor;
 
 	VisualizationContext.VisualizationFlags = static_cast<uint32>(UChaosVDCharacterConstraintsVisualizationSettings::GetDataVisualizationFlags());
-	VisualizationContext.DebugDrawSettings = GetDefault<UChaosVDCharacterConstraintsVisualizationSettings>();
+	VisualizationContext.DebugDrawSettings = FChaosVDSettingsManager::Get().GetSettingsObject<UChaosVDCharacterConstraintsVisualizationSettings>();
 
 	if (!VisualizationContext.IsVisualizationFlagEnabled(EChaosVDCharacterGroundConstraintDataVisualizationFlags::EnableDraw))
 	{
diff --git a/Engine/Plugins/ChaosVD/Source/ChaosVD/Private/Visualizers/ChaosVDComponentVisualizerBase.h b/Engine/Plugins/ChaosVD/Source/ChaosVD/Private/Visualizers/ChaosVDComponentVisualizerBase.h
--- a/Engine/Plugins/ChaosVD/Source/ChaosVD/Private/Visualizers/ChaosVDComponentVisualizerBase.h
+++ b/Engine/Plugins/ChaosVD/Source/ChaosVD/Private/Visualizers/ChaosVDComponentVisualizerBase.h
@@ -94,7 +94,7 @@ void FChaosVDComponentVisualizerBase::CreateGenericVisualizerMenu(FName MenuToEx
 						   false, FlagsMenuIcon);
 
 		using namespace Chaos::VisualDebugger::Utils;
-		Section.AddSubMenu(FName(InSettingsMenuLabel.ToString()), InSettingsMenuLabel, InSettingsMenuTooltip, FNewToolMenuDelegate::CreateStatic(&CreateMenuEntryForDefaultObject<ObjectSettingsType>, EChaosVDSaveSettingsOptions::ShowSaveButton),
+		Section.AddSubMenu(FName(InSettingsMenuLabel.ToString()), InSettingsMenuLabel, InSettingsMenuTooltip, FNewToolMenuDelegate::CreateStatic(&CreateMenuEntryForSettingsObject<ObjectSettingsType>, EChaosVDSaveSettingsOptions::ShowResetButton),
 						   false, FSlateIcon(FAppStyle::Get().GetStyleSetName(), TEXT("Icons.Toolbar.Settings")));
 	}
 }
diff --git a/Engine/Plugins/ChaosVD/Source/ChaosVD/Private/Visualizers/ChaosVDJointConstraintsDataComponentVisualizer.cpp b/Engine/Plugins/ChaosVD/Source/ChaosVD/Private/Visualizers/ChaosVDJointConstraintsDataComponentVisualizer.cpp
--- a/Engine/Plugins/ChaosVD/Source/ChaosVD/Private/Visualizers/ChaosVDJointConstraintsDataComponentVisualizer.cpp
+++ b/Engine/Plugins/ChaosVD/Source/ChaosVD/Private/Visualizers/ChaosVDJointConstraintsDataComponentVisualizer.cpp
@@ -4,6 +4,7 @@
 
 #include "ChaosVDEngine.h"
 #include "ChaosVDScene.h"
+#include "ChaosVDSettingsManager.h"
 #include "ChaosVDStyle.h"
 #include "ChaosVDTabsIDs.h"
 #include "EditorViewportClient.h"
@@ -160,7 +161,7 @@ void FChaosVDJointConstraintsDataComponentVisualizer::DrawVisualization(const UA
 	VisualizationContext.SpaceTransform = SolverInfoActor->GetSimulationTransform();
 	VisualizationContext.SolverInfoActor = SolverInfoActor;
 
-	if (const UChaosVDJointConstraintsVisualizationSettings* EditorSettings = GetDefault<UChaosVDJointConstraintsVisualizationSettings>())
+	if (const UChaosVDJointConstraintsVisualizationSettings* EditorSettings = FChaosVDSettingsManager::Get().GetSettingsObject<UChaosVDJointConstraintsVisualizationSettings>())
 	{
 		VisualizationContext.VisualizationFlags = static_cast<uint32>(UChaosVDJointConstraintsVisualizationSettings::GetDataVisualizationFlags());
 		VisualizationContext.bShowDebugText = EditorSettings->bShowDebugText;
diff --git a/Engine/Plugins/ChaosVD/Source/ChaosVD/Private/Visualizers/ChaosVDParticleDataComponentVisualizer.cpp b/Engine/Plugins/ChaosVD/Source/ChaosVD/Private/Visualizers/ChaosVDParticleDataComponentVisualizer.cpp
--- a/Engine/Plugins/ChaosVD/Source/ChaosVD/Private/Visualizers/ChaosVDParticleDataComponentVisualizer.cpp
+++ b/Engine/Plugins/ChaosVD/Source/ChaosVD/Private/Visualizers/ChaosVDParticleDataComponentVisualizer.cpp
@@ -5,6 +5,7 @@
 #include "Actors/ChaosVDSolverInfoActor.h"
 #include "ChaosVDParticleActor.h"
 #include "ChaosVDScene.h"
+#include "ChaosVDSettingsManager.h"
 #include "ChaosVDTabsIDs.h"
 #include "Components/ChaosVDParticleDataComponent.h"
 #include "DataWrappers/ChaosVDParticleDataWrapper.h"
@@ -88,21 +89,16 @@ void FChaosVDParticleDataComponentVisualizer::RegisterVisualizerMenus()
 			}
 		});
 		
-		FNewToolMenuDelegate ParticleDataVisualizationSettingsMenuBuilder = FNewToolMenuDelegate::CreateLambda([](UToolMenu* Menu)
-		{
-			using namespace Chaos::VisualDebugger::Utils;
-			CreateMenuEntryForDefaultObject<UChaosVDParticleVisualizationDebugDrawSettings>(Menu, EChaosVDSaveSettingsOptions::ShowSaveButton);
-		});
-
-		FNewToolMenuDelegate ParticleColorizationMenuBuilder = FNewToolMenuDelegate::CreateLambda([](UToolMenu* Menu)
-		{
-			using namespace Chaos::VisualDebugger::Utils;
-			CreateMenuEntryForDefaultObject<UChaosVDParticleVisualizationColorSettings>(Menu, EChaosVDSaveSettingsOptions::ShowSaveButton);
-		});
+		using namespace Chaos::VisualDebugger::Utils;
+		
+		FNewToolMenuDelegate GeometryVisualizationSettingsMenuBuilder = FNewToolMenuDelegate::CreateStatic(&CreateMenuEntryForSettingsObject<UChaosVDParticleVisualizationSettings>, EChaosVDSaveSettingsOptions::ShowResetButton);
+		FNewToolMenuDelegate ParticleDataVisualizationSettingsMenuBuilder = FNewToolMenuDelegate::CreateStatic(&CreateMenuEntryForSettingsObject<UChaosVDParticleVisualizationDebugDrawSettings>, EChaosVDSaveSettingsOptions::ShowResetButton);
+		FNewToolMenuDelegate ParticleColorizationMenuBuilder = FNewToolMenuDelegate::CreateStatic(&CreateMenuEntryForSettingsObject<UChaosVDParticleVisualizationColorSettings>, EChaosVDSaveSettingsOptions::ShowResetButton);
 
 		constexpr bool bOpenSubMenuOnClick = false;
 		
 		Section.AddSubMenu(TEXT("GeometryVisualizationFlags"), LOCTEXT("GeometryVisualizationFlagsMenuLabel", "Geometry Flags"), LOCTEXT("GeometryVisualizationFlagsMenuToolTip", "Set of flags to enable/disable visibility of specific types of geometry/particles"), GeometryVisualizationFlagsMenuBuilder, bOpenSubMenuOnClick, FSlateIcon(FAppStyle::Get().GetStyleSetName(), TEXT("ShowFlagsMenu.StaticMeshes")));
+		Section.AddSubMenu(TEXT("GeometryVisualizationSettings"), LOCTEXT("GeometryVisualizationSettingsMenuLabel", "Geometry Visualization Settings"), LOCTEXT("GeometryVisualizationSettingsMenuToolTip", "Options to control how particle data is debug geometry is visualized"), GeometryVisualizationSettingsMenuBuilder, bOpenSubMenuOnClick, FSlateIcon(FAppStyle::Get().GetStyleSetName(), TEXT("Icons.Toolbar.Settings")));
 		Section.AddSubMenu(TEXT("ParticleDataVisualizationFlags"), LOCTEXT("ParticleDataVisualizationFlagsMenuLabel", "Particle Data Flags"), LOCTEXT("ParticleDataVisualizationFlagsMenuToolTip", "Set of flags to enable/disable visualization of specific particle data as debug draw"), ParticleDataVisualizationFlagsMenuBuilder, bOpenSubMenuOnClick, FSlateIcon(FAppStyle::Get().GetStyleSetName(), TEXT("StaticMeshEditor.SetDrawAdditionalData")));
 		Section.AddSubMenu(TEXT("ParticleDataVisualizationSettings"), LOCTEXT("ParticleDataVisualizationSettingsMenuLabel", "Particle Data Visualization Settings"), LOCTEXT("ParticleDataVisualizationSettingsMenuToolTip", "Options to control how particle data is debug drawn"), ParticleDataVisualizationSettingsMenuBuilder, bOpenSubMenuOnClick, FSlateIcon(FAppStyle::Get().GetStyleSetName(), TEXT("Icons.Toolbar.Settings")));
 		Section.AddSubMenu(TEXT("ParticleColorizationFlags"), LOCTEXT("ParticleColorizationOptionsMenuLabel", "Particle Colorization"), LOCTEXT("Particle ColorizationMenuToolTip", "Changes what colors are used to draw the particles and its data"), ParticleColorizationMenuBuilder, bOpenSubMenuOnClick, FSlateIcon(FAppStyle::Get().GetStyleSetName(), TEXT("ColorPicker.ColorThemes")));
@@ -111,7 +107,7 @@ void FChaosVDParticleDataComponentVisualizer::RegisterVisualizerMenus()
 
 void FChaosVDParticleDataComponentVisualizer::DrawVisualization(const UActorComponent* Component, const FSceneView* View, FPrimitiveDrawInterface* PDI)
 {
-	const UChaosVDParticleVisualizationDebugDrawSettings* VisualizationSettings = GetDefault<UChaosVDParticleVisualizationDebugDrawSettings>();
+	const UChaosVDParticleVisualizationDebugDrawSettings* VisualizationSettings = FChaosVDSettingsManager::Get().GetSettingsObject<UChaosVDParticleVisualizationDebugDrawSettings>();
 	if (!VisualizationSettings)
 	{
 		return;
diff --git a/Engine/Plugins/ChaosVD/Source/ChaosVD/Private/Visualizers/ChaosVDSceneQueryDataComponentVisualizer.cpp b/Engine/Plugins/ChaosVD/Source/ChaosVD/Private/Visualizers/ChaosVDSceneQueryDataComponentVisualizer.cpp
--- a/Engine/Plugins/ChaosVD/Source/ChaosVD/Private/Visualizers/ChaosVDSceneQueryDataComponentVisualizer.cpp
+++ b/Engine/Plugins/ChaosVD/Source/ChaosVD/Private/Visualizers/ChaosVDSceneQueryDataComponentVisualizer.cpp
@@ -6,6 +6,7 @@
 #include "Components/ChaosVDSceneQueryDataComponent.h"
 #include "ChaosVDGeometryBuilder.h"
 #include "ChaosVDScene.h"
+#include "ChaosVDSettingsManager.h"
 #include "ChaosVDStyle.h"
 #include "ChaosVDTabsIDs.h"
 #include "EditorViewportClient.h"
@@ -94,7 +95,7 @@ void FChaosVDSceneQueryDataComponentVisualizer::DrawVisualization(const UActorCo
 	VisualizationContext.SolverDataSelectionObject = SolverDataSelectionObject;
 
 	VisualizationContext.VisualizationFlags = static_cast<uint32>(UChaosVDSceneQueriesVisualizationSettings::GetDataVisualizationFlags());
-	VisualizationContext.DebugDrawSettings = GetDefault<UChaosVDSceneQueriesVisualizationSettings>();
+	VisualizationContext.DebugDrawSettings = FChaosVDSettingsManager::Get().GetSettingsObject<UChaosVDSceneQueriesVisualizationSettings>();
 
 	if (EnumHasAnyFlags(EChaosVDSceneQueryVisualizationFlags::EnableDraw, static_cast<EChaosVDSceneQueryVisualizationFlags>(VisualizationContext.VisualizationFlags)))
 	{
diff --git a/Engine/Plugins/ChaosVD/Source/ChaosVD/Private/Visualizers/ChaosVDSolverCollisionDataComponentVisualizer.cpp b/Engine/Plugins/ChaosVD/Source/ChaosVD/Private/Visualizers/ChaosVDSolverCollisionDataComponentVisualizer.cpp
--- a/Engine/Plugins/ChaosVD/Source/ChaosVD/Private/Visualizers/ChaosVDSolverCollisionDataComponentVisualizer.cpp
+++ b/Engine/Plugins/ChaosVD/Source/ChaosVD/Private/Visualizers/ChaosVDSolverCollisionDataComponentVisualizer.cpp
@@ -5,6 +5,7 @@
 #include "ChaosVDComponentVisualizerBase.h"
 #include "ChaosVDParticleActor.h"
 #include "ChaosVDScene.h"
+#include "ChaosVDSettingsManager.h"
 #include "ChaosVDTabsIDs.h"
 #include "EditorViewportClient.h"
 #include "Actors/ChaosVDSolverInfoActor.h"
@@ -75,7 +76,7 @@ void FChaosVDSolverCollisionDataComponentVisualizer::DrawVisualization(const UAc
 	VisualizationContext.CVDScene = SolverInfoContainer->GetScene();
 	VisualizationContext.SpaceTransform = SolverInfoContainer->GetSimulationTransform();
 	VisualizationContext.VisualizationFlags = static_cast<uint32>(UChaosVDCollisionDataVisualizationSettings::GetDataVisualizationFlags());
-	VisualizationContext.DebugDrawSettings = GetDefault<UChaosVDCollisionDataVisualizationSettings>();
+	VisualizationContext.DebugDrawSettings = FChaosVDSettingsManager::Get().GetSettingsObject<UChaosVDCollisionDataVisualizationSettings>();
 
 	if (!EnumHasAnyFlags(EChaosVDCollisionVisualizationFlags::EnableDraw, static_cast<EChaosVDCollisionVisualizationFlags>(VisualizationContext.VisualizationFlags)))
 	{
diff --git a/Engine/Plugins/ChaosVD/Source/ChaosVD/Public/ChaosVDModule.h b/Engine/Plugins/ChaosVD/Source/ChaosVD/Public/ChaosVDModule.h
--- a/Engine/Plugins/ChaosVD/Source/ChaosVD/Public/ChaosVDModule.h
+++ b/Engine/Plugins/ChaosVD/Source/ChaosVD/Public/ChaosVDModule.h
@@ -12,6 +12,7 @@ class SDockTab;
 class FSpawnTabArgs;
 struct FGuid;
 
+/** Main module class for the Chaos Visual Debugger editor */
 /** Main module class for the Chaos Visual Debugger editor */
 DECLARE_LOG_CATEGORY_EXTERN(LogChaosVDEditor, Log, Log);
 class FChaosVDModule : public IModuleInterface
@@ -59,4 +60,5 @@ private:
 	static FString ChaosVisualDebuggerProgramName;
 
 	bool bIsShuttingDown = false;
+
 };
