[ChaosVD] Initial implementation to add support to record and visualize game thread acceleration structures

This CL adds only AABB Tree are support

In addition as part of this change

- Added a way to generate a struct view that can be fed to the vanilla details panel for Solver Data structs. CVD solver data usually consist of a main struct that can contain several substructs in containers. When we select something in the viewport, we usually want to see the data in the parent struct and then some specific items in the containers.
Until now for these cases we created dedicated details panels with added functionality to interact with the data. But there are cases where we just want to see the struct data. Now for these cases we can just use this struct view API in the CVD Selection handle system to use the normal CVD's details panel.
- Added a way to allow CVD's main details panel to ignore struct types (needed to ignore structs that have their own data inspector widgets).
-Fixed an issue with CVD's Details panel where if you try to inspect a struct after uobject was inspected, the uobject details view was not being properly cleared out and hidden

diff --git a/Engine/Plugins/ChaosVD/Source/ChaosVD/Private/Actors/ChaosVDSolverInfoActor.cpp b/Engine/Plugins/ChaosVD/Source/ChaosVD/Private/Actors/ChaosVDSolverInfoActor.cpp
--- a/Engine/Plugins/ChaosVD/Source/ChaosVD/Private/Actors/ChaosVDSolverInfoActor.cpp
+++ b/Engine/Plugins/ChaosVD/Source/ChaosVD/Private/Actors/ChaosVDSolverInfoActor.cpp
@@ -7,6 +7,7 @@
 #include "ChaosVDScene.h"
 #include "ChaosVDSettingsManager.h"
 #include "EditorActorFolders.h"
+#include "Components/ChaosVDGTAccelerationStructuresDataComponent.h"
 #include "Components/ChaosVDParticleDataComponent.h"
 #include "Components/ChaosVDSolverCharacterGroundConstraintDataComponent.h"
 #include "Components/ChaosVDSolverCollisionDataComponent.h"
@@ -24,6 +25,7 @@ AChaosVDSolverInfoActor::AChaosVDSolverInfoActor(const FObjectInitializer& Objec
 	ParticleDataComponent = CreateDefaultSubobject<UChaosVDParticleDataComponent>(TEXT("ParticleCollisionDataComponent"));
 	JointsDataComponent = CreateDefaultSubobject<UChaosVDSolverJointConstraintDataComponent>(TEXT("JointDataComponent"));
 	CharacterGroundConstraintDataComponent = CreateDefaultSubobject<UChaosVDSolverCharacterGroundConstraintDataComponent>(TEXT("CharacterGroundConstraintDataComponent"));
+	GTAccelerationStructuresDataComponent = CreateDefaultSubobject<UChaosVDGTAccelerationStructuresDataComponent>(TEXT("GTAccelerationStructuresDataComponent"));
 	bIsServer = false;
 
 	if (UChaosVDParticleVisualizationSettings* ParticleVisualizationSettings = FChaosVDSettingsManager::Get().GetSettingsObject<UChaosVDParticleVisualizationSettings>())
diff --git a/Engine/Plugins/ChaosVD/Source/ChaosVD/Private/Actors/ChaosVDSolverInfoActor.h b/Engine/Plugins/ChaosVD/Source/ChaosVD/Private/Actors/ChaosVDSolverInfoActor.h
--- a/Engine/Plugins/ChaosVD/Source/ChaosVD/Private/Actors/ChaosVDSolverInfoActor.h
+++ b/Engine/Plugins/ChaosVD/Source/ChaosVD/Private/Actors/ChaosVDSolverInfoActor.h
@@ -8,6 +8,7 @@
 #include "GameFramework/Actor.h"
 #include "ChaosVDSolverInfoActor.generated.h"
 
+class UChaosVDGTAccelerationStructuresDataComponent;
 class UChaosVDSolverCharacterGroundConstraintDataComponent;
 class UChaosVDSolverJointConstraintDataComponent;
 struct FChaosVDParticleDataWrapper;
@@ -44,6 +45,7 @@ public:
 	UChaosVDParticleDataComponent* GetParticleDataComponent() { return ParticleDataComponent; }
 	UChaosVDSolverJointConstraintDataComponent* GetJointsDataComponent() { return JointsDataComponent; }
 	UChaosVDSolverCharacterGroundConstraintDataComponent* GetCharacterGroundConstraintDataComponent() { return CharacterGroundConstraintDataComponent; }
+	UChaosVDGTAccelerationStructuresDataComponent* GetGTAccelerationStructuresDataComponent() { return GTAccelerationStructuresDataComponent; }
 
 	void RegisterParticleActor(int32 ParticleID, AChaosVDParticleActor* ParticleActor);
 
@@ -113,6 +115,9 @@ protected:
 
 	UPROPERTY()
 	TObjectPtr<UChaosVDSolverCharacterGroundConstraintDataComponent> CharacterGroundConstraintDataComponent;
+
+	UPROPERTY()
+	TObjectPtr<UChaosVDGTAccelerationStructuresDataComponent> GTAccelerationStructuresDataComponent;
 };
 
 template <typename TCallback>
diff --git a/Engine/Plugins/ChaosVD/Source/ChaosVD/Private/ChaosVDModule.cpp b/Engine/Plugins/ChaosVD/Source/ChaosVD/Private/ChaosVDModule.cpp
--- a/Engine/Plugins/ChaosVD/Source/ChaosVD/Private/ChaosVDModule.cpp
+++ b/Engine/Plugins/ChaosVD/Source/ChaosVD/Private/ChaosVDModule.cpp
@@ -20,6 +20,7 @@
 #include "Widgets/SChaosVDMainTab.h"
 #include "WorkspaceMenuStructure.h"
 #include "WorkspaceMenuStructureModule.h"
+#include "DetailsCustomizations/ChaosVDSelectionMultipleViewCustomization.h"
 
 #define LOCTEXT_NAMESPACE "ChaosVisualDebugger"
 
@@ -100,6 +101,7 @@ void FChaosVDModule::RegisterClassesCustomDetails() const
 	PropertyModule.RegisterCustomClassLayout("ChaosVDStaticMeshComponent", FOnGetDetailCustomizationInstance::CreateStatic(&FChaosVDGeometryComponentCustomization::MakeInstance));
 	PropertyModule.RegisterCustomClassLayout("ChaosVDQueryVisitStep", FOnGetDetailCustomizationInstance::CreateStatic(&FChaosVDQueryVisitDataCustomization::MakeInstance));
 	PropertyModule.RegisterCustomClassLayout("ChaosVDQueryDataWrapper", FOnGetDetailCustomizationInstance::CreateStatic(&FChaosVDQueryDataWrapperCustomization::MakeInstance));
+	PropertyModule.RegisterCustomClassLayout("ChaosVDSelectionMultipleView", FOnGetDetailCustomizationInstance::CreateStatic(&FChaosVDSelectionMultipleViewCustomization::MakeInstance));
 
 	//TODO: Rename FChaosVDParticleDataWrapperCustomization to something generic as currently works with any type that wants to hide properties of type FChaosVDWrapperDataBase with invalid data.
 	// Or another option is create a new custom layout intended to be generic from the get go
diff --git a/Engine/Plugins/ChaosVD/Source/ChaosVD/Private/ChaosVDObjectDetailsTab.cpp b/Engine/Plugins/ChaosVD/Source/ChaosVD/Private/ChaosVDObjectDetailsTab.cpp
--- a/Engine/Plugins/ChaosVD/Source/ChaosVD/Private/ChaosVDObjectDetailsTab.cpp
+++ b/Engine/Plugins/ChaosVD/Source/ChaosVD/Private/ChaosVDObjectDetailsTab.cpp
@@ -22,7 +22,10 @@
 
 #define LOCTEXT_NAMESPACE "ChaosVisualDebugger"
 
-class SSubobjectEditor;
+void FChaosVDObjectDetailsTab::AddUnsupportedStruct(const UStruct* Struct)
+{
+	UnsupportedStructs.Add(Struct);
+}
 
 TSharedRef<SDockTab> FChaosVDObjectDetailsTab::HandleTabSpawnRequest(const FSpawnTabArgs& Args)
 {
@@ -33,10 +36,20 @@ TSharedRef<SDockTab> FChaosVDObjectDetailsTab::HandleTabSpawnRequest(const FSpaw
 	.Label(LOCTEXT("DetailsPanel", "Details"))
 	.ToolTipText(LOCTEXT("DetailsPanelToolTip", "See the details of the selected object"));
 
+	// The following types have their own data inspectors, we should not open them in the details pannel
+	AddUnsupportedStruct(FChaosVDConstraintDataWrapperBase::StaticStruct());
+	AddUnsupportedStruct(FChaosVDQueryDataWrapper::StaticStruct());
+	AddUnsupportedStruct(FChaosVDParticlePairMidPhase::StaticStruct());
+
 	if (const TSharedPtr<FChaosVDScene> ScenePtr = GetChaosVDScene().Pin())
 	{
 		RegisterSelectionSetObject(ScenePtr->GetElementSelectionSet());
 
+		if (TSharedPtr<FChaosVDSolverDataSelection> SolverDataSelectionObject = ScenePtr->GetSolverDataSelectionObject().Pin())
+		{
+			SolverDataSelectionObject->GetDataSelectionChangedDelegate().AddSP(this, &FChaosVDObjectDetailsTab::HandleSolverDataSelectionChange);
+		}
+
 		DetailsPanelTab->SetContent
 		(
 			SNew(SVerticalBox)
@@ -72,14 +85,20 @@ TSharedRef<SDockTab> FChaosVDObjectDetailsTab::HandleTabSpawnRequest(const FSpaw
 void FChaosVDObjectDetailsTab::HandleTabClosed(TSharedRef<SDockTab> InTabClosed)
 {
 	FChaosVDTabSpawnerBase::HandleTabClosed(InTabClosed);
+
+	if (const TSharedPtr<FChaosVDScene> ScenePtr = GetChaosVDScene().Pin())
+	{
+		if (TSharedPtr<FChaosVDSolverDataSelection> SolverDataSelectionObject = ScenePtr->GetSolverDataSelectionObject().Pin())
+		{
+			SolverDataSelectionObject->GetDataSelectionChangedDelegate().RemoveAll(this);
+		}
+	}
 	
 	DetailsPanelView.Reset();
 }
 
-void FChaosVDObjectDetailsTab::HandlePostSelectionChange(const UTypedElementSelectionSet* ChangedSelectionSet)
+void FChaosVDObjectDetailsTab::HandleActorsSelection(TArrayView<AActor*> SelectedActors)
 {
-	TArray<AActor*> SelectedActors = ChangedSelectionSet->GetSelectedObjects<AActor>();
-
 	if (SelectedActors.Num() > 0)
 	{
 		// We don't support multi selection yet
@@ -98,6 +117,32 @@ void FChaosVDObjectDetailsTab::HandlePostSelectionChange(const UTypedElementSele
 	}
 }
 
+void FChaosVDObjectDetailsTab::HandlePostSelectionChange(const UTypedElementSelectionSet* ChangedSelectionSet)
+{
+	TArray<AActor*> SelectedActors = ChangedSelectionSet->GetSelectedObjects<AActor>();
+
+	HandleActorsSelection(SelectedActors);
+}
+
+void FChaosVDObjectDetailsTab::HandleSolverDataSelectionChange(const TSharedPtr<FChaosVDSolverDataSelectionHandle>& SelectionHandle)
+{
+	TSharedPtr<FStructOnScope> StructOnScope = SelectionHandle->GetDataAsStructScope();
+	if (!StructOnScope || !IsSupportedStruct(StructOnScope->GetStructPtr()))
+	{
+		DetailsPanelView->SetSelectedStruct(nullptr);
+		return;
+	}
+
+	HandleActorsSelection(TArrayView<AActor*>());
+
+	DetailsPanelView->SetSelectedStruct(SelectionHandle->GetCustomDataReadOnlyStructViewForDetails());
+}
+
+bool FChaosVDObjectDetailsTab::IsSupportedStruct(const TWeakObjectPtr<const UStruct>& InWeakStructPtr)
+{
+	return !UnsupportedStructs.Contains(InWeakStructPtr);
+}
+
 EVisibility FChaosVDObjectDetailsTab::GetCollisionDataButtonVisibility() const
 {
 	EVisibility DesiredVisibility = EVisibility::Collapsed;
diff --git a/Engine/Plugins/ChaosVD/Source/ChaosVD/Private/ChaosVDObjectDetailsTab.h b/Engine/Plugins/ChaosVD/Source/ChaosVD/Private/ChaosVDObjectDetailsTab.h
--- a/Engine/Plugins/ChaosVD/Source/ChaosVD/Private/ChaosVDObjectDetailsTab.h
+++ b/Engine/Plugins/ChaosVD/Source/ChaosVD/Private/ChaosVDObjectDetailsTab.h
@@ -7,6 +7,7 @@
 #include "ChaosVDTabSpawnerBase.h"
 #include "Delegates/IDelegateInstance.h"
 
+struct FChaosVDSolverDataSelectionHandle;
 struct FChaosVDParticleDebugData;
 
 class AActor;
@@ -27,6 +28,8 @@ public:
 	
 	TSharedPtr<SChaosVDDetailsView> GetDetailsPanel() { return DetailsPanelView; }
 
+	void AddUnsupportedStruct(const UStruct* Struct);
+
 protected:
 
 	virtual TSharedRef<SDockTab> HandleTabSpawnRequest(const FSpawnTabArgs& Args) override;
@@ -34,6 +37,12 @@ protected:
 
 	virtual void HandlePostSelectionChange(const UTypedElementSelectionSet* ChangedSelectionSet) override;
 
+	void HandleActorsSelection(TArrayView<AActor*> SelectedActors);
+
+	void HandleSolverDataSelectionChange(const TSharedPtr<FChaosVDSolverDataSelectionHandle>& SelectionHandle);
+
+	bool IsSupportedStruct(const TWeakObjectPtr<const UStruct>& InWeakStructPtr);
+
 	EVisibility GetCollisionDataButtonVisibility() const;
 	bool GetCollisionDataButtonEnabled() const;
 
@@ -44,4 +53,6 @@ protected:
 	FDelegateHandle SelectionDelegateHandle;
 	TSharedPtr<SChaosVDDetailsView> DetailsPanelView;
 	TWeakObjectPtr<UObject> CurrentSelectedObject = nullptr;
+
+	TSet<TWeakObjectPtr<const UStruct>> UnsupportedStructs;
 };
diff --git a/Engine/Plugins/ChaosVD/Source/ChaosVD/Private/ChaosVDScene.cpp b/Engine/Plugins/ChaosVD/Source/ChaosVD/Private/ChaosVDScene.cpp
--- a/Engine/Plugins/ChaosVD/Source/ChaosVD/Private/ChaosVDScene.cpp
+++ b/Engine/Plugins/ChaosVD/Source/ChaosVD/Private/ChaosVDScene.cpp
@@ -31,6 +31,7 @@
 #include "Settings/ChaosVDCoreSettings.h"
 #include "UObject/Package.h"
 #include "Actors/ChaosVDGeometryContainer.h"
+#include "Components/ChaosVDGTAccelerationStructuresDataComponent.h"
 
 #define LOCTEXT_NAMESPACE "ChaosVisualDebugger"
 
@@ -359,6 +360,20 @@ void FChaosVDScene::HandleEnterNewGameFrame(int32 FrameNumber, const TArray<int3
 			QueryDataComponent->UpdateQueriesFromFrameData(InNewGameFrameData);
 		}
 	}
+
+	for (const TPair<int32, TArray<TSharedPtr<FChaosVDAABBTreeDataWrapper>>>& AABBTreeDataWithSolverID : InNewGameFrameData.RecordedAABBTreesBySolverID)
+	{
+		if (AChaosVDSolverInfoActor** SolverInfoActorPtrPtr =  SolverDataContainerBySolverID.Find(AABBTreeDataWithSolverID.Key))
+		{
+			if (AChaosVDSolverInfoActor* SolverInfoActorPtr = *SolverInfoActorPtrPtr)
+			{
+				if (UChaosVDGTAccelerationStructuresDataComponent* GTAccelerationStructureDataComponent = SolverInfoActorPtr->GetGTAccelerationStructuresDataComponent())
+				{
+					GTAccelerationStructureDataComponent->UpdateAABBTreeData(AABBTreeDataWithSolverID.Value);
+				}
+			}
+		}
+	}
 }
 
 void FChaosVDScene::HandleEnterNewSolverFrame(int32 FrameNumber, const FChaosVDSolverFrameData& InFrameData)
diff --git a/Engine/Plugins/ChaosVD/Source/ChaosVD/Private/ChaosVDSolverDataSelection.h b/Engine/Plugins/ChaosVD/Source/ChaosVD/Private/ChaosVDSolverDataSelection.h
--- a/Engine/Plugins/ChaosVD/Source/ChaosVD/Private/ChaosVDSolverDataSelection.h
+++ b/Engine/Plugins/ChaosVD/Source/ChaosVD/Private/ChaosVDSolverDataSelection.h
@@ -7,12 +7,41 @@
 
 class FChaosVDSolverDataSelection;
 
+/** Base struct type used for any context data we want to add for a selection handle */
 USTRUCT()
 struct FChaosVDSelectionContext
 {
 	GENERATED_BODY()
 };
 
+/** Struct used to create a combined view of multiple structs to be used in a vanilla details panel.
+ * This struct type has a customization that will show each data entry as an individual property
+ */
+USTRUCT()
+struct FChaosVDSelectionMultipleView
+{
+	GENERATED_BODY()
+
+	template<typename StructType>
+	void AddData(StructType* Struct);
+
+protected:
+	TArray<TSharedPtr<FStructOnScope>> DataInstances;
+
+	friend class FChaosVDSelectionMultipleViewCustomization;
+};
+
+template <typename StructType>
+void FChaosVDSelectionMultipleView::AddData(StructType* Struct)
+{
+	if (!Struct)
+	{
+		return;
+	}
+
+	DataInstances.Emplace(MakeShared<FStructOnScope>(StructType::StaticStruct(), reinterpret_cast<uint8*>(Struct)));
+}
+
 struct FChaosVDSolverDataSelectionHandle : public TSharedFromThis<FChaosVDSolverDataSelectionHandle>
 {
 	virtual ~FChaosVDSolverDataSelectionHandle() = default;
@@ -43,6 +72,11 @@ struct FChaosVDSolverDataSelectionHandle : public TSharedFromThis<FChaosVDSolver
 
 	TSharedPtr<FStructOnScope> GetDataAsStructScope() { return SelectedDataStruct; }
 
+	/** Returns a struct on Scope view that can be fed into a CVD details panel -
+	 * Usually used to combine data and context into a single read only struct that can be inspected
+	 */
+	virtual TSharedPtr<FStructOnScope> GetCustomDataReadOnlyStructViewForDetails() { return SelectedDataStruct; }
+
 	bool operator==(const FChaosVDSolverDataSelectionHandle& Other) const
 	{
 		return DataSharedPtr.Get() == Other.DataSharedPtr.Get();
@@ -64,6 +98,7 @@ private:
 	TSharedPtr<void> DataSharedPtr;
 	TSharedPtr<void> SelectedDataContextSharedPtr;
 
+protected:
 	TWeakPtr<FChaosVDSolverDataSelection> Owner;
 };
 
@@ -151,7 +186,7 @@ public:
 
 	void SelectData(const TSharedPtr<FChaosVDSolverDataSelectionHandle>& InSelectionHandle);
 
-	template<typename SolverDataType>
+	template<typename SolverDataType, typename HandleType = FChaosVDSolverDataSelectionHandle>
 	TSharedPtr<FChaosVDSolverDataSelectionHandle> MakeSelectionHandle(const TSharedPtr<SolverDataType>& InSolverData);
 
 	FChaosVDSolverDataSelectionChangedDelegate& GetDataSelectionChangedDelegate() { return SolverDataSelectionChangeDelegate; }
@@ -167,14 +202,16 @@ private:
 	TSharedPtr<FChaosVDSolverDataSelectionHandle> CurrentSelectedSolverDataHandle;
 };
 
-template <typename SolverDataType>
+template <typename SolverDataType, typename HandleType>
 TSharedPtr<FChaosVDSolverDataSelectionHandle> FChaosVDSolverDataSelection::MakeSelectionHandle(const TSharedPtr<SolverDataType>& InSolverData)
 {
-	TSharedPtr<FChaosVDSolverDataSelectionHandle> NewSelectionHandle = MakeShared<FChaosVDSolverDataSelectionHandle>();
+	static_assert(std::is_base_of_v<FChaosVDSolverDataSelectionHandle, HandleType>, "MakeSelectionHandle only supports handles derived from FChaosVDSolverDataSelectionHandle");
+
+	TSharedPtr<HandleType> NewSelectionHandle = MakeShared<HandleType>();
 	NewSelectionHandle->SetHandleData(StaticCastSharedPtr<SolverDataType>(InSolverData));
 	NewSelectionHandle->SetOwner(AsShared());
 
-	return NewSelectionHandle;
+	return StaticCastSharedPtr<HandleType>(NewSelectionHandle);
 }
 
 template <typename SolverDataType>
diff --git a/Engine/Plugins/ChaosVD/Source/ChaosVD/Private/Components/ChaosVDGTAccelerationStructuresDataComponent.cpp b/Engine/Plugins/ChaosVD/Source/ChaosVD/Private/Components/ChaosVDGTAccelerationStructuresDataComponent.cpp
new file mode 100644
--- /dev/null
+++ b/Engine/Plugins/ChaosVD/Source/ChaosVD/Private/Components/ChaosVDGTAccelerationStructuresDataComponent.cpp
@@ -0,0 +1,24 @@
+// Copyright Epic Games, Inc. All Rights Reserved.
+
+#include "ChaosVDGTAccelerationStructuresDataComponent.h"
+
+#include "Algo/Copy.h"
+
+UChaosVDGTAccelerationStructuresDataComponent::UChaosVDGTAccelerationStructuresDataComponent()
+{
+	PrimaryComponentTick.bCanEverTick = false;
+	
+	SetCanEverAffectNavigation(false);
+	bNavigationRelevant = false;
+}
+
+void UChaosVDGTAccelerationStructuresDataComponent::UpdateAABBTreeData(TConstArrayView<TSharedPtr<FChaosVDAABBTreeDataWrapper>> AABBTreeDataView)
+{
+	RecordedABBTreeData.Reset(AABBTreeDataView.Num());
+	Algo::Copy(AABBTreeDataView, RecordedABBTreeData);
+}
+
+void UChaosVDGTAccelerationStructuresDataComponent::ClearData()
+{
+	RecordedABBTreeData.Reset();
+}
diff --git a/Engine/Plugins/ChaosVD/Source/ChaosVD/Private/Components/ChaosVDGTAccelerationStructuresDataComponent.h b/Engine/Plugins/ChaosVD/Source/ChaosVD/Private/Components/ChaosVDGTAccelerationStructuresDataComponent.h
new file mode 100644
--- /dev/null
+++ b/Engine/Plugins/ChaosVD/Source/ChaosVD/Private/Components/ChaosVDGTAccelerationStructuresDataComponent.h
@@ -0,0 +1,28 @@
+// Copyright Epic Games, Inc. All Rights Reserved.
+
+#pragma once
+
+#include "ChaosVDSolverDataComponent.h"
+#include "ChaosVDGTAccelerationStructuresDataComponent.generated.h"
+
+
+struct FChaosVDAABBTreeDataWrapper;
+
+UCLASS()
+class UChaosVDGTAccelerationStructuresDataComponent : public UChaosVDSolverDataComponent
+{
+	GENERATED_BODY()
+
+public:
+	// Sets default values for this component's properties
+	UChaosVDGTAccelerationStructuresDataComponent();
+
+	void UpdateAABBTreeData(TConstArrayView<TSharedPtr<FChaosVDAABBTreeDataWrapper>> AABBTreeDataView);
+
+	TConstArrayView<TSharedPtr<FChaosVDAABBTreeDataWrapper>> GetAABBTreeData() const { return RecordedABBTreeData; }
+
+	virtual void ClearData() override;
+protected:
+	
+	TArray<TSharedPtr<FChaosVDAABBTreeDataWrapper>> RecordedABBTreeData;
+};
diff --git a/Engine/Plugins/ChaosVD/Source/ChaosVD/Private/DetailsCustomizations/ChaosVDSelectionMultipleViewCustomization.cpp b/Engine/Plugins/ChaosVD/Source/ChaosVD/Private/DetailsCustomizations/ChaosVDSelectionMultipleViewCustomization.cpp
new file mode 100644
--- /dev/null
+++ b/Engine/Plugins/ChaosVD/Source/ChaosVD/Private/DetailsCustomizations/ChaosVDSelectionMultipleViewCustomization.cpp
@@ -0,0 +1,55 @@
+// Copyright Epic Games, Inc. All Rights Reserved.
+
+#include "ChaosVDSelectionMultipleViewCustomization.h"
+
+#include "ChaosVDModule.h"
+#include "ChaosVDSolverDataSelection.h"
+#include "DetailCategoryBuilder.h"
+#include "DetailLayoutBuilder.h"
+#include "PropertyHandle.h"
+
+
+TSharedRef<IDetailCustomization> FChaosVDSelectionMultipleViewCustomization::MakeInstance()
+{
+	return MakeShareable( new FChaosVDSelectionMultipleViewCustomization );
+}
+
+void FChaosVDSelectionMultipleViewCustomization::CustomizeDetails(IDetailLayoutBuilder& DetailBuilder)
+{
+	TArray<TSharedPtr<FStructOnScope>> StructsBeingCustomized;
+	DetailBuilder.GetStructsBeingCustomized(StructsBeingCustomized);
+
+	if (!ensure(StructsBeingCustomized.Num() == 1))
+	{
+		UE_LOG(LogChaosVDEditor, Warning, TEXT("[%s] [%d] objects were selectioned but this customization panel only support single object selection."), ANSI_TO_TCHAR(__FUNCTION__), StructsBeingCustomized.Num())
+	}
+
+	if (StructsBeingCustomized.IsEmpty())
+	{
+		return;
+	}
+	TSharedPtr<FStructOnScope> SelectionViewStructOnScope = StructsBeingCustomized.IsEmpty() ? nullptr : StructsBeingCustomized[0];
+
+	if (!SelectionViewStructOnScope)
+	{
+		return;
+	}
+
+	FChaosVDSelectionMultipleView* StructView = reinterpret_cast<FChaosVDSelectionMultipleView*>(SelectionViewStructOnScope->GetStructMemory());
+	
+	IDetailCategoryBuilder& CategoryBuilder = DetailBuilder.EditCategory("Recorded Data");
+	
+	for (const TSharedPtr<FStructOnScope>& DataInstance : StructView->DataInstances)
+	{
+		if (!DataInstance || !DataInstance->IsValid())
+		{
+			continue;
+		}
+
+		if (IDetailPropertyRow* CreatedRow = CategoryBuilder.AddExternalStructure(DataInstance))
+		{
+			CreatedRow->DisplayName(DataInstance->GetStructPtr()->GetDisplayNameText());
+			CreatedRow->ShouldAutoExpand(true);
+		}
+	}
+}
diff --git a/Engine/Plugins/ChaosVD/Source/ChaosVD/Private/DetailsCustomizations/ChaosVDSelectionMultipleViewCustomization.h b/Engine/Plugins/ChaosVD/Source/ChaosVD/Private/DetailsCustomizations/ChaosVDSelectionMultipleViewCustomization.h
new file mode 100644
--- /dev/null
+++ b/Engine/Plugins/ChaosVD/Source/ChaosVD/Private/DetailsCustomizations/ChaosVDSelectionMultipleViewCustomization.h
@@ -0,0 +1,15 @@
+// Copyright Epic Games, Inc. All Rights Reserved.
+
+#pragma once
+
+#include "Templates/SharedPointer.h"
+#include "IDetailCustomization.h"
+
+class FChaosVDSelectionMultipleViewCustomization : public IDetailCustomization
+{
+public:
+
+	static TSharedRef<IDetailCustomization> MakeInstance();
+
+	virtual void CustomizeDetails(IDetailLayoutBuilder& DetailBuilder) override;
+};
diff --git a/Engine/Plugins/ChaosVD/Source/ChaosVD/Private/Settings/ChaosVDAccelerationStructureVisualizationSettings.cpp b/Engine/Plugins/ChaosVD/Source/ChaosVD/Private/Settings/ChaosVDAccelerationStructureVisualizationSettings.cpp
new file mode 100644
--- /dev/null
+++ b/Engine/Plugins/ChaosVD/Source/ChaosVD/Private/Settings/ChaosVDAccelerationStructureVisualizationSettings.cpp
@@ -0,0 +1,24 @@
+// Copyright Epic Games, Inc. All Rights Reserved.
+
+#include "ChaosVDAccelerationStructureVisualizationSettings.h"
+
+#include "ChaosVDSettingsManager.h"
+
+void UChaosVDAccelerationStructureVisualizationSettings::SetDataVisualizationFlags(EChaosVDAccelerationStructureDataVisualizationFlags NewFlags)
+{
+	if (UChaosVDAccelerationStructureVisualizationSettings* Settings = FChaosVDSettingsManager::Get().GetSettingsObject<UChaosVDAccelerationStructureVisualizationSettings>())
+	{
+		Settings->AccelerationStructureDataVisualizationFlags = static_cast<uint32>(NewFlags);
+		Settings->BroadcastSettingsChanged();
+	}
+}
+
+EChaosVDAccelerationStructureDataVisualizationFlags UChaosVDAccelerationStructureVisualizationSettings::GetDataVisualizationFlags()
+{
+	if (UChaosVDAccelerationStructureVisualizationSettings* Settings = FChaosVDSettingsManager::Get().GetSettingsObject<UChaosVDAccelerationStructureVisualizationSettings>())
+	{
+		return static_cast<EChaosVDAccelerationStructureDataVisualizationFlags>(Settings->AccelerationStructureDataVisualizationFlags);
+	}
+
+	return EChaosVDAccelerationStructureDataVisualizationFlags::None;
+}
diff --git a/Engine/Plugins/ChaosVD/Source/ChaosVD/Private/Settings/ChaosVDAccelerationStructureVisualizationSettings.h b/Engine/Plugins/ChaosVD/Source/ChaosVD/Private/Settings/ChaosVDAccelerationStructureVisualizationSettings.h
new file mode 100644
--- /dev/null
+++ b/Engine/Plugins/ChaosVD/Source/ChaosVD/Private/Settings/ChaosVDAccelerationStructureVisualizationSettings.h
@@ -0,0 +1,54 @@
+// Copyright Epic Games, Inc. All Rights Reserved.
+
+#pragma once
+
+#include "ChaosVDCoreSettings.h"
+#include "ChaosVDAccelerationStructureVisualizationSettings.generated.h"
+
+/** Visualization flags used to control what is debug draw of the recorded acceleration structure data */
+UENUM(meta = (Bitflags, UseEnumValuesAsMaskValuesInEditor = "true"))
+enum class EChaosVDAccelerationStructureDataVisualizationFlags : uint32
+{
+	None	= 0 UMETA(Hidden),
+	/** If Set, draws the bound of all nodes in the tree  */
+	DrawNodesBounds = 1 << 0,
+	/** If Set, draws lines to represent the branches of the tree  */
+	DrawBranches = 1 << 1,
+	/** If Set, draws the bounds of the leaves of the tree  */
+	DrawLeavesBounds = 1 << 2,
+	/** If Set, draws the bounds of each element in the leaves of the tree  */
+	DrawLeavesElementBounds = 1 << 3,
+	/** If Set, inside each leaf bound, it draws lines from the elements to the center of the leaf bounds */
+	DrawLeavesElementConnections = 1 << 4,
+	/** If Set, all dynamic trees will be drawn */
+	DrawDynamicTrees = 1 << 5,
+	/** If Set, all static trees will be drawn */
+	DrawStaticTrees = 1 << 6,
+	/** If set, enabled debug drawing for any recorded acceleration structure available at the current visualized frame */
+	EnableDraw = DrawDynamicTrees | DrawStaticTrees
+};
+ENUM_CLASS_FLAGS(EChaosVDAccelerationStructureDataVisualizationFlags);
+
+/**
+ * Settings object that stores the values that control how acceleration structures are debug drawn
+ */
+UCLASS()
+class CHAOSVD_API UChaosVDAccelerationStructureVisualizationSettings : public UChaosVDVisualizationSettingsObjectBase
+{
+	GENERATED_BODY()
+public:
+
+	/** The depth priority used for while drawing contact data. Can be World or Foreground (with this one the shapes will be drawn on top of the geometry and be always visible) */
+	UPROPERTY(EditAnywhere, config, Category=DebugDraw)
+	TEnumAsByte<ESceneDepthPriorityGroup> DepthPriority = ESceneDepthPriorityGroup::SDPG_Foreground;
+
+	/** The base thickness used to draw node bounds */
+	UPROPERTY(EditAnywhere, config, Category=DebugDraw)
+	float BaseThickness = 3.0f;
+	
+	static void SetDataVisualizationFlags(EChaosVDAccelerationStructureDataVisualizationFlags NewFlags);
+	static EChaosVDAccelerationStructureDataVisualizationFlags GetDataVisualizationFlags();
+private:
+	UPROPERTY(config, meta = (Bitmask, BitmaskEnum = "/Script/ChaosVD.EChaosVDAccelerationStructureDataVisualizationFlags"))
+	uint32 AccelerationStructureDataVisualizationFlags = static_cast<uint32>(EChaosVDAccelerationStructureDataVisualizationFlags::DrawNodesBounds | EChaosVDAccelerationStructureDataVisualizationFlags::DrawBranches);
+};
diff --git a/Engine/Plugins/ChaosVD/Source/ChaosVD/Private/Trace/ChaosVDTraceProvider.cpp b/Engine/Plugins/ChaosVD/Source/ChaosVD/Private/Trace/ChaosVDTraceProvider.cpp
--- a/Engine/Plugins/ChaosVD/Source/ChaosVD/Private/Trace/ChaosVDTraceProvider.cpp
+++ b/Engine/Plugins/ChaosVD/Source/ChaosVD/Private/Trace/ChaosVDTraceProvider.cpp
@@ -9,6 +9,7 @@
 #include "ChaosVisualDebugger/ChaosVisualDebuggerTrace.h"
 
 #include "Compression/OodleDataCompressionUtil.h"
+#include "DataProcessors/ChaosVDAABBTreeDataProcessor.h"
 #include "Trace/DataProcessors/ChaosVDCharacterGroundConstraintDataProcessor.h"
 #include "Trace/DataProcessors/ChaosVDConstraintDataProcessor.h"
 #include "Trace/DataProcessors/ChaosVDJointConstraintDataProcessor.h"
@@ -274,6 +275,10 @@ void FChaosVDTraceProvider::RegisterDefaultDataProcessorsIfNeeded()
 	ArchiveHeaderDataProcessor->SetTraceProvider(AsShared());
 	RegisterDataProcessor(ArchiveHeaderDataProcessor);
 
+	TSharedPtr<FChaosVDAABBTreeDataProcessor> AABBTreeDataProcessor = MakeShared<FChaosVDAABBTreeDataProcessor>();
+	AABBTreeDataProcessor->SetTraceProvider(AsShared());
+	RegisterDataProcessor(AABBTreeDataProcessor);
+
 	bDefaultDataProcessorsRegistered = true;
 }
 
diff --git a/Engine/Plugins/ChaosVD/Source/ChaosVD/Private/Trace/DataProcessors/ChaosVDAABBTreeDataProcessor.cpp b/Engine/Plugins/ChaosVD/Source/ChaosVD/Private/Trace/DataProcessors/ChaosVDAABBTreeDataProcessor.cpp
new file mode 100644
--- /dev/null
+++ b/Engine/Plugins/ChaosVD/Source/ChaosVD/Private/Trace/DataProcessors/ChaosVDAABBTreeDataProcessor.cpp
@@ -0,0 +1,64 @@
+// Copyright Epic Games, Inc. All Rights Reserved.
+
+#include "ChaosVDAABBTreeDataProcessor.h"
+
+#include "DataWrappers/ChaosVDAccelerationStructureDataWrappers.h"
+
+namespace Chaos::VisualDebugger::Utils
+{
+	int32 CalculateAABBTreeDepth(const TSharedRef<FChaosVDAABBTreeDataWrapper>& InTree, int32 StartNodeIndex, int32 CurrentDepth = 0)
+	{
+		if(InTree->Nodes.IsEmpty() || (!InTree->Nodes.IsValidIndex(StartNodeIndex)))
+		{
+			return 0;
+		}
+
+		if (InTree->Nodes[StartNodeIndex].bLeaf)
+		{
+			return 0;
+		}
+
+		int32 LeftDepth = CalculateAABBTreeDepth(InTree, InTree->Nodes[StartNodeIndex].ChildrenNodes[0]);
+		int32 RightDepth = CalculateAABBTreeDepth(InTree, InTree->Nodes[StartNodeIndex].ChildrenNodes[1]);
+
+		if (LeftDepth > RightDepth)
+		{
+			return LeftDepth + 1;
+		}
+		else
+		{
+			return RightDepth + 1;
+		}
+	}
+}
+
+FChaosVDAABBTreeDataProcessor::FChaosVDAABBTreeDataProcessor() : FChaosVDDataProcessorBase(FChaosVDAABBTreeDataWrapper::WrapperTypeName)
+{
+}
+
+bool FChaosVDAABBTreeDataProcessor::ProcessRawData(const TArray<uint8>& InData)
+{
+	FChaosVDDataProcessorBase::ProcessRawData(InData);
+
+	const TSharedPtr<FChaosVDTraceProvider> ProviderSharedPtr = TraceProvider.Pin();
+	if (!ensure(ProviderSharedPtr.IsValid()))
+	{
+		return false;
+	}
+
+	const TSharedPtr<FChaosVDAABBTreeDataWrapper> AABBTreeData = MakeShared<FChaosVDAABBTreeDataWrapper>();
+	const bool bSuccess = Chaos::VisualDebugger::ReadDataFromBuffer(InData, *AABBTreeData, ProviderSharedPtr.ToSharedRef());
+	
+	if (bSuccess)
+	{
+		if (TSharedPtr<FChaosVDGameFrameData> CurrentFrameData = ProviderSharedPtr->GetCurrentGameFrame().Pin())
+		{
+			AABBTreeData->TreeDepth = Chaos::VisualDebugger::Utils::CalculateAABBTreeDepth(AABBTreeData.ToSharedRef(), AABBTreeData->GetCorrectedRootNodeIndex());
+
+			TArray<TSharedPtr<FChaosVDAABBTreeDataWrapper>>& ABBTrees = CurrentFrameData->RecordedAABBTreesBySolverID.FindOrAdd(AABBTreeData->SolverId);
+			ABBTrees.Add(AABBTreeData);
+		}
+	}
+
+	return bSuccess;
+}
diff --git a/Engine/Plugins/ChaosVD/Source/ChaosVD/Private/Trace/DataProcessors/ChaosVDAABBTreeDataProcessor.h b/Engine/Plugins/ChaosVD/Source/ChaosVD/Private/Trace/DataProcessors/ChaosVDAABBTreeDataProcessor.h
new file mode 100644
--- /dev/null
+++ b/Engine/Plugins/ChaosVD/Source/ChaosVD/Private/Trace/DataProcessors/ChaosVDAABBTreeDataProcessor.h
@@ -0,0 +1,15 @@
+// Copyright Epic Games, Inc. All Rights Reserved.
+
+#pragma once
+#include "ChaosVDDataProcessorBase.h"
+
+/**
+ * Data processor implementation that is able to deserialize traced AABB Tree Data
+ */
+class FChaosVDAABBTreeDataProcessor final : public FChaosVDDataProcessorBase
+{
+public:
+	explicit FChaosVDAABBTreeDataProcessor();
+	
+	virtual bool ProcessRawData(const TArray<uint8>& InData) override;
+};
diff --git a/Engine/Plugins/ChaosVD/Source/ChaosVD/Private/Visualizers/ChaosVDComponentVisualizerBase.cpp b/Engine/Plugins/ChaosVD/Source/ChaosVD/Private/Visualizers/ChaosVDComponentVisualizerBase.cpp
--- a/Engine/Plugins/ChaosVD/Source/ChaosVD/Private/Visualizers/ChaosVDComponentVisualizerBase.cpp
+++ b/Engine/Plugins/ChaosVD/Source/ChaosVD/Private/Visualizers/ChaosVDComponentVisualizerBase.cpp
@@ -44,13 +44,13 @@ bool FChaosVDComponentVisualizerBase::VisProxyHandleClick(FEditorViewportClient*
 			TabManager->TryInvokeTab(InspectorTabID);
 		}
 
-		return SelectVisualizedData(*ChaosVDHitProxy, CVDScene.ToSharedRef());
+		return SelectVisualizedData(*ChaosVDHitProxy, CVDScene.ToSharedRef(), MainTabToolkitHost.ToSharedRef());
 	}
 
 	return false;
 }
 
-bool FChaosVDComponentVisualizerBase::SelectVisualizedData(const HChaosVDComponentVisProxy& VisProxy, const TSharedRef<FChaosVDScene>& InCVDScene)
+bool FChaosVDComponentVisualizerBase::SelectVisualizedData(const HChaosVDComponentVisProxy& VisProxy, const TSharedRef<FChaosVDScene>& InCVDScene, const TSharedRef<SChaosVDMainTab>& InMainTabToolkitHost)
 {
 	if (TSharedPtr<FChaosVDSolverDataSelection> SelectionObject = InCVDScene->GetSolverDataSelectionObject().Pin())
 	{
diff --git a/Engine/Plugins/ChaosVD/Source/ChaosVD/Private/Visualizers/ChaosVDComponentVisualizerBase.h b/Engine/Plugins/ChaosVD/Source/ChaosVD/Private/Visualizers/ChaosVDComponentVisualizerBase.h
--- a/Engine/Plugins/ChaosVD/Source/ChaosVD/Private/Visualizers/ChaosVDComponentVisualizerBase.h
+++ b/Engine/Plugins/ChaosVD/Source/ChaosVD/Private/Visualizers/ChaosVDComponentVisualizerBase.h
@@ -13,6 +13,7 @@
 #include "Widgets/SChaosVDEnumFlagsMenu.h"
 #include "Widgets/SWidget.h"
 
+class SChaosVDMainTab;
 class FChaosVDSolverDataSelection;
 struct FChaosVDSolverDataSelectionHandle;
 class UChaosVDSettingsObjectBase;
@@ -60,7 +61,7 @@ protected:
 	
 	virtual bool CanHandleClick(const HChaosVDComponentVisProxy& VisProxy);
 
-	virtual bool SelectVisualizedData(const HChaosVDComponentVisProxy& VisProxy, const TSharedRef<FChaosVDScene>& InCVDScene);
+	virtual bool SelectVisualizedData(const HChaosVDComponentVisProxy& VisProxy, const TSharedRef<FChaosVDScene>& InCVDScene,const TSharedRef<SChaosVDMainTab>& InMainTabToolkitHost);
 
 	template<typename ObjectSettingsType, typename VisualizationFlagsType>
 	void CreateGenericVisualizerMenu(FName MenuToExtend, FName SectionName, const FText& InSectionLabel, const FText& InFlagsMenuLabel,  const FText& InFlagsMenuTooltip, FSlateIcon FlagsMenuIcon,  const FText& InSettingsMenuLabel, const FText& InSettingsMenuTooltip);
diff --git a/Engine/Plugins/ChaosVD/Source/ChaosVD/Private/Visualizers/ChaosVDGTAccelerationStructureDataComponentVisualizer.cpp b/Engine/Plugins/ChaosVD/Source/ChaosVD/Private/Visualizers/ChaosVDGTAccelerationStructureDataComponentVisualizer.cpp
new file mode 100644
--- /dev/null
+++ b/Engine/Plugins/ChaosVD/Source/ChaosVD/Private/Visualizers/ChaosVDGTAccelerationStructureDataComponentVisualizer.cpp
@@ -0,0 +1,365 @@
+// Copyright Epic Games, Inc. All Rights Reserved.
+
+#include "ChaosVDGTAccelerationStructureDataComponentVisualizer.h"
+
+#include "Actors/ChaosVDSolverInfoActor.h"
+#include "ChaosVDDebugDrawUtils.h"
+#include "ChaosVDScene.h"
+#include "ChaosVDStyle.h"
+#include "ChaosVDTabsIDs.h"
+#include "Components/ChaosVDGTAccelerationStructuresDataComponent.h"
+#include "SceneView.h"
+#include "Settings/ChaosVDAccelerationStructureVisualizationSettings.h"
+#include "Widgets/SChaosVDViewportToolbar.h"
+
+#define LOCTEXT_NAMESPACE "ChaosVisualDebugger"
+
+class AChaosVDSolverInfoActor;
+
+namespace Chaos::VisualDebugger::Utils
+{
+	void DrawFBoxAtLocation(FPrimitiveDrawInterface* PDI, const FBox& InBox, FColor Color, ESceneDepthPriorityGroup DepthPriority, float Thickness)
+	{
+		FVector Center;
+		FVector Extents;	
+		InBox.GetCenterAndExtents(Center, Extents);
+
+		FTransform LocationTransform;
+		LocationTransform.SetLocation(Center);
+			
+		FChaosVDDebugDrawUtils::DrawBox(PDI, Extents, Color, LocationTransform, FText::GetEmpty(), DepthPriority, Thickness);
+	}
+
+	void DrawBoxAtLocation(FPrimitiveDrawInterface* PDI, const FVector& Center, const FVector& Extents, FColor Color, ESceneDepthPriorityGroup DepthPriority, float Thickness)
+	{
+		FTransform LocationTransform;
+		LocationTransform.SetLocation(Center);
+			
+		FChaosVDDebugDrawUtils::DrawBox(PDI, Extents, Color, LocationTransform, FText::GetEmpty(), DepthPriority, Thickness);
+	}
+}
+
+bool FChaosVDGTAccelerationStructureSelectionHandle::IsSelected()
+{
+	// In contrast to other recorded data types, AABB Tree data is recorded as a single struct, so we need to also use the context data to match a selection handle
+	if (bool bIsPrimaryDataSelected = FChaosVDSolverDataSelectionHandle::IsSelected())
+	{
+		if (TSharedPtr<FChaosVDSolverDataSelection> OwnerPtr = Owner.Pin())
+		{
+			TSharedPtr<FChaosVDSolverDataSelectionHandle> CurrentSelectedDataHandle = OwnerPtr->GetCurrentSelectionHandle();
+
+			FChaosVDAABBTreeSelectionContext* CurrentSelectionContext = CurrentSelectedDataHandle->GetContextData<FChaosVDAABBTreeSelectionContext>();
+			FChaosVDAABBTreeSelectionContext* HandleSelectionContext = GetContextData<FChaosVDAABBTreeSelectionContext>();
+
+			return CurrentSelectionContext && HandleSelectionContext && (*CurrentSelectionContext) == (*HandleSelectionContext);
+		}
+	}
+
+	return false;
+}
+
+void FChaosVDGTAccelerationStructureSelectionHandle::CreateStructViewForDetailsPanelIfNeeded()
+{
+	if (StructDataView)
+	{
+		return;
+	}
+
+	StructDataView = MakeShared<FChaosVDSelectionMultipleView>();
+
+	if (FChaosVDAABBTreeDataWrapper* TreeData = GetData<FChaosVDAABBTreeDataWrapper>())
+	{
+		StructDataView->AddData(TreeData);
+	}
+
+	if (FChaosVDAABBTreeSelectionContext* SelectionContext = GetContextData<FChaosVDAABBTreeSelectionContext>())
+	{
+		StructDataView->AddData(const_cast<FChaosVDAABBTreeNodeDataWrapper*>(SelectionContext->NodeData));
+		StructDataView->AddData(const_cast<FChaosVDAABBTreeLeafDataWrapper*>(SelectionContext->LeafData));
+	}
+
+	StructDataViewStructOnScope = MakeShared<FStructOnScope>(FChaosVDSelectionMultipleView::StaticStruct(), reinterpret_cast<uint8*>(StructDataView.Get()));
+}
+
+TSharedPtr<FStructOnScope> FChaosVDGTAccelerationStructureSelectionHandle::GetCustomDataReadOnlyStructViewForDetails()
+{
+	// To avoid unnecessary work, only create and cache a view struct when requested (which happens when we try to use this selection handle to udpate a details panel)
+	CreateStructViewForDetailsPanelIfNeeded();
+
+	return StructDataViewStructOnScope;
+}
+
+FChaosVDGTAccelerationStructureDataComponentVisualizer::FChaosVDGTAccelerationStructureDataComponentVisualizer()
+{
+	RegisterVisualizerMenus();
+	InspectorTabID = FChaosVDTabID::DetailsPanel;
+}
+
+void FChaosVDGTAccelerationStructureDataComponentVisualizer::RegisterVisualizerMenus()
+{
+	FName MenuSection("AccelerationStructureDataVisualization.Show");
+	FText MenuSectionLabel = LOCTEXT("AccelerationStructureDataShowMenuLabel", "Acceleration Structure Data Visualization");
+	FText FlagsMenuLabel = LOCTEXT("AccelerationStructureDataFlagsMenuLabel", "Acceleration Structure Data Flags");
+	FText FlagsMenuTooltip = LOCTEXT("AccelerationStructureDataFlagsMenuToolTip", "Set of flags to enable/disable visibility of specific types of acceleration structure data");
+	FSlateIcon FlagsMenuIcon = FSlateIcon(FChaosVDStyle::Get().GetStyleSetName(), TEXT("SceneQueriesInspectorIcon"));
+
+	FText SettingsMenuLabel = LOCTEXT("AccelerationStructureSettingsMenuLabel", "Acceleration Structure Visualization Settings");
+	FText SettingsMenuTooltip = LOCTEXT("AccelerationStructureSettingsMenuToolTip", "Options to change how the recorded acceleration structure data is debug drawn");
+	
+	CreateGenericVisualizerMenu<UChaosVDAccelerationStructureVisualizationSettings, EChaosVDAccelerationStructureDataVisualizationFlags>(SChaosVDViewportToolbar::ShowMenuName, MenuSection, MenuSectionLabel, FlagsMenuLabel, FlagsMenuTooltip, FlagsMenuIcon, SettingsMenuLabel, SettingsMenuTooltip);
+}
+
+void FChaosVDGTAccelerationStructureDataComponentVisualizer::DrawVisualization(const UActorComponent* Component, const FSceneView* View, FPrimitiveDrawInterface* PDI)
+{
+	const UChaosVDGTAccelerationStructuresDataComponent* DataComponent = Cast<UChaosVDGTAccelerationStructuresDataComponent>(Component);
+	if (!DataComponent)
+	{
+		return;
+	}
+	
+	AChaosVDSolverInfoActor* SolverInfoActor = Cast<AChaosVDSolverInfoActor>(Component->GetOwner());
+	if (!SolverInfoActor)
+	{
+		return;
+	}
+
+	if (!SolverInfoActor->IsVisible())
+	{
+		return;
+	}
+
+	const TSharedPtr<FChaosVDScene> CVDScene = SolverInfoActor->GetScene().Pin();
+	if (!CVDScene)
+	{
+		return;
+	}
+	
+	FChaosGTAccelerationStructureVisualizationDataContext VisualizationContext;
+	VisualizationContext.CVDScene = CVDScene;
+	VisualizationContext.SpaceTransform = SolverInfoActor->GetSimulationTransform();
+	VisualizationContext.SolverDataSelectionObject = CVDScene->GetSolverDataSelectionObject().Pin();
+	VisualizationContext.DataComponent = DataComponent;
+
+	if (const UChaosVDAccelerationStructureVisualizationSettings* EditorSettings = FChaosVDSettingsManager::Get().GetSettingsObject<UChaosVDAccelerationStructureVisualizationSettings>())
+	{
+		VisualizationContext.VisualizationFlags = static_cast<uint32>(UChaosVDAccelerationStructureVisualizationSettings::GetDataVisualizationFlags());
+		VisualizationContext.DebugDrawSettings = EditorSettings;
+		VisualizationContext.DepthPriority = EditorSettings->DepthPriority;
+	}
+
+	if (!VisualizationContext.IsVisualizationFlagEnabled(EChaosVDAccelerationStructureDataVisualizationFlags::EnableDraw))
+	{
+		return;
+	}
+
+	TConstArrayView<TSharedPtr<FChaosVDAABBTreeDataWrapper>> RecordedAABBTrees = DataComponent->GetAABBTreeData();
+
+	for (const TSharedPtr<FChaosVDAABBTreeDataWrapper>& AABBTreeDataWrapper : RecordedAABBTrees)
+	{
+		if (AABBTreeDataWrapper)
+		{
+			const bool bCanDrawTree = (AABBTreeDataWrapper->bDynamicTree && VisualizationContext.IsVisualizationFlagEnabled(EChaosVDAccelerationStructureDataVisualizationFlags::DrawDynamicTrees))
+								|| (!AABBTreeDataWrapper->bDynamicTree && VisualizationContext.IsVisualizationFlagEnabled(EChaosVDAccelerationStructureDataVisualizationFlags::DrawStaticTrees));
+
+			if (bCanDrawTree)
+			{
+				DrawAABBTree(View, PDI, VisualizationContext, AABBTreeDataWrapper.ToSharedRef());
+			}	
+		}
+	}
+}
+
+bool FChaosVDGTAccelerationStructureDataComponentVisualizer::CanHandleClick(const HChaosVDComponentVisProxy& VisProxy)
+{
+	return VisProxy.DataSelectionHandle && (VisProxy.DataSelectionHandle->IsA<FChaosVDAABBTreeDataWrapper>()
+											|| VisProxy.DataSelectionHandle->IsA<FChaosVDAABBTreeNodeDataWrapper>()
+											|| VisProxy.DataSelectionHandle->IsA<FChaosVDAABBTreeLeafDataWrapper>());
+}
+
+void FChaosVDGTAccelerationStructureDataComponentVisualizer::DrawAABBTree(const FSceneView* View, FPrimitiveDrawInterface* PDI, const FChaosGTAccelerationStructureVisualizationDataContext& VisualizationContext, const TSharedRef<FChaosVDAABBTreeDataWrapper>& AABBTreeData)
+{
+	const UChaosVDAccelerationStructureVisualizationSettings* Settings = Cast<const UChaosVDAccelerationStructureVisualizationSettings>(VisualizationContext.DebugDrawSettings);
+	if (!ensure(Settings))
+	{
+		return;
+	}
+
+	int32 RootNodeIndex = AABBTreeData->GetCorrectedRootNodeIndex();
+	if (AABBTreeData->Nodes.Num() > 0 && AABBTreeData->Nodes.IsValidIndex(RootNodeIndex))
+	{
+		DrawAABBTreeNode(View, PDI, VisualizationContext, AABBTreeData, AABBTreeData->Nodes[RootNodeIndex], Settings->BaseThickness);
+	}
+}
+
+void FChaosVDGTAccelerationStructureDataComponentVisualizer::DrawAABBTreeNode(const FSceneView* View, FPrimitiveDrawInterface* PDI, const FChaosGTAccelerationStructureVisualizationDataContext& VisualizationContext, const TSharedRef<FChaosVDAABBTreeDataWrapper>& AABBTreeData, const FChaosVDAABBTreeNodeDataWrapper& AABBTreeNodeData, float Thickness, int32 CurrentTreeLevel)
+{
+	const bool bCanDrawNodeData = VisualizationContext.IsVisualizationFlagEnabled(EChaosVDAccelerationStructureDataVisualizationFlags::DrawNodesBounds | EChaosVDAccelerationStructureDataVisualizationFlags::DrawBranches);
+	const bool bCanDrawLeavesData = VisualizationContext.IsVisualizationFlagEnabled(EChaosVDAccelerationStructureDataVisualizationFlags::DrawLeavesBounds | EChaosVDAccelerationStructureDataVisualizationFlags::DrawLeavesElementBounds | EChaosVDAccelerationStructureDataVisualizationFlags::DrawLeavesElementConnections);
+
+	if (!bCanDrawNodeData && !bCanDrawLeavesData)
+	{
+		return;
+	}
+
+	if (AABBTreeNodeData.bLeaf)
+	{
+		if (bCanDrawLeavesData)
+		{
+			for (int32 ChildIndex : AABBTreeNodeData.ChildrenNodes)
+			{
+				if (AABBTreeData->TreeArrayLeafs.IsValidIndex(ChildIndex))
+				{
+					DrawAABBTreeArrayLeaf(View, PDI, VisualizationContext, AABBTreeData->TreeArrayLeafs[ChildIndex], AABBTreeData, Thickness);
+				}	
+			}
+		}
+		return;
+	}
+
+	auto IsNodeVisible = [View](const FBox& NodeBounds)
+	{
+		return View->ViewFrustum.IntersectBox(NodeBounds.GetCenter(), NodeBounds.GetExtent());
+	};
+
+	// Calculate and cache the total bounds of this node and its visibility state
+	FBox TotalNodeBounds(ForceInitToZero);
+	bool IsChildNodeVisible[2] = { false, false };
+
+	constexpr int32 MaxChildNodeNum = 2;
+	for (int32 ChildIndex = 0; ChildIndex < MaxChildNodeNum; ++ChildIndex)
+	{
+		const FBox& ChildBounds = AABBTreeNodeData.ChildrenBounds[ChildIndex];
+		IsChildNodeVisible[ChildIndex] = IsNodeVisible(ChildBounds);
+
+		TotalNodeBounds += ChildBounds;
+	}
+
+	const bool bIsCurrentNodeVisible = IsChildNodeVisible[0] || IsChildNodeVisible[1];
+	if (!bIsCurrentNodeVisible)
+	{
+		// If this node is not visible at all, nothing to do here
+		return;
+	}
+
+	// If node data drawing is disabled, we can skip all the logic to create the selection handle and draw the lines and go straight to continue traversing the tree
+	if (bCanDrawNodeData)
+	{
+		bool bIsRootNode = AABBTreeNodeData.ParentNode == INDEX_NONE;
+		float FinalThickness = Thickness;
+
+		FColor BoundsColor = FColor::MakeRedToGreenColorFromScalar(static_cast<float>(CurrentTreeLevel) / static_cast<float>(AABBTreeData->TreeDepth));
+
+		TSharedPtr<FChaosVDSolverDataSelectionHandle> NodeSelectionHandle = VisualizationContext.SolverDataSelectionObject->MakeSelectionHandle<FChaosVDAABBTreeDataWrapper, FChaosVDGTAccelerationStructureSelectionHandle>(AABBTreeData.ToSharedPtr());
+
+		FChaosVDAABBTreeSelectionContext ContextData;
+		// The lifetime of the structure where this node data lives is bound to the selection handle, so we can safely store a ptr to it
+		ContextData.NodeData = &AABBTreeNodeData;
+		NodeSelectionHandle->SetHandleContext(MoveTemp(ContextData));
+			
+		FinalThickness = NodeSelectionHandle->IsSelected() ? FinalThickness * 2.5f : FinalThickness;
+
+		PDI->SetHitProxy(new HChaosVDComponentVisProxy(VisualizationContext.DataComponent, NodeSelectionHandle));
+	
+		for (int32 ChildIndex = 0; ChildIndex < MaxChildNodeNum; ++ChildIndex)
+		{
+			if (!IsChildNodeVisible[ChildIndex])
+			{
+				continue;
+			}
+			
+			const FBox& ChildBounds = AABBTreeNodeData.ChildrenBounds[ChildIndex];
+
+			if (VisualizationContext.IsVisualizationFlagEnabled(EChaosVDAccelerationStructureDataVisualizationFlags::DrawNodesBounds))
+			{
+				Chaos::VisualDebugger::Utils::DrawFBoxAtLocation(PDI, ChildBounds, BoundsColor, VisualizationContext.DepthPriority, FinalThickness);
+			}
+		
+			if (VisualizationContext.IsVisualizationFlagEnabled(EChaosVDAccelerationStructureDataVisualizationFlags::DrawBranches))
+			{			
+				const FVector NodeCenter = TotalNodeBounds.GetCenter();
+				FColor BranchColor = FColor::MakeRedToGreenColorFromScalar(static_cast<float>(CurrentTreeLevel) / static_cast<float>(AABBTreeData->TreeDepth));
+
+				FChaosVDDebugDrawUtils::DrawLine(PDI, NodeCenter, ChildBounds.GetCenter(), BranchColor, FText::GetEmpty(), VisualizationContext.DepthPriority, FinalThickness * 1.2f);
+
+				constexpr float BranchStatPointBoxSize = 1.0f;
+				static FVector StartPointBoxExtent(BranchStatPointBoxSize, BranchStatPointBoxSize, BranchStatPointBoxSize);
+				Chaos::VisualDebugger::Utils::DrawBoxAtLocation(PDI, NodeCenter, StartPointBoxExtent, bIsRootNode ? FColor::Red : BranchColor, VisualizationContext.DepthPriority, FinalThickness * (bIsRootNode ? 7.0f : 4.0f));
+			}
+		}
+
+		if (bIsRootNode && VisualizationContext.IsVisualizationFlagEnabled(EChaosVDAccelerationStructureDataVisualizationFlags::DrawNodesBounds))
+		{
+			// If we are the root node, also draw a box showing the bounds of the whole three
+			Chaos::VisualDebugger::Utils::DrawFBoxAtLocation(PDI, TotalNodeBounds, FColor::Red, VisualizationContext.DepthPriority, FinalThickness);
+		}
+	
+		PDI->SetHitProxy(nullptr);
+	}
+
+	// We can have leaf data drawing enabled while node data drawing is disabled, therefore we still need to continue traversing the tree to get to the leaves.
+	if (bCanDrawLeavesData || bCanDrawNodeData)
+	{
+		for (int32 ChildIndex = 0; ChildIndex < MaxChildNodeNum; ++ChildIndex)
+		{
+			// if the child node is not visible, we can discard the entire branch
+			if (!IsChildNodeVisible[ChildIndex])
+			{
+				continue;
+			}
+
+			int32 ChildNodeIndex = AABBTreeNodeData.ChildrenNodes[ChildIndex];
+			if (ChildNodeIndex > 0 && ChildNodeIndex < AABBTreeData->Nodes.Num())
+			{
+				constexpr float LineThicknessRatio = 0.75f;
+				DrawAABBTreeNode(View, PDI, VisualizationContext, AABBTreeData, AABBTreeData->Nodes[ChildNodeIndex], Thickness * LineThicknessRatio, CurrentTreeLevel + 1);
+			}
+		}
+	}
+}
+
+void FChaosVDGTAccelerationStructureDataComponentVisualizer::DrawAABBTreeArrayLeaf(const FSceneView* View, FPrimitiveDrawInterface* PDI, const FChaosGTAccelerationStructureVisualizationDataContext& VisualizationContext, const FChaosVDAABBTreeLeafDataWrapper& AABBTreeArrayLeafData, const TSharedRef<FChaosVDAABBTreeDataWrapper>& AABBTreeData, float Thickness)
+{
+	// Early out if this leaf will not be visible
+	if (!View->ViewFrustum.IntersectBox(AABBTreeArrayLeafData.Bounds.GetCenter(), AABBTreeArrayLeafData.Bounds.GetExtent()))
+	{
+		return;
+	}
+
+	constexpr float MaxDensityNumForColor = 10.0f;
+	float InverseAlpha =  1.0f - (static_cast<float>(AABBTreeArrayLeafData.Elements.Num()) / MaxDensityNumForColor);
+
+	FColor ColorByDensity = FColor::MakeRedToGreenColorFromScalar(InverseAlpha);
+
+	TSharedPtr<FChaosVDSolverDataSelectionHandle> LeafSelectionHandle = VisualizationContext.SolverDataSelectionObject->MakeSelectionHandle<FChaosVDAABBTreeDataWrapper, FChaosVDGTAccelerationStructureSelectionHandle>(AABBTreeData.ToSharedPtr());
+
+	FChaosVDAABBTreeSelectionContext ContextData;
+	// The lifetime of the structure where this leaf data lives is bound to the selection handle, so we can safely store a ptr to it
+	ContextData.LeafData = &AABBTreeArrayLeafData;
+	LeafSelectionHandle->SetHandleContext(MoveTemp(ContextData));
+
+	PDI->SetHitProxy(new HChaosVDComponentVisProxy(VisualizationContext.DataComponent, LeafSelectionHandle));
+
+	float FinalThickness = LeafSelectionHandle->IsSelected() ? Thickness * 2.5f : Thickness; 
+
+	if (VisualizationContext.IsVisualizationFlagEnabled(EChaosVDAccelerationStructureDataVisualizationFlags::DrawLeavesBounds))
+	{		
+		Chaos::VisualDebugger::Utils::DrawFBoxAtLocation(PDI, AABBTreeArrayLeafData.Bounds, FColor::Green, VisualizationContext.DepthPriority, FinalThickness);
+	}
+
+	for (const FChaosVDAABBTreePayloadBoundsElement& TreeArrayLeafElement : AABBTreeArrayLeafData.Elements)
+	{
+		if (VisualizationContext.IsVisualizationFlagEnabled(EChaosVDAccelerationStructureDataVisualizationFlags::DrawLeavesElementConnections))
+		{
+			FChaosVDDebugDrawUtils::DrawLine(PDI, AABBTreeArrayLeafData.Bounds.GetCenter(), TreeArrayLeafElement.Bounds.GetCenter(), ColorByDensity, FText::GetEmpty(), VisualizationContext.DepthPriority, FinalThickness);
+		}
+
+		if (VisualizationContext.IsVisualizationFlagEnabled(EChaosVDAccelerationStructureDataVisualizationFlags::DrawLeavesElementBounds))
+		{	
+			Chaos::VisualDebugger::Utils::DrawFBoxAtLocation(PDI, TreeArrayLeafElement.Bounds, ColorByDensity, VisualizationContext.DepthPriority, FinalThickness * 0.7f);
+		}
+	}
+
+	PDI->SetHitProxy(nullptr);
+}
+
+#undef LOCTEXT_NAMESPACE
diff --git a/Engine/Plugins/ChaosVD/Source/ChaosVD/Private/Visualizers/ChaosVDGTAccelerationStructureDataComponentVisualizer.h b/Engine/Plugins/ChaosVD/Source/ChaosVD/Private/Visualizers/ChaosVDGTAccelerationStructureDataComponentVisualizer.h
new file mode 100644
--- /dev/null
+++ b/Engine/Plugins/ChaosVD/Source/ChaosVD/Private/Visualizers/ChaosVDGTAccelerationStructureDataComponentVisualizer.h
@@ -0,0 +1,81 @@
+// Copyright Epic Games, Inc. All Rights Reserved.
+
+#pragma once
+#include "ChaosVDComponentVisualizerBase.h"
+#include "ChaosVDSolverDataSelection.h"
+#include "Settings/ChaosVDAccelerationStructureVisualizationSettings.h"
+
+#include "ChaosVDGTAccelerationStructureDataComponentVisualizer.generated.h"
+
+class UChaosVDGTAccelerationStructuresDataComponent;
+struct FChaosVDAABBTreeLeafDataWrapper;
+struct FChaosVDAABBTreeNodeDataWrapper;
+struct FChaosVDAABBTreeDataWrapper;
+
+enum class EChaosVDVisibleAABBTreeNodes
+{
+	None = 0,
+	Left = 1 << 0,
+	Right = 1 << 1
+};
+ENUM_CLASS_FLAGS(EChaosVDVisibleAABBTreeNodes)
+
+USTRUCT()
+struct FChaosVDAABBTreeSelectionContext
+{
+	GENERATED_BODY()
+
+	const FChaosVDAABBTreeNodeDataWrapper* NodeData = nullptr;
+	const FChaosVDAABBTreeLeafDataWrapper* LeafData = nullptr;
+
+	bool operator==(const FChaosVDAABBTreeSelectionContext& Other) const
+	{
+		return NodeData == Other.NodeData && LeafData == Other.LeafData;
+	}
+};
+
+struct FChaosVDGTAccelerationStructureSelectionHandle : public FChaosVDSolverDataSelectionHandle
+{
+	virtual bool IsSelected() override;
+	void CreateStructViewForDetailsPanelIfNeeded();
+
+	virtual TSharedPtr<FStructOnScope> GetCustomDataReadOnlyStructViewForDetails() override;
+
+private:
+	TSharedPtr<FChaosVDSelectionMultipleView> StructDataView;
+	TSharedPtr<FStructOnScope> StructDataViewStructOnScope;
+};
+
+/** Visualization context structure specific for acceleration structure visualizations */
+struct FChaosGTAccelerationStructureVisualizationDataContext : public FChaosVDVisualizationContext
+{
+	TSharedPtr<FChaosVDSolverDataSelectionHandle> DataSelectionHandle = MakeShared<FChaosVDSolverDataSelectionHandle>();
+
+	ESceneDepthPriorityGroup DepthPriority = SDPG_Foreground;
+
+	const UChaosVDGTAccelerationStructuresDataComponent* DataComponent = nullptr;
+
+	bool IsVisualizationFlagEnabled(EChaosVDAccelerationStructureDataVisualizationFlags Flag) const
+	{
+		const EChaosVDAccelerationStructureDataVisualizationFlags FlagsAsAccelerationStructureVisFlags = static_cast<EChaosVDAccelerationStructureDataVisualizationFlags>(VisualizationFlags);
+		return EnumHasAnyFlags(FlagsAsAccelerationStructureVisFlags, Flag);
+	}	
+};
+
+class FChaosVDGTAccelerationStructureDataComponentVisualizer final : public FChaosVDComponentVisualizerBase
+{
+public:
+	FChaosVDGTAccelerationStructureDataComponentVisualizer();
+
+	virtual void RegisterVisualizerMenus() override;
+
+	virtual void DrawVisualization(const UActorComponent* Component, const FSceneView* View, FPrimitiveDrawInterface* PDI) override;
+
+	virtual bool CanHandleClick(const HChaosVDComponentVisProxy& VisProxy) override;
+
+protected:
+
+	void DrawAABBTree(const FSceneView* View, FPrimitiveDrawInterface* PDI, const FChaosGTAccelerationStructureVisualizationDataContext& VisualizationContext, const TSharedRef<FChaosVDAABBTreeDataWrapper>& AABBTreeData);
+	void DrawAABBTreeNode(const FSceneView* View, FPrimitiveDrawInterface* PDI, const FChaosGTAccelerationStructureVisualizationDataContext& VisualizationContext, const TSharedRef<FChaosVDAABBTreeDataWrapper>& AABBTreeData, const FChaosVDAABBTreeNodeDataWrapper& AABBTreeNodeData, float Thickness, int32 CurrentTreeLevel = 1);
+	void DrawAABBTreeArrayLeaf(const FSceneView* View, FPrimitiveDrawInterface* PDI, const FChaosGTAccelerationStructureVisualizationDataContext& VisualizationContext, const FChaosVDAABBTreeLeafDataWrapper& AABBTreeArrayLeafData, const TSharedRef<FChaosVDAABBTreeDataWrapper>& AABBTreeData, float Thickness);
+};
diff --git a/Engine/Plugins/ChaosVD/Source/ChaosVD/Private/Visualizers/ChaosVDParticleDataComponentVisualizer.cpp b/Engine/Plugins/ChaosVD/Source/ChaosVD/Private/Visualizers/ChaosVDParticleDataComponentVisualizer.cpp
--- a/Engine/Plugins/ChaosVD/Source/ChaosVD/Private/Visualizers/ChaosVDParticleDataComponentVisualizer.cpp
+++ b/Engine/Plugins/ChaosVD/Source/ChaosVD/Private/Visualizers/ChaosVDParticleDataComponentVisualizer.cpp
@@ -185,7 +185,7 @@ bool FChaosVDParticleDataComponentVisualizer::CanHandleClick(const HChaosVDCompo
 	return VisProxy.DataSelectionHandle && VisProxy.DataSelectionHandle->IsA<FChaosVDParticleDataWrapper>();
 }
 
-bool FChaosVDParticleDataComponentVisualizer::SelectVisualizedData(const HChaosVDComponentVisProxy& VisProxy, const TSharedRef<FChaosVDScene>& InCVDScene)
+bool FChaosVDParticleDataComponentVisualizer::SelectVisualizedData(const HChaosVDComponentVisProxy& VisProxy, const TSharedRef<FChaosVDScene>& InCVDScene, const TSharedRef<SChaosVDMainTab>& InMainTabToolkitHost)
 {
 	bool bHandled = false;
 	const UChaosVDParticleDataComponent* ParticleDataComponent = Cast<UChaosVDParticleDataComponent>(VisProxy.Component.Get());
diff --git a/Engine/Plugins/ChaosVD/Source/ChaosVD/Private/Visualizers/ChaosVDParticleDataComponentVisualizer.h b/Engine/Plugins/ChaosVD/Source/ChaosVD/Private/Visualizers/ChaosVDParticleDataComponentVisualizer.h
--- a/Engine/Plugins/ChaosVD/Source/ChaosVD/Private/Visualizers/ChaosVDParticleDataComponentVisualizer.h
+++ b/Engine/Plugins/ChaosVD/Source/ChaosVD/Private/Visualizers/ChaosVDParticleDataComponentVisualizer.h
@@ -43,7 +43,7 @@ public:
 
 	virtual bool CanHandleClick(const HChaosVDComponentVisProxy& VisProxy) override;
 
-	virtual bool SelectVisualizedData(const HChaosVDComponentVisProxy& VisProxy, const TSharedRef<FChaosVDScene>& InCVDScene) override;
+	virtual bool SelectVisualizedData(const HChaosVDComponentVisProxy& VisProxy, const TSharedRef<FChaosVDScene>& InCVDScene, const TSharedRef<SChaosVDMainTab>& InMainTabToolkitHost) override;
 
 protected:
 	
diff --git a/Engine/Plugins/ChaosVD/Source/ChaosVD/Private/Widgets/SChaosVDDetailsView.cpp b/Engine/Plugins/ChaosVD/Source/ChaosVD/Private/Widgets/SChaosVDDetailsView.cpp
--- a/Engine/Plugins/ChaosVD/Source/ChaosVD/Private/Widgets/SChaosVDDetailsView.cpp
+++ b/Engine/Plugins/ChaosVD/Source/ChaosVD/Private/Widgets/SChaosVDDetailsView.cpp
@@ -38,6 +38,15 @@ void SChaosVDDetailsView::Construct(const FArguments& InArgs)
 	];
 }
 
+void SChaosVDDetailsView::SetSelectedStruct(const TSharedPtr<FStructOnScope>& NewStruct)
+{
+	// Clear the object selection view, as now we will have a struct view active
+	SetSelectedObject(nullptr);
+	
+	CurrentStructInView = NewStruct;
+	StructDetailsView->SetStructureData(NewStruct);
+}
+
 TSharedPtr<IDetailsView> SChaosVDDetailsView::CreateObjectDetailsView()
 {
 	FPropertyEditorModule& PropertyEditorModule = FModuleManager::GetModuleChecked<FPropertyEditorModule>("PropertyEditor");
@@ -81,6 +90,7 @@ void SChaosVDDetailsView::SetSelectedObject(UObject* NewObject)
 {
 	// Even if the object is not valid, clear any active structure view
 	StructDetailsView->SetStructureData(nullptr);
+	CurrentStructInView = nullptr;
 
 	if (DetailsView->IsLocked())
 	{
diff --git a/Engine/Plugins/ChaosVD/Source/ChaosVD/Private/Widgets/SChaosVDDetailsView.h b/Engine/Plugins/ChaosVD/Source/ChaosVD/Private/Widgets/SChaosVDDetailsView.h
--- a/Engine/Plugins/ChaosVD/Source/ChaosVD/Private/Widgets/SChaosVDDetailsView.h
+++ b/Engine/Plugins/ChaosVD/Source/ChaosVD/Private/Widgets/SChaosVDDetailsView.h
@@ -38,6 +38,8 @@ public:
 	template<typename TStruct>
 	void SetSelectedStruct(TStruct* NewStruct);
 
+	void SetSelectedStruct(const TSharedPtr<FStructOnScope>& NewStruct);
+
 protected:
 
 	TSharedPtr<IDetailsView> CreateObjectDetailsView();
@@ -58,8 +60,6 @@ protected:
 template <typename TStruct>
 void SChaosVDDetailsView::SetSelectedStruct(TStruct* NewStruct)
 {
-	SetSelectedObject(nullptr);
-
 	TSharedPtr<FStructOnScope> StructDataView = nullptr;
 
 	if (NewStruct)
@@ -67,7 +67,5 @@ void SChaosVDDetailsView::SetSelectedStruct(TStruct* NewStruct)
 		StructDataView = MakeShared<FStructOnScope>(TStruct::StaticStruct(), reinterpret_cast<uint8*>(NewStruct));
 	}
 
-	CurrentStructInView = StructDataView;
-
-	StructDetailsView->SetStructureData(StructDataView);
+	SetSelectedStruct(StructDataView);
 }
diff --git a/Engine/Plugins/ChaosVD/Source/ChaosVD/Private/Widgets/SChaosVDMainTab.cpp b/Engine/Plugins/ChaosVD/Source/ChaosVD/Private/Widgets/SChaosVDMainTab.cpp
--- a/Engine/Plugins/ChaosVD/Source/ChaosVD/Private/Widgets/SChaosVDMainTab.cpp
+++ b/Engine/Plugins/ChaosVD/Source/ChaosVD/Private/Widgets/SChaosVDMainTab.cpp
@@ -29,6 +29,7 @@
 #include "StatusBarSubsystem.h"
 #include "ToolMenu.h"
 #include "ToolMenus.h"
+#include "Components/ChaosVDGTAccelerationStructuresDataComponent.h"
 #include "Components/ChaosVDParticleDataComponent.h"
 #include "Components/ChaosVDSolverCharacterGroundConstraintDataComponent.h"
 #include "Components/ChaosVDSolverCollisionDataComponent.h"
@@ -37,6 +38,7 @@
 #include "Styling/ToolBarStyle.h"
 #include "Trace/ChaosVDTraceManager.h"
 #include "Visualizers/ChaosVDCharacterGroundConstraintsDataComponentVisualizer.h"
+#include "Visualizers/ChaosVDGTAccelerationStructureDataComponentVisualizer.h"
 #include "Visualizers/ChaosVDJointConstraintsDataComponentVisualizer.h"
 #include "Visualizers/ChaosVDParticleDataComponentVisualizer.h"
 #include "Visualizers/ChaosVDSceneQueryDataComponentVisualizer.h"
@@ -64,6 +66,7 @@ void SChaosVDMainTab::Construct(const FArguments& InArgs, TSharedPtr<FChaosVDEng
 	RegisterComponentVisualizer(UChaosVDParticleDataComponent::StaticClass()->GetFName(), MakeShared<FChaosVDParticleDataComponentVisualizer>());
 	RegisterComponentVisualizer(UChaosVDSolverJointConstraintDataComponent::StaticClass()->GetFName(), MakeShared<FChaosVDJointConstraintsDataComponentVisualizer>());
 	RegisterComponentVisualizer(UChaosVDSolverCharacterGroundConstraintDataComponent::StaticClass()->GetFName(), MakeShared<FChaosVDCharacterGroundConstraintDataComponentVisualizer>());
+	RegisterComponentVisualizer(UChaosVDGTAccelerationStructuresDataComponent::StaticClass()->GetFName(), MakeShared<FChaosVDGTAccelerationStructureDataComponentVisualizer>());
 
 	TabManager = FGlobalTabmanager::Get()->NewTabManager(InArgs._OwnerTab.ToSharedRef()).ToSharedPtr();
 
diff --git a/Engine/Source/Runtime/Experimental/Chaos/Private/ChaosVisualDebugger/ChaosVDDataWrapperUtils.cpp b/Engine/Source/Runtime/Experimental/Chaos/Private/ChaosVisualDebugger/ChaosVDDataWrapperUtils.cpp
--- a/Engine/Source/Runtime/Experimental/Chaos/Private/ChaosVisualDebugger/ChaosVDDataWrapperUtils.cpp
+++ b/Engine/Source/Runtime/Experimental/Chaos/Private/ChaosVisualDebugger/ChaosVDDataWrapperUtils.cpp
@@ -2,6 +2,8 @@
 
 #include "ChaosVisualDebugger/ChaosVDDataWrapperUtils.h"
 
+#include "Chaos/SpatialAccelerationCollection.h"
+
 #if WITH_CHAOS_VISUAL_DEBUGGER
 
 #include "Chaos/PBDJointConstraints.h"
@@ -17,23 +19,6 @@
 #include "DataWrappers/ChaosVDParticleDataWrapper.h"
 #include "Math/UnitConversion.h"
 
-namespace Chaos::VisualDebugger::Utils
-{
-	template<typename InType,typename OutType, int32 Size, typename TransformT>
-	void TransformStaticArray(const InType (&In)[Size], OutType (&Out)[Size], TransformT Trans)
-	{
-		for (int32 Index = 0; Index < Size; Index++)
-		{
-			Out[Index] = Invoke(Trans, In[Index]);
-		}
-	}
-
-	inline FTransform ConvertToFTransform(const FRigidTransform3& InChaosTransform)
-	{
-		return InChaosTransform;
-	}
-}
-
 void FChaosVDDataWrapperUtils::CopyManifoldPointsToDataWrapper(const Chaos::FManifoldPoint& InCopyFrom, FChaosVDManifoldPoint& OutCopyTo)
 {
 	OutCopyTo.bDisabled = InCopyFrom.Flags.bDisabled;
@@ -456,4 +441,110 @@ void FChaosVDDataWrapperUtils::CopyShapeDataToWrapper(const Chaos::FShapeInstanc
 	OutCopyTo.QueryData.Word2 = CollisionData.QueryData.Word2;
 	OutCopyTo.QueryData.Word3 = CollisionData.QueryData.Word3;
 }
+
+void FChaosVDDataWrapperUtils::BuildDataWrapperFromAABBStructure(const Chaos::ISpatialAccelerationCollection<Chaos::FAccelerationStructureHandle, Chaos::FReal, 3>* SceneAccelerationStructures, int32 OwnerSolverID, TArray<FChaosVDAABBTreeDataWrapper>& OutAABBTrees)
+{
+	using namespace Chaos;
+
+	if (!SceneAccelerationStructures)
+	{
+		return;
+	}
+
+	ISpatialAccelerationCollection<FAccelerationStructureHandle, FReal, 3>* MutableSceneAccelerationStructures = const_cast<ISpatialAccelerationCollection<Chaos::FAccelerationStructureHandle, Chaos::FReal, 3>*>(SceneAccelerationStructures);
+
+	TArray<FSpatialAccelerationIdx> SpatialIndices = MutableSceneAccelerationStructures->GetAllSpatialIndices();
+	OutAABBTrees.Reserve(SpatialIndices.Num());
+	for (const FSpatialAccelerationIdx SpatialIndex : SpatialIndices)
+	{
+		const ISpatialAcceleration<FAccelerationStructureHandle, FReal, 3>* AccelerationStructure = MutableSceneAccelerationStructures->GetSubstructure(SpatialIndex);
+
+		FChaosVDAABBTreeDataWrapper AABBTreDataWrapper;
+		AABBTreDataWrapper.SolverId = OwnerSolverID;
+		
+		if (const TAABBTree<FAccelerationStructureHandle, TAABBTreeLeafArray<FAccelerationStructureHandle>>* AABBTree = AccelerationStructure->template As<TAABBTree<FAccelerationStructureHandle, TAABBTreeLeafArray<FAccelerationStructureHandle>>>())
+		{
+			BuildDataWrapperFromAABBStructure(*AABBTree, AABBTreDataWrapper);
+			OutAABBTrees.Emplace(AABBTreDataWrapper);
+		}
+		else if (const TAABBTree<FAccelerationStructureHandle, TBoundingVolume<FAccelerationStructureHandle>>* AABBTreeBV = AccelerationStructure->template As<TAABBTree<FAccelerationStructureHandle, TBoundingVolume<FAccelerationStructureHandle>>>())
+		{
+			BuildDataWrapperFromAABBStructure(*AABBTreeBV, AABBTreDataWrapper);
+			OutAABBTrees.Emplace(AABBTreDataWrapper);
+		}
+	}
+}
+
+void FChaosVDDataWrapperUtils::AddTreeLeaves(const TConstArrayView<Chaos::TAABBTreeLeafArray<Chaos::FAccelerationStructureHandle>>& LeavesContainer, FChaosVDAABBTreeDataWrapper& InOutAABBTreeWrapper)
+{
+	using namespace Chaos;
+
+	InOutAABBTreeWrapper.LeavesNum = LeavesContainer.Num();
+	InOutAABBTreeWrapper.TreeArrayLeafs.Reserve(InOutAABBTreeWrapper.LeavesNum);
+	for (const TAABBTreeLeafArray<FAccelerationStructureHandle>& TreeArrayLeaf : LeavesContainer)
+	{
+		FChaosVDAABBTreeLeafDataWrapper CVDLeaf;
+		CVDLeaf.Elements.Reserve(TreeArrayLeaf.Elems.Num());
+			
+		CVDLeaf.Bounds = ConvertToFBox(TreeArrayLeaf.GetBounds());
+
+		for (const TPayloadBoundsElement<FAccelerationStructureHandle, FReal>& Elem : TreeArrayLeaf.Elems)
+		{
+			FChaosVDAABBTreePayloadBoundsElement CVELeafElement;
+			CVELeafElement.ParticleIndex = GetUniqueIdx(Elem).Idx;
+			CVELeafElement.Bounds = ConvertToFBox(Elem.Bounds);
+
+			CVELeafElement.MarkAsValid();
+	
+			CVDLeaf.Elements.Emplace(CVELeafElement);
+		}
+			
+		CVDLeaf.MarkAsValid();
+
+		InOutAABBTreeWrapper.TreeArrayLeafs.Emplace(MoveTemp(CVDLeaf));
+	}
+}
+
+void FChaosVDDataWrapperUtils::AddTreeLeaves(const TConstArrayView<Chaos::TBoundingVolume<Chaos::FAccelerationStructureHandle>>& LeavesContainer, FChaosVDAABBTreeDataWrapper& InOutAABBTreeWrapper)
+{
+	using namespace Chaos;
+
+	InOutAABBTreeWrapper.LeavesNum = LeavesContainer.Num();
+	InOutAABBTreeWrapper.BoundingVolumeLeafs.Reserve(InOutAABBTreeWrapper.LeavesNum);
+
+	for (const TBoundingVolume<FAccelerationStructureHandle>& BoundingVolumeLeaf : LeavesContainer)
+	{
+		TVector<int32, 3> ElementCounts = BoundingVolumeLeaf.MElements.Counts();
+
+		FChaosVDBoundingVolumeDataWrapper CVDBoundingVolume;
+		CVDBoundingVolume.MElementsCounts = FIntVector3(ElementCounts.X, ElementCounts.Y, ElementCounts.Z);
+
+		CVDBoundingVolume.MaxPayloadBounds = BoundingVolumeLeaf.MaxPayloadBounds;
+
+		typedef TBoundingVolume<FAccelerationStructureHandle>::FCellElement FCellElement;
+
+		TConstArrayView<TArray<FCellElement>> ElementsFlatArray = MakeArrayView(BoundingVolumeLeaf.MElements.GetData(), ElementCounts.X * ElementCounts.Y * ElementCounts.Z);
+			
+		for (const TArray<FCellElement>& InCellElements : ElementsFlatArray)
+		{
+			TArray<FChaosVDBVCellElementDataWrapper> CellDataWrapper;
+			Algo::Transform(InCellElements, CellDataWrapper, [](const FCellElement& InElement)
+			{
+				FChaosVDBVCellElementDataWrapper CVDCellElement;
+				CVDCellElement.Bounds = ConvertToFBox(InElement.Bounds);
+				CVDCellElement.ParticleIndex = GetUniqueIdx(InElement.Payload).Idx;
+				CVDCellElement.StartIdx = FIntVector3(InElement.StartIdx.X, InElement.StartIdx.Y, InElement.StartIdx.Z);
+				CVDCellElement.EndIdx = FIntVector3(InElement.EndIdx.X, InElement.EndIdx.Y, InElement.EndIdx.Z);
+				CVDCellElement.MarkAsValid();
+	
+				return CVDCellElement;
+			});
+		}
+
+		CVDBoundingVolume.MarkAsValid();
+		
+		InOutAABBTreeWrapper.BoundingVolumeLeafs.Emplace(MoveTemp(CVDBoundingVolume));
+	}
+}
+
 #endif //WITH_CHAOS_VISUAL_DEBUGGER
diff --git a/Engine/Source/Runtime/Experimental/Chaos/Private/ChaosVisualDebugger/ChaosVDOptionalDataChannel.cpp b/Engine/Source/Runtime/Experimental/Chaos/Private/ChaosVisualDebugger/ChaosVDOptionalDataChannel.cpp
--- a/Engine/Source/Runtime/Experimental/Chaos/Private/ChaosVisualDebugger/ChaosVDOptionalDataChannel.cpp
+++ b/Engine/Source/Runtime/Experimental/Chaos/Private/ChaosVisualDebugger/ChaosVDOptionalDataChannel.cpp
@@ -85,5 +85,6 @@ CVD_DEFINE_OPTIONAL_DATA_CHANNEL(EvolutionEnd, EChaosVDDataChannelInitialization
 CVD_DEFINE_OPTIONAL_DATA_CHANNEL(SceneQueries, EChaosVDDataChannelInitializationFlags::StartEnabled | EChaosVDDataChannelInitializationFlags::CanChangeEnabledState);
 CVD_DEFINE_OPTIONAL_DATA_CHANNEL(JointConstraints, EChaosVDDataChannelInitializationFlags::CanChangeEnabledState);
 CVD_DEFINE_OPTIONAL_DATA_CHANNEL(CharacterGroundConstraints, EChaosVDDataChannelInitializationFlags::CanChangeEnabledState);
+CVD_DEFINE_OPTIONAL_DATA_CHANNEL(AccelerationStructures, EChaosVDDataChannelInitializationFlags::CanChangeEnabledState);
 
 #endif
diff --git a/Engine/Source/Runtime/Experimental/Chaos/Private/ChaosVisualDebugger/ChaosVisualDebuggerTrace.cpp b/Engine/Source/Runtime/Experimental/Chaos/Private/ChaosVisualDebugger/ChaosVisualDebuggerTrace.cpp
--- a/Engine/Source/Runtime/Experimental/Chaos/Private/ChaosVisualDebugger/ChaosVisualDebuggerTrace.cpp
+++ b/Engine/Source/Runtime/Experimental/Chaos/Private/ChaosVisualDebugger/ChaosVisualDebuggerTrace.cpp
@@ -7,6 +7,7 @@
 #include "Chaos/Character/CharacterGroundConstraintContainer.h"
 #include "Chaos/Framework/PhysicsSolverBase.h"
 #include "Chaos/ImplicitObject.h"
+#include "Chaos/ISpatialAccelerationCollection.h"
 #include "Chaos/ParticleHandle.h"
 #include "Chaos/PBDCollisionConstraints.h"
 #include "Chaos/PBDJointConstraints.h"
@@ -796,6 +797,38 @@ void FChaosVisualDebuggerTrace::TraceSceneQueryVisit(FChaosVDQueryVisitStep&& In
 	TraceBinaryData(TLSDataBuffer.BufferRef, FChaosVDQueryVisitStep::WrapperTypeName);
 }
 
+void FChaosVisualDebuggerTrace::TraceSceneAccelerationStructures(const Chaos::ISpatialAccelerationCollection<Chaos::FAccelerationStructureHandle, Chaos::FReal, 3>* InAccelerationCollection)
+{
+	using namespace Chaos::VisualDebugger::Utils;
+
+	if (!IsTracing())
+	{
+		return;
+	}
+
+	if (!InAccelerationCollection)
+	{
+		return;
+	}
+	
+	const FChaosVDContext* CVDContextData = FChaosVDThreadContext::Get().GetCurrentContext();
+	if (!IsContextEnabledAndValid(CVDContextData))
+	{
+		return;
+	}
+
+	TArray<FChaosVDAABBTreeDataWrapper> AABBTreeDataWrappers;
+	FChaosVDDataWrapperUtils::BuildDataWrapperFromAABBStructure(InAccelerationCollection, CVDContextData->Id, AABBTreeDataWrappers);
+
+	for (FChaosVDAABBTreeDataWrapper& DataWrapper : AABBTreeDataWrappers)
+	{
+		FChaosVDScopedTLSBufferAccessor TLSDataBuffer;
+		Chaos::VisualDebugger::WriteDataToBuffer(TLSDataBuffer.BufferRef, DataWrapper);
+
+		TraceBinaryData(TLSDataBuffer.BufferRef, FChaosVDAABBTreeDataWrapper::WrapperTypeName);
+	}
+}
+
 bool FChaosVisualDebuggerTrace::IsTracing()
 {
 	return bIsTracing;
diff --git a/Engine/Source/Runtime/Experimental/Chaos/Public/Chaos/AABBTree.h b/Engine/Source/Runtime/Experimental/Chaos/Public/Chaos/AABBTree.h
--- a/Engine/Source/Runtime/Experimental/Chaos/Public/Chaos/AABBTree.h
+++ b/Engine/Source/Runtime/Experimental/Chaos/Public/Chaos/AABBTree.h
@@ -72,6 +72,8 @@ struct FAABBTimeSliceCVars
 	static CHAOS_API FAutoConsoleVariableRef CVarMinDataChunkToProcessBetweenTimeChecks;
 };
 
+class FChaosVDDataWrapperUtils;
+
 namespace Chaos
 {
 
@@ -526,6 +528,8 @@ struct TAABBTreeLeafArray : public TBoundsWrapperHelper<TPayloadType, T, bComput
 
 	/** Flag on the leaf to know if it has been updated */
 	bool bDirtyLeaf = false;
+
+	friend ::FChaosVDDataWrapperUtils;
 };
 
 template <typename TPayloadType, bool bComputeBounds, typename T>
@@ -607,6 +611,8 @@ public:
 	}
 private:
 	int32 NumOfValidElements = 0;
+
+	friend ::FChaosVDDataWrapperUtils;
 };
 
 template <typename LeafType>
@@ -3775,9 +3781,10 @@ private:
 	{
 		if (InInterface)
 		{
-			if (Nodes.Num() > 0)
+			int32 RootNodeIndex = bDynamicTree ? RootNode : 0;
+			if (Nodes.Num() > 0 && Nodes.IsValidIndex(RootNodeIndex))
 			{
-				Nodes[0].DebugDraw(*InInterface, Nodes, { 1.f, 1.f, 1.f }, 5.f);
+				Nodes[RootNodeIndex].DebugDraw(*InInterface, Nodes, { 1.f, 1.f, 1.f }, 5.f);
 			}
 			for (int LeafIndex = 0; LeafIndex < Leaves.Num(); LeafIndex++)
 			{
@@ -3968,6 +3975,7 @@ private:
 	using FNodeIndexAndCost = TTuple<FNode&, int32, FReal>;
 	TArray<FNodeIndexAndCost> PriorityQ;	
 
+	friend ::FChaosVDDataWrapperUtils;
 };
 
 template<typename TPayloadType, typename TLeafType, bool bMutable, typename T, typename StorageTraits>
diff --git a/Engine/Source/Runtime/Experimental/Chaos/Public/Chaos/BoundingVolume.h b/Engine/Source/Runtime/Experimental/Chaos/Public/Chaos/BoundingVolume.h
--- a/Engine/Source/Runtime/Experimental/Chaos/Public/Chaos/BoundingVolume.h
+++ b/Engine/Source/Runtime/Experimental/Chaos/Public/Chaos/BoundingVolume.h
@@ -25,6 +25,8 @@
 //#include "Engine/World.h"
 //#include "DrawDebugHelpers.h"
 
+class FChaosVDDataWrapperUtils;
+
 template <typename T, bool>
 struct TSpatialAccelerationTraits
 {
@@ -1356,6 +1358,8 @@ private:
 	TArrayAsMap<TPayloadType, FPayloadInfo> MPayloadInfo;
 	T MaxPayloadBounds;
 	bool bIsEmpty;
+
+	friend ::FChaosVDDataWrapperUtils;
 };
 
 template<typename TPayloadType, class T, int d>
diff --git a/Engine/Source/Runtime/Experimental/Chaos/Public/Chaos/SpatialAccelerationCollection.h b/Engine/Source/Runtime/Experimental/Chaos/Public/Chaos/SpatialAccelerationCollection.h
--- a/Engine/Source/Runtime/Experimental/Chaos/Public/Chaos/SpatialAccelerationCollection.h
+++ b/Engine/Source/Runtime/Experimental/Chaos/Public/Chaos/SpatialAccelerationCollection.h
@@ -10,6 +10,8 @@
 
 #include <tuple>
 
+class FChaosVDDataWrapperUtils;
+
 namespace Chaos
 {
 
@@ -806,6 +808,8 @@ private:
 	BucketType Buckets[MaxBuckets];
 	TSpatialTypeTuple< TSpatialAccelerationTypes...> Types; // Have buckets of acceleration structure pointers
 	static constexpr uint32 NumTypes = sizeof...(TSpatialAccelerationTypes);
+
+	friend ::FChaosVDDataWrapperUtils;
 };
 
 
diff --git a/Engine/Source/Runtime/Experimental/Chaos/Public/ChaosVisualDebugger/ChaosVDDataWrapperUtils.h b/Engine/Source/Runtime/Experimental/Chaos/Public/ChaosVisualDebugger/ChaosVDDataWrapperUtils.h
--- a/Engine/Source/Runtime/Experimental/Chaos/Public/ChaosVisualDebugger/ChaosVDDataWrapperUtils.h
+++ b/Engine/Source/Runtime/Experimental/Chaos/Public/ChaosVisualDebugger/ChaosVDDataWrapperUtils.h
@@ -1,13 +1,16 @@
 ﻿// Copyright Epic Games, Inc. All Rights Reserved.
 
 #pragma once
-#include "Chaos/ShapeInstanceFwd.h"
 
 #if WITH_CHAOS_VISUAL_DEBUGGER
+#include "Chaos/AABBTree.h"
 #include "Chaos/Collision/ContactPoint.h"
-#include "Chaos/ParticleHandleFwd.h"
+#include "Chaos/ParticleHandle.h"
+#include "Chaos/ShapeInstanceFwd.h"
 #include "DataWrappers/ChaosVDCollisionDataWrappers.h"
+#include "DataWrappers/ChaosVDAccelerationStructureDataWrappers.h"
 
+struct FChaosVDAABBTreeDataWrapper;
 struct FChaosVDCharacterGroundConstraint;
 struct FChaosVDJointConstraint;
 
@@ -32,6 +35,29 @@ struct FChaosVDConstraint;
 struct FChaosVDParticlePairMidPhase;
 struct FChaosVDParticleDataWrapper;
 
+namespace Chaos::VisualDebugger::Utils
+{
+	template<typename InType,typename OutType, int32 Size, typename TransformT>
+	void TransformStaticArray(const InType (&In)[Size], OutType (&Out)[Size], TransformT Trans)
+	{
+		for (int32 Index = 0; Index < Size; Index++)
+		{
+			Out[Index] = Invoke(Trans, In[Index]);
+		}
+	}
+
+	template<typename InType,typename OutType, int32 Size>
+	void CopyStaticArray(const InType (&In)[Size], OutType (&Out)[Size])
+	{
+		FMemory::Memcpy(Out, In, Size * sizeof(InType));
+	}
+
+	inline FTransform ConvertToFTransform(const FRigidTransform3& InChaosTransform)
+	{
+		return InChaosTransform;
+	}
+}
+
 /**
  * Helper class used to build Chaos Visual Debugger data wrappers, without directly referencing chaos' types in them.
  *
@@ -71,7 +97,60 @@ private:
 	static FVector ConvertToFVector(const Chaos::FVec3f& VectorRef) { return FVector(VectorRef); }
 
 	static void CopyShapeDataToWrapper(const Chaos::FShapeInstancePtr& ShapeDataPtr, FChaosVDShapeCollisionData& OutCopyTo);
+	
+	template<typename LeafType>
+	static void BuildDataWrapperFromAABBStructure(const Chaos::TAABBTree<Chaos::FAccelerationStructureHandle, LeafType>& AABBTree, FChaosVDAABBTreeDataWrapper& OutAABBTreeWrapper);
+
+	/** Converts a Chaos::FVec3 to a FVector. It is worth notice that FVector is double precision and FVec3 is single */
+	static FBox ConvertToFBox(const Chaos::TAABB<Chaos::FReal, 3>& Bounds)
+	{
+		FBox ConvertedBox(Bounds.Min(), Bounds.Max());
+		return ConvertedBox;
+	}
+
+	static void BuildDataWrapperFromAABBStructure(const Chaos::ISpatialAccelerationCollection<Chaos::FAccelerationStructureHandle, Chaos::FReal, 3>*, int32 OwnerSolverID, TArray<FChaosVDAABBTreeDataWrapper>& OutAABBTrees);
+
+	static void AddTreeLeaves(const TConstArrayView<Chaos::TAABBTreeLeafArray<Chaos::FAccelerationStructureHandle>>& LeavesContainer, FChaosVDAABBTreeDataWrapper& InOutAABBTreeWrapper);
+	static void AddTreeLeaves(const TConstArrayView<Chaos::TBoundingVolume<Chaos::FAccelerationStructureHandle>>& LeavesContainer, FChaosVDAABBTreeDataWrapper& InOutAABBTreeWrapper);
 
 	friend FChaosVisualDebuggerTrace;
 };
+
+template<typename LeafType>
+void FChaosVDDataWrapperUtils::BuildDataWrapperFromAABBStructure(const Chaos::TAABBTree<Chaos::FAccelerationStructureHandle, LeafType>& AABBTree, FChaosVDAABBTreeDataWrapper& OutAABBTreeWrapper)
+{
+	using namespace Chaos;
+
+	OutAABBTreeWrapper.MaxTreeDepth = AABBTree.MaxTreeDepth;
+	OutAABBTreeWrapper.MaxChildrenInLeaf = AABBTree.MaxTreeDepth;
+	OutAABBTreeWrapper.MaxPayloadBounds = AABBTree.MaxPayloadBounds;
+	OutAABBTreeWrapper.RootNodeIndex = AABBTree.RootNode;
+	OutAABBTreeWrapper.bDynamicTree = AABBTree.bDynamicTree;
+
+	const int32 NodesToCopyNum = AABBTree.Nodes.Num();
+	OutAABBTreeWrapper.Nodes.Reserve(NodesToCopyNum);
+	OutAABBTreeWrapper.NodesNum = NodesToCopyNum;
+
+	OutAABBTreeWrapper.Type = static_cast<EChaosVDAccelerationStructureType>(AABBTree.StaticType);
+		
+	for (const typename TAABBTree<FAccelerationStructureHandle, LeafType>::FNode& Node : AABBTree.Nodes)
+	{
+		FChaosVDAABBTreeNodeDataWrapper CVDNode;
+		CVDNode.bLeaf = Node.bLeaf;
+		CVDNode.bDirtyNode = Node.bDirtyNode;
+
+		VisualDebugger::Utils::TransformStaticArray(Node.ChildrenBounds, CVDNode.ChildrenBounds, &FChaosVDDataWrapperUtils::ConvertToFBox);
+		VisualDebugger::Utils::CopyStaticArray(Node.ChildrenNodes, CVDNode.ChildrenNodes);
+
+		CVDNode.ParentNode =  Node.ParentNode;
+		CVDNode.MarkAsValid();
+
+		OutAABBTreeWrapper.Nodes.Emplace(CVDNode);
+	}
+	
+	AddTreeLeaves(MakeArrayView(AABBTree.Leaves.GetData(), AABBTree.Leaves.Num()), OutAABBTreeWrapper);
+	
+	OutAABBTreeWrapper.MarkAsValid();
+}
+
 #endif //WITH_CHAOS_VISUAL_DEBUGGER
diff --git a/Engine/Source/Runtime/Experimental/Chaos/Public/ChaosVisualDebugger/ChaosVDOptionalDataChannel.h b/Engine/Source/Runtime/Experimental/Chaos/Public/ChaosVisualDebugger/ChaosVDOptionalDataChannel.h
--- a/Engine/Source/Runtime/Experimental/Chaos/Public/ChaosVisualDebugger/ChaosVDOptionalDataChannel.h
+++ b/Engine/Source/Runtime/Experimental/Chaos/Public/ChaosVisualDebugger/ChaosVDOptionalDataChannel.h
@@ -153,6 +153,7 @@ CVD_DECLARE_OPTIONAL_DATA_CHANNEL(EvolutionEnd);							// Final particle positio
 CVD_DECLARE_OPTIONAL_DATA_CHANNEL(SceneQueries);
 CVD_DECLARE_OPTIONAL_DATA_CHANNEL(JointConstraints);
 CVD_DECLARE_OPTIONAL_DATA_CHANNEL(CharacterGroundConstraints);
+CVD_DECLARE_OPTIONAL_DATA_CHANNEL(AccelerationStructures);
 
 #else  //WITH_CHAOS_VISUAL_DEBUGGER
 
diff --git a/Engine/Source/Runtime/Experimental/Chaos/Public/ChaosVisualDebugger/ChaosVDTraceMacros.h b/Engine/Source/Runtime/Experimental/Chaos/Public/ChaosVisualDebugger/ChaosVDTraceMacros.h
--- a/Engine/Source/Runtime/Experimental/Chaos/Public/ChaosVisualDebugger/ChaosVDTraceMacros.h
+++ b/Engine/Source/Runtime/Experimental/Chaos/Public/ChaosVisualDebugger/ChaosVDTraceMacros.h
@@ -166,6 +166,17 @@
 			FChaosVisualDebuggerTrace::TraceConstraintsContainer(ContainerView);
 	#endif
 
+	#ifndef CVD_TRACE_ACCELERATION_STRUCTURES
+		#define CVD_TRACE_ACCELERATION_STRUCTURES(AccelerationStructuresCollections, SolverType, SolverRef, DataChannel) \
+			{ \
+				FChaosVDContext StartAccelerationStrutureCVDContextData; \
+				FChaosVisualDebuggerTrace::GetCVDContext<SolverType>(SolverRef, StartAccelerationStrutureCVDContextData); \
+				CVD_SCOPE_CONTEXT(StartAccelerationStrutureCVDContextData); \
+				CVD_SCOPED_DATA_CHANNEL_OVERRIDE(DataChannel) \
+				FChaosVisualDebuggerTrace::TraceSceneAccelerationStructures(AccelerationStructuresCollections); \
+			}
+	#endif
+
 #else // WITH_CHAOS_VISUAL_DEBUGGER
 
 	#ifndef CVD_TRACE_PARTICLE
@@ -262,4 +273,8 @@
 		#define CVD_TRACE_CONSTRAINTS_CONTAINER(ContainerView)
 	#endif
 
+	#ifndef CVD_TRACE_ACCELERATION_STRUCTURES
+			#define CVD_TRACE_ACCELERATION_STRUCTURES(AccelerationStructuresCollections, SolverType, SolverRef, DataChannel)
+	#endif
+
 #endif // WITH_CHAOS_VISUAL_DEBUGGER
\ No newline at end of file
diff --git a/Engine/Source/Runtime/Experimental/Chaos/Public/ChaosVisualDebugger/ChaosVisualDebuggerTrace.h b/Engine/Source/Runtime/Experimental/Chaos/Public/ChaosVisualDebugger/ChaosVisualDebuggerTrace.h
--- a/Engine/Source/Runtime/Experimental/Chaos/Public/ChaosVisualDebugger/ChaosVisualDebuggerTrace.h
+++ b/Engine/Source/Runtime/Experimental/Chaos/Public/ChaosVisualDebugger/ChaosVisualDebuggerTrace.h
@@ -28,6 +28,7 @@
 
 namespace Chaos
 {
+class FAccelerationStructureHandle;
 class FCharacterGroundConstraintContainer;
 class FPBDConstraintContainer;
 class FPBDJointConstraints;
@@ -145,6 +146,9 @@ namespace Chaos
 
 	class FPBDCollisionConstraint;
 	class FParticlePairMidPhase;
+
+	template <typename PayloadType, typename T, int d>
+	class ISpatialAccelerationCollection;
 }
 
 using FChaosVDImplicitObjectWrapper = FChaosVDImplicitObjectDataWrapper<Chaos::FImplicitObjectPtr, Chaos::FChaosArchive>;
@@ -252,6 +256,8 @@ public:
 	static CHAOS_API void TraceSceneQueryStart(const Chaos::FImplicitObject* InputGeometry, const FQuat& GeometryOrientation, const FVector& Start, const FVector& End, ECollisionChannel TraceChannel, FChaosVDCollisionQueryParams&& Params, FChaosVDCollisionResponseParams&& ResponseParams, FChaosVDCollisionObjectQueryParams&& ObjectParams, EChaosVDSceneQueryType QueryType, EChaosVDSceneQueryMode QueryMode, int32 SolverID, bool bIsRetry);
 	static CHAOS_API void TraceSceneQueryVisit(FChaosVDQueryVisitStep&& InQueryVisitData);
 
+	static CHAOS_API void TraceSceneAccelerationStructures(const Chaos::ISpatialAccelerationCollection<Chaos::FAccelerationStructureHandle, Chaos::FReal, 3>* InAccelerationCollection);
+
 	/** Returns true if the provided solver ID needs a Full Capture */
 	static bool ShouldPerformFullCapture(int32 SolverID);
 
diff --git a/Engine/Source/Runtime/Experimental/ChaosVDData/Public/ChaosVDRecording.h b/Engine/Source/Runtime/Experimental/ChaosVDData/Public/ChaosVDRecording.h
--- a/Engine/Source/Runtime/Experimental/ChaosVDData/Public/ChaosVDRecording.h
+++ b/Engine/Source/Runtime/Experimental/ChaosVDData/Public/ChaosVDRecording.h
@@ -16,6 +16,8 @@
 #include "DataWrappers/ChaosVDQueryDataWrappers.h"
 #include <atomic>
 
+#include "DataWrappers/ChaosVDAccelerationStructureDataWrappers.h"
+
 namespace Chaos::VisualDebugger
 {
 	class FChaosVDSerializableNameTable;
@@ -111,6 +113,7 @@ struct FChaosVDGameFrameData
 	TMap<FName, FChaosVDTrackedLocation> RecordedNonSolverLocationsByID;
 	TMap<FName, FChaosVDTrackedTransform> RecordedNonSolverTransformsByID;
 	TMap<int32, TSharedPtr<FChaosVDQueryDataWrapper>> RecordedSceneQueries;
+	TMap<int32, TArray<TSharedPtr<FChaosVDAABBTreeDataWrapper>>> RecordedAABBTreesBySolverID;
 };
 
 /**
diff --git a/Engine/Source/Runtime/Experimental/ChaosVisualDebugger/Private/DataWrappers/ChaosVDAccelerationStructureDataWrappers.cpp b/Engine/Source/Runtime/Experimental/ChaosVisualDebugger/Private/DataWrappers/ChaosVDAccelerationStructureDataWrappers.cpp
new file mode 100644
--- /dev/null
+++ b/Engine/Source/Runtime/Experimental/ChaosVisualDebugger/Private/DataWrappers/ChaosVDAccelerationStructureDataWrappers.cpp
@@ -0,0 +1,131 @@
+// Copyright Epic Games, Inc. All Rights Reserved.
+
+#include "DataWrappers/ChaosVDAccelerationStructureDataWrappers.h"
+
+#include "DataWrappers/ChaosVDDataSerializationMacros.h"
+
+
+bool FChaosVDBVCellElementDataWrapper::Serialize(FArchive& Ar)
+{
+	Ar << bHasValidData;
+
+	if (!bHasValidData)
+	{
+		return !Ar.IsError();
+	}
+
+	Ar << Bounds;
+	Ar << ParticleIndex;
+	Ar << StartIdx;
+	Ar << EndIdx;
+
+	return !Ar.IsError();
+}
+
+bool FChaosVDBoundingVolumeDataWrapper::Serialize(FArchive& Ar)
+{
+	Ar << bHasValidData;
+
+	if (!bHasValidData)
+	{
+		return !Ar.IsError();
+	}
+
+	Ar << SolverId;
+	Ar << Type;
+	Ar << MElementsCounts;
+	Ar << MElements;
+	Ar << MaxPayloadBounds;
+
+	return !Ar.IsError();
+}
+
+bool FChaosVDAABBTreeNodeDataWrapper::Serialize(FArchive& Ar)
+{
+	Ar << bHasValidData;
+
+	if (!bHasValidData)
+	{
+		return !Ar.IsError();
+	}
+
+	CVD_SERIALIZE_STATIC_ARRAY(Ar, ChildrenBounds);
+	CVD_SERIALIZE_STATIC_ARRAY(Ar, ChildrenNodes);
+	
+	Ar << ParentNode;
+
+	EChaosVDAABBTreeNodeFlags PackedFlags = EChaosVDAABBTreeNodeFlags::None;
+
+	// Note: When the UI is done to show the enums in the same way we show these bitfiedls as read only booleans in CVD, we can remove all
+	// this macro boilerplate and just serialize the the flags directly
+	if (Ar.IsLoading())
+	{
+		Ar << PackedFlags;
+		CVD_UNPACK_BITFIELD_DATA(bLeaf, PackedFlags, EChaosVDAABBTreeNodeFlags::IsLeaf);
+		CVD_UNPACK_BITFIELD_DATA(bDirtyNode, PackedFlags, EChaosVDAABBTreeNodeFlags::IsDirty);
+	}
+	else
+	{
+		CVD_PACK_BITFIELD_DATA(bLeaf, PackedFlags,  EChaosVDAABBTreeNodeFlags::IsLeaf);
+		CVD_PACK_BITFIELD_DATA(bDirtyNode, PackedFlags, EChaosVDAABBTreeNodeFlags::IsDirty);
+		Ar << PackedFlags;
+	}
+
+	return !Ar.IsError();
+}
+
+bool FChaosVDAABBTreePayloadBoundsElement::Serialize(FArchive& Ar)
+{
+	Ar << bHasValidData;
+
+	if (!bHasValidData)
+	{
+		return !Ar.IsError();
+	}
+
+	Ar << ParticleIndex;
+	Ar << Bounds;
+
+	return !Ar.IsError();
+}
+
+
+bool FChaosVDAABBTreeLeafDataWrapper::Serialize(FArchive& Ar)
+{
+	Ar << bHasValidData;
+
+	if (!bHasValidData)
+	{
+		return !Ar.IsError();
+	}
+
+	Ar << Elements;
+	Ar << Bounds;
+
+	return !Ar.IsError();
+}
+
+bool FChaosVDAABBTreeDataWrapper::Serialize(FArchive& Ar)
+{
+	Ar << bHasValidData;
+
+	if (!bHasValidData)
+	{
+		return !Ar.IsError();
+	}
+
+	Ar << SolverId;
+	Ar << RootNodeIndex;
+	Ar << bDynamicTree;
+	Ar << Nodes;
+	Ar << BoundingVolumeLeafs;
+	Ar << TreeArrayLeafs;
+	Ar << MaxPayloadBounds;
+	Ar << MaxTreeDepth;
+	Ar << MaxChildrenInLeaf;
+	Ar << LeavesNum;
+	Ar << NodesNum;
+	Ar << Type;
+
+	return !Ar.IsError();
+}
diff --git a/Engine/Source/Runtime/Experimental/ChaosVisualDebugger/Public/DataWrappers/ChaosVDAccelerationStructureDataWrappers.h b/Engine/Source/Runtime/Experimental/ChaosVisualDebugger/Public/DataWrappers/ChaosVDAccelerationStructureDataWrappers.h
new file mode 100644
--- /dev/null
+++ b/Engine/Source/Runtime/Experimental/ChaosVisualDebugger/Public/DataWrappers/ChaosVDAccelerationStructureDataWrappers.h
@@ -0,0 +1,172 @@
+// Copyright Epic Games, Inc. All Rights Reserved.
+
+#pragma once
+#include "ChaosVDDataSerializationMacros.h"
+#include "ChaosVDParticleDataWrapper.h"
+
+#include "ChaosVDAccelerationStructureDataWrappers.generated.h"
+
+UENUM()
+enum class EChaosVDAABBTreeNodeFlags : uint8
+{
+	None = 0,
+	IsLeaf = 1 << 0,
+	IsDirty = 1 << 1
+};
+
+ENUM_CLASS_FLAGS(EChaosVDAABBTreeNodeFlags)
+
+UENUM()
+enum class EChaosVDAccelerationStructureType : uint32
+{
+	BoundingVolume,
+	AABBTree,
+	AABBTreeBV,
+	Collection,
+	Unknown,
+};
+
+USTRUCT()
+struct FChaosVDAccelerationStructureBase : public FChaosVDWrapperDataBase
+{
+	GENERATED_BODY()
+
+	UPROPERTY(VisibleAnywhere, Category="CVD Debug")
+	int32 SolverId = INDEX_NONE;
+
+	UPROPERTY(VisibleAnywhere, Category="Data")
+	EChaosVDAccelerationStructureType Type = EChaosVDAccelerationStructureType::Unknown;
+};
+
+USTRUCT()
+struct FChaosVDBVCellElementDataWrapper : public FChaosVDWrapperDataBase
+{
+	GENERATED_BODY()
+
+	UPROPERTY(VisibleAnywhere, Category="Data")
+	FBox Bounds;
+	UPROPERTY(VisibleAnywhere, Category="Data")
+	int32 ParticleIndex;
+
+	FIntVector3 StartIdx = FIntVector3::ZeroValue;
+	FIntVector3 EndIdx = FIntVector3::ZeroValue;
+	
+	CHAOSVDRUNTIME_API bool Serialize(FArchive& Ar);
+};
+
+CVD_IMPLEMENT_SERIALIZER(FChaosVDBVCellElementDataWrapper)
+
+USTRUCT()
+struct FChaosVDBoundingVolumeDataWrapper : public FChaosVDAccelerationStructureBase
+{
+	GENERATED_BODY()
+
+	FIntVector3 MElementsCounts = FIntVector3::ZeroValue; 
+
+	TArray<TArray<FChaosVDBVCellElementDataWrapper>> MElements;
+
+	UPROPERTY(VisibleAnywhere, Category="Settings")
+	double MaxPayloadBounds = 0.0;
+
+	CHAOSVDRUNTIME_API bool Serialize(FArchive& Ar);
+};
+
+CVD_IMPLEMENT_SERIALIZER(FChaosVDBoundingVolumeDataWrapper)
+
+USTRUCT(DisplayName="AABB Tree Node")
+struct FChaosVDAABBTreeNodeDataWrapper : public FChaosVDWrapperDataBase
+{
+	GENERATED_BODY()
+
+	UPROPERTY(VisibleAnywhere, Category="Node")
+	FBox ChildrenBounds[2] = { FBox(ForceInitToZero), FBox(ForceInitToZero) };
+	UPROPERTY(VisibleAnywhere, Category="Node")
+	int32 ChildrenNodes[2] = { INDEX_NONE, INDEX_NONE };
+	UPROPERTY(VisibleAnywhere, Category="Node")
+	int32 ParentNode = INDEX_NONE;
+	UPROPERTY(VisibleAnywhere, Category="Node")
+	uint8 bLeaf : 1 = false;
+	UPROPERTY(VisibleAnywhere, Category="Node")
+	uint8 bDirtyNode : 1 = false;
+
+	CHAOSVDRUNTIME_API bool Serialize(FArchive& Ar);
+};
+
+CVD_IMPLEMENT_SERIALIZER(FChaosVDAABBTreeNodeDataWrapper)
+
+USTRUCT()
+struct FChaosVDAABBTreePayloadBoundsElement : public FChaosVDWrapperDataBase
+{
+	GENERATED_BODY()
+
+	UPROPERTY(VisibleAnywhere, Category="Leaf")
+	int32 ParticleIndex = INDEX_NONE;
+
+	UPROPERTY(VisibleAnywhere, Category="Leaf")
+	FBox Bounds = FBox(ForceInitToZero);
+	
+	CHAOSVDRUNTIME_API bool Serialize(FArchive& Ar);
+};
+
+CVD_IMPLEMENT_SERIALIZER(FChaosVDAABBTreePayloadBoundsElement)
+
+USTRUCT(DisplayName="AABB Tree Leaf")
+struct FChaosVDAABBTreeLeafDataWrapper : public FChaosVDWrapperDataBase
+{
+	GENERATED_BODY()
+
+	UPROPERTY(VisibleAnywhere, Category="Leaf")
+	TArray<FChaosVDAABBTreePayloadBoundsElement> Elements;
+
+	UPROPERTY(VisibleAnywhere, Category="Leaf")
+	FBox Bounds = FBox(ForceInitToZero);
+	
+	CHAOSVDRUNTIME_API bool Serialize(FArchive& Ar);
+};
+
+CVD_IMPLEMENT_SERIALIZER(FChaosVDAABBTreeLeafDataWrapper)
+
+USTRUCT(DisplayName="AABB Tree Data")
+struct FChaosVDAABBTreeDataWrapper : public FChaosVDAccelerationStructureBase
+{
+	GENERATED_BODY()
+
+	inline static FStringView WrapperTypeName = TEXT("FChaosVDAABBTreeDataWrapper");
+
+	UPROPERTY(VisibleAnywhere, Category="Tree Data")
+	int32 RootNodeIndex = INDEX_NONE;
+
+	UPROPERTY(VisibleAnywhere, Category="Tree Data")
+	int32 TreeDepth = INDEX_NONE;
+
+	UPROPERTY(VisibleAnywhere, Category="Tree Data")
+	int32 NodesNum = INDEX_NONE;
+
+	UPROPERTY(VisibleAnywhere, Category="Tree Data")
+	int32 LeavesNum = INDEX_NONE;
+
+	UPROPERTY(VisibleAnywhere, Category="Tree Settings")
+	bool bDynamicTree = false;
+
+	UPROPERTY(VisibleAnywhere, Category="Tree Settings")
+	int32 MaxChildrenInLeaf = INDEX_NONE;
+	
+	UPROPERTY(VisibleAnywhere, Category="Tree Settings")
+	int32 MaxTreeDepth = INDEX_NONE;
+	
+	UPROPERTY(VisibleAnywhere, Category="Tree Settings")
+	double MaxPayloadBounds = -1.0;
+
+	TArray<FChaosVDAABBTreeNodeDataWrapper> Nodes;
+
+	TArray<FChaosVDAABBTreeLeafDataWrapper> TreeArrayLeafs;
+
+	TArray<FChaosVDBoundingVolumeDataWrapper> BoundingVolumeLeafs;
+
+	CHAOSVDRUNTIME_API bool Serialize(FArchive& Ar);
+
+	/** Returns a valid index for the root node taking ito account if this tree is dynamic or not */
+	int32 GetCorrectedRootNodeIndex() const { return bDynamicTree ? RootNodeIndex : 0;};
+};
+
+CVD_IMPLEMENT_SERIALIZER(FChaosVDAABBTreeDataWrapper)
diff --git a/Engine/Source/Runtime/Experimental/ChaosVisualDebugger/Public/DataWrappers/ChaosVDDataSerializationMacros.h b/Engine/Source/Runtime/Experimental/ChaosVisualDebugger/Public/DataWrappers/ChaosVDDataSerializationMacros.h
--- a/Engine/Source/Runtime/Experimental/ChaosVisualDebugger/Public/DataWrappers/ChaosVDDataSerializationMacros.h
+++ b/Engine/Source/Runtime/Experimental/ChaosVisualDebugger/Public/DataWrappers/ChaosVDDataSerializationMacros.h
@@ -24,4 +24,22 @@
 			Archive << Array[Index]; \
 		}\
 	}
+#endif
+
+#ifndef CVD_IMPLEMENT_SERIALIZER
+		#define CVD_IMPLEMENT_SERIALIZER(Type) \
+		inline FArchive& operator<<(FArchive& Ar, Type& Data) \
+		{\
+			Data.Serialize(Ar); \
+			return Ar; \
+		} \
+		template<>\
+		struct TStructOpsTypeTraits<Type> : public TStructOpsTypeTraitsBase2<Type> \
+		{\
+			enum\
+			{\
+				WithSerializer = true,\
+			};\
+		};\
+
 #endif
\ No newline at end of file
diff --git a/Engine/Source/Runtime/Experimental/ChaosVisualDebugger/Public/DataWrappers/ChaosVDJointDataWrappers.h b/Engine/Source/Runtime/Experimental/ChaosVisualDebugger/Public/DataWrappers/ChaosVDJointDataWrappers.h
--- a/Engine/Source/Runtime/Experimental/ChaosVisualDebugger/Public/DataWrappers/ChaosVDJointDataWrappers.h
+++ b/Engine/Source/Runtime/Experimental/ChaosVisualDebugger/Public/DataWrappers/ChaosVDJointDataWrappers.h
@@ -2,29 +2,12 @@
 
 #pragma once
 #include "ChaosVDParticleDataWrapper.h"
+#include "DataWrappers/ChaosVDDataSerializationMacros.h"
 #include "UObject/Class.h"
 #include "UObject/ObjectMacros.h"
 
 #include "ChaosVDJointDataWrappers.generated.h"
 
-#ifndef CVD_IMPLEMENT_SERIALIZER
-		#define CVD_IMPLEMENT_SERIALIZER(Type) \
-		inline FArchive& operator<<(FArchive& Ar, Type& Data) \
-		{\
-			Data.Serialize(Ar); \
-			return Ar; \
-		} \
-		template<>\
-		struct TStructOpsTypeTraits<Type> : public TStructOpsTypeTraitsBase2<Type> \
-		{\
-			enum\
-			{\
-				WithSerializer = true,\
-			};\
-		};\
-
-#endif
-
 UENUM()
 enum class EChaosVDJointReSimType
 {
diff --git a/Engine/Source/Runtime/PhysicsCore/Private/ChaosScene.cpp b/Engine/Source/Runtime/PhysicsCore/Private/ChaosScene.cpp
--- a/Engine/Source/Runtime/PhysicsCore/Private/ChaosScene.cpp
+++ b/Engine/Source/Runtime/PhysicsCore/Private/ChaosScene.cpp
@@ -492,6 +492,8 @@ void FChaosScene::EndFrame()
 		return;
 	}
 
+	CVD_TRACE_ACCELERATION_STRUCTURES(SolverAccelerationStructure, Chaos::FPhysicsSolver, *SceneSolver, CVDDC_AccelerationStructures);
+
 #if !UE_BUILD_SHIPPING
 	{
 		Chaos::AABBTreeStatistics TreeStats;
