[ChaosVD] Fix an issue where if you scrub backwards one frame, and you happen to land on a frame that is not a keyframe (but it has a generated keyframe), the generated keyframe will not be processed and the visualization will get out of sync.

diff --git a/Engine/Plugins/ChaosVD/Source/ChaosVD/Private/ChaosVDPlaybackController.cpp b/Engine/Plugins/ChaosVD/Source/ChaosVD/Private/ChaosVDPlaybackController.cpp
--- a/Engine/Plugins/ChaosVD/Source/ChaosVD/Private/ChaosVDPlaybackController.cpp
+++ b/Engine/Plugins/ChaosVD/Source/ChaosVD/Private/ChaosVDPlaybackController.cpp
@@ -127,7 +127,7 @@ void FChaosVDPlaybackController::UnloadCurrentRecording(EChaosVDUnloadRecordingF
 	bPlayedFirstFrame = false;
 }
 
-void FChaosVDPlaybackController::PlayFromClosestKeyFrame_AssumesLocked(const int32 InTrackID, const int32 FrameNumber, FChaosVDScene& InSceneToControl) const
+void FChaosVDPlaybackController::PlayFromClosestKeyFrame_AssumesLocked(const int32 InTrackID, const int32 FrameNumber, FChaosVDScene& InSceneToControl)
 {
 	if (!LoadedRecording.IsValid())
 	{
@@ -143,16 +143,31 @@ void FChaosVDPlaybackController::PlayFromClosestKeyFrame_AssumesLocked(const int
 		return;
 	}
 
-	// Instead of playing back each delta frame since the key frame, generate a new solver frame with all the deltas collapsed in one
-	// This increases the tool performance while scrubbing or live debugging if there are few keyframes
-	const int32 LastFrameToEvaluateIndex = FrameNumber - 1;
-	FChaosVDSolverFrameData CollapsedFrameData;
-	LoadedRecording->CollapseSolverFramesRange_AssumesLocked(InTrackID, KeyFrameNumber, LastFrameToEvaluateIndex, CollapsedFrameData);
+	// All keyframes should be played from stage 0 as in some scenarios we will generate a keyframe by collapsing multiple delta frames. In these frames there will be only a single "Generated" stage.
+	constexpr int32 SolverStage = 0;
 
-	if (CollapsedFrameData.SolverSteps.Num() > 0)
+	// If this frame number has keyframe data, just use it directly and save the cost of copying the data to a "collapsed keyframe"
+	if (KeyFrameNumber == FrameNumber)
 	{
-		InSceneToControl.UpdateFromRecordedStepData(InTrackID, CollapsedFrameData.SolverSteps[0], CollapsedFrameData);
+		constexpr bool bRequestingKeyFrameOnly = true;
+		if (const FChaosVDSolverFrameData* SolverFrameData = LoadedRecording->GetSolverFrameData_AssumesLocked(InTrackID, FrameNumber, bRequestingKeyFrameOnly))
+		{
+			PlaySolverStepData(InTrackID, InSceneToControl.AsShared(), *SolverFrameData, SolverStage);
+		}
+		else
+		{
+			UE_LOG(LogChaosVDEditor, Warning, TEXT("[%s] Failed to find a keyframe data for frame [%d] of track [%d]. The visualization might be out of sync until a new keyframe is played."), ANSI_TO_TCHAR(__FUNCTION__), FrameNumber, InTrackID);
+		}
+
+		return;
 	}
+
+	// If the frame number we wanted to play is not a keyframe, instead of playing back each delta frame since the key frame, generate a new solver frame with all the deltas collapsed in one
+	// This increases the tool performance while scrubbing or live debugging if there are few keyframes
+	FChaosVDSolverFrameData CollapsedFrameData;
+	LoadedRecording->CollapseSolverFramesRange_AssumesLocked(InTrackID, KeyFrameNumber, FrameNumber, CollapsedFrameData);
+
+	PlaySolverStepData(InTrackID, InSceneToControl.AsShared(), CollapsedFrameData, SolverStage);
 }
 
 void FChaosVDPlaybackController::EnqueueTrackInfoUpdate(const TSharedRef<const FChaosVDTrackInfo>& InTrackInfo, FGuid InstigatorID)
diff --git a/Engine/Plugins/ChaosVD/Source/ChaosVD/Private/ChaosVDPlaybackController.h b/Engine/Plugins/ChaosVD/Source/ChaosVD/Private/ChaosVDPlaybackController.h
--- a/Engine/Plugins/ChaosVD/Source/ChaosVD/Private/ChaosVDPlaybackController.h
+++ b/Engine/Plugins/ChaosVD/Source/ChaosVD/Private/ChaosVDPlaybackController.h
@@ -293,7 +293,7 @@ protected:
 	void HandleCurrentRecordingUpdated();
 
 	/** Finds the closest Key frame to the provided frame number, and plays all the following frames until the specified frame number (no inclusive) */
-	void PlayFromClosestKeyFrame_AssumesLocked(int32 InTrackID, int32 FrameNumber, FChaosVDScene& InSceneToControl) const;
+	void PlayFromClosestKeyFrame_AssumesLocked(int32 InTrackID, int32 FrameNumber, FChaosVDScene& InSceneToControl);
 
 	/** Add the provided track info update to the queue. The update will be broadcast in the game thread */
 	void EnqueueTrackInfoUpdate(const TSharedRef<const FChaosVDTrackInfo>& InTrackInfo, FGuid InstigatorID);
