In AChaosVDSolverInfoActor::Destroyed handle the case where SolverParticlesByID is empty ( don't divide by zero )

Avoids the FScopedSlowTask / MakeDialog() when it isn't needed.

diff --git a/Engine/Plugins/ChaosVD/Source/ChaosVD/Private/Actors/ChaosVDSolverInfoActor.cpp b/Engine/Plugins/ChaosVD/Source/ChaosVD/Private/Actors/ChaosVDSolverInfoActor.cpp
--- a/Engine/Plugins/ChaosVD/Source/ChaosVD/Private/Actors/ChaosVDSolverInfoActor.cpp
+++ b/Engine/Plugins/ChaosVD/Source/ChaosVD/Private/Actors/ChaosVDSolverInfoActor.cpp
@@ -243,17 +243,20 @@ void AChaosVDSolverInfoActor::Destroyed()
 		ColorVisualizationSettings->OnSettingsChanged().RemoveAll(this);
 	}
 
-	constexpr float AmountOfWork = 1.0f;
-	const float PercentagePerElement = 1.0f / SolverParticlesByID.Num();
-
-	FScopedSlowTask CleaningParticleDataSlowTask(AmountOfWork, LOCTEXT("CleaningParticleDataMessage", "Cleaning Up Particle Data ..."));
-	CleaningParticleDataSlowTask.MakeDialog();
-	
-	for (const TPair<int32, AChaosVDParticleActor*>& ParticleVDInstanceWithID : SolverParticlesByID)
+	if (SolverParticlesByID.Num() > 0)
 	{
-		World->DestroyActor(ParticleVDInstanceWithID.Value);
+		constexpr float AmountOfWork = 1.0f;
+		const float PercentagePerElement = 1.0f / SolverParticlesByID.Num();
+
+		FScopedSlowTask CleaningParticleDataSlowTask(AmountOfWork, LOCTEXT("CleaningParticleDataMessage", "Cleaning Up Particle Data ..."));
+		CleaningParticleDataSlowTask.MakeDialog();
 
-		CleaningParticleDataSlowTask.EnterProgressFrame(PercentagePerElement);
+		for (const TPair<int32, AChaosVDParticleActor*>& ParticleVDInstanceWithID : SolverParticlesByID)
+		{
+			World->DestroyActor(ParticleVDInstanceWithID.Value);
+
+			CleaningParticleDataSlowTask.EnterProgressFrame(PercentagePerElement);
+		}
 	}
 
 	RemoveSolverFolders(World);
