[ChaosVD] Added Sync by [REDACTED_PERSON] mode

Until now, during playback or scrubbing, solver timelines were synced by their recorded timestamp. Now when the Sync by network tick mode is enabled, CVD will take the network tick offsets into account to choose the correct frame to playback

#reviiew

diff --git a/Engine/Plugins/ChaosVD/Source/ChaosVD/ChaosVD.Build.cs b/Engine/Plugins/ChaosVD/Source/ChaosVD/ChaosVD.Build.cs
--- a/Engine/Plugins/ChaosVD/Source/ChaosVD/ChaosVD.Build.cs
+++ b/Engine/Plugins/ChaosVD/Source/ChaosVD/ChaosVD.Build.cs
@@ -47,7 +47,8 @@ public class ChaosVD : ModuleRules
 				"SubobjectEditor",
 				"SubobjectDataInterface",
 				"PropertyEditor",
-				"CommonMenuExtensions"
+				"CommonMenuExtensions", 
+				"EditorWidgets"
 			}
 			);
 
diff --git a/Engine/Plugins/ChaosVD/Source/ChaosVD/Private/ChaosVDPlaybackController.cpp b/Engine/Plugins/ChaosVD/Source/ChaosVD/Private/ChaosVDPlaybackController.cpp
--- a/Engine/Plugins/ChaosVD/Source/ChaosVD/Private/ChaosVDPlaybackController.cpp
+++ b/Engine/Plugins/ChaosVD/Source/ChaosVD/Private/ChaosVDPlaybackController.cpp
@@ -2,12 +2,11 @@
 
 #include "ChaosVDPlaybackController.h"
 
+#include "Actors/ChaosVDSolverInfoActor.h"
 #include "ChaosVDModule.h"
-#include "ChaosVDPlaybackControllerInstigator.h"
 #include "ChaosVDRecording.h"
 #include "ChaosVDRuntimeModule.h"
 #include "ChaosVDScene.h"
-#include "Actors/ChaosVDSolverInfoActor.h"
 #include "Misc/MessageDialog.h"
 #include "Trace/ChaosVDTraceManager.h"
 #include "Trace/ChaosVDTraceProvider.h"
@@ -254,6 +253,10 @@ void FChaosVDPlaybackController::GoToRecordedSolverStage_AssumesLocked(const int
 				if (CurrentTrackInfo->CurrentFrame != FrameNumber)
 				{
 					CurrentTrackInfo->CurrentFrame = FrameNumber;
+
+					// For server tracks, we only need to have a recorded internal frame number
+					CurrentTrackInfo->bHasNetworkSyncData = SolverFrameData->HasNetworkSyncData(CurrentTrackInfo->bIsServer ? EChaosVDNetworkSyncDataRequirements::InternalFrameNumber : EChaosVDNetworkSyncDataRequirements::All);
+			
 					SceneToControlSharedPtr->HandleEnterNewSolverFrame(FrameNumber, *SolverFrameData);
 				}
 
@@ -339,7 +342,6 @@ void FChaosVDPlaybackController::GoToTrackFrame_AssumesLocked(FGuid InstigatorID
 	}
 }
 
-
 void FChaosVDPlaybackController::GoToTrackFrameAndSync(FGuid InstigatorID, EChaosVDTrackType TrackType, int32 InTrackID, int32 FrameNumber, int32 StageNumber)
 {
 	if (!ensure(LoadedRecording.IsValid()))
@@ -352,18 +354,16 @@ void FChaosVDPlaybackController::GoToTrackFrameAndSync(FGuid InstigatorID, EChao
 	GoToTrackFrame_AssumesLockedAndSync(InstigatorID, TrackType, InTrackID, FrameNumber, StageNumber);
 }
 
-
 void FChaosVDPlaybackController::GoToTrackFrame_AssumesLockedAndSync(FGuid InstigatorID, EChaosVDTrackType TrackType, int32 InTrackID, int32 FrameNumber, int32 StageNumber)
 {
 	GoToTrackFrame_AssumesLocked(InstigatorID, TrackType, InTrackID, FrameNumber, StageNumber);
 	
 	if (TSharedPtr<const FChaosVDTrackInfo> TrackToSyncWith = GetTrackInfo(TrackType, InTrackID))
 	{
-		SyncTracks_AssumesLocked(TrackToSyncWith.ToSharedRef());
+		SyncTracks_AssumesLocked(TrackToSyncWith.ToSharedRef(), CurrentSyncMode);
 	}
 }
 
-
 int32 FChaosVDPlaybackController::GetTrackStepsNumberAtFrame_AssumesLocked(EChaosVDTrackType TrackType, const int32 InTrackID, const int32 FrameNumber) const
 {
 	if (!LoadedRecording.IsValid())
@@ -473,23 +473,60 @@ int32 FChaosVDPlaybackController::ConvertCurrentFrameToOtherTrackFrame_AssumesLo
 
 	switch (InFromTrack->TrackType)
 	{
-	case EChaosVDTrackType::Game:
-		{
-			// Convert from Game Frame to Solver Frame
-			return LoadedRecording->GetLowestSolverFrameNumberGameFrame_AssumesLocked(InFromTrack->TrackID, InFromTrack->CurrentFrame);
-		}	
-	case EChaosVDTrackType::Solver:
-		{
-			if (InToTrack->TrackType == EChaosVDTrackType::Solver)
+		case EChaosVDTrackType::Game:
 			{
-				const FChaosVDSolverFrameData* FromSolverFrameData = LoadedRecording->GetSolverFrameData_AssumesLocked(InFromTrack->TrackID, InFromTrack->CurrentFrame);
-				return ensure(FromSolverFrameData) ? LoadedRecording->GetLowestSolverFrameNumberAtCycle_AssumesLocked(InToTrack->TrackID, FromSolverFrameData->FrameCycle) : INDEX_NONE;
+				ensureMsgf(TrackSyncMode != EChaosVDSyncTimelinesMode::NetworkTick, TEXT("Game tracks cannot be converted to solver tracks using network sync mode. Falling back to timestamp mode."));
+				// Convert from Game Frame to Solver Frame
+				return LoadedRecording->GetLowestSolverFrameNumberGameFrame_AssumesLocked(InFromTrack->TrackID, InFromTrack->CurrentFrame);
+			}	
+		case EChaosVDTrackType::Solver:
+			{
+				if (InToTrack->TrackType == EChaosVDTrackType::Solver)
+				{
+					switch (TrackSyncMode)
+					{
+						case EChaosVDSyncTimelinesMode::RecordedTimestamp:
+							{
+								const FChaosVDSolverFrameData* FromSolverFrameData = LoadedRecording->GetSolverFrameData_AssumesLocked(InFromTrack->TrackID, InFromTrack->CurrentFrame);
+								return ensure(FromSolverFrameData) ? LoadedRecording->GetLowestSolverFrameNumberAtCycle_AssumesLocked(InToTrack->TrackID, FromSolverFrameData->FrameCycle) : INDEX_NONE;
+								break;
+							}
+						case EChaosVDSyncTimelinesMode::NetworkTick:
+							{
+								int32 ToFrame = INDEX_NONE;
+								const FChaosVDSolverFrameData* FromSolverFrameData = LoadedRecording->GetSolverFrameData_AssumesLocked(InFromTrack->TrackID, InFromTrack->CurrentFrame);
+								const FChaosVDSolverFrameData* ToSolverFrameData = LoadedRecording->GetSolverFrameData_AssumesLocked(InToTrack->TrackID, InToTrack->CurrentFrame);
+								if (FromSolverFrameData && ToSolverFrameData)
+								{
+									if (InFromTrack->bIsServer) // Server --> Client track
+									{
+										ToFrame = FromSolverFrameData->InternalFrameNumber - ToSolverFrameData->GetClampedNetworkTickOffset();
+									}
+									else // Client --> Client Track or Client --> Server Track
+									{
+										// This works for Client --> Server conversion because in that case we want to add the frame offset.
+										// As the tick offset in server tracks is 0,  the following calculation will return a negative offset, which it will result in the intended addition in the last caluclation
+										int32 FrameOffset = ToSolverFrameData->GetClampedNetworkTickOffset() - FromSolverFrameData->GetClampedNetworkTickOffset();
+										ToFrame = FromSolverFrameData->InternalFrameNumber - FrameOffset;
+									}
+								}
+
+								return ensure(ToFrame != INDEX_NONE) ? LoadedRecording->GetLowestSolverFrameNumberAtNetworkFrameNumber_AssumesLocked(InToTrack->TrackID, ToFrame) : INDEX_NONE;
+								break;
+							}
+						default:
+							break;
+					}
+				}
+				
+				// From Solver to Game track, we can only convert a frame based on timestapms
+				// TODO: Techincally we are ignoring the requested sync mode, but the current implementation of CVD relies in this fallback as we always want to sync from solver to game tracks using the original timestamo mode
+				// We should investigate if it is worth coming up with a better API so this default fallback does not catch anyone using it off guard.
+				return LoadedRecording->GetLowestGameFrameAtSolverFrameNumber_AssumesLocked(InFromTrack->TrackID, InFromTrack->CurrentFrame);
 			}
-			return LoadedRecording->GetLowestGameFrameAtSolverFrameNumber_AssumesLocked(InFromTrack->TrackID, InFromTrack->CurrentFrame);
-		}
-	default:
-		ensure(false);
-		return INDEX_NONE;
+		default:
+			ensure(false);
+			return INDEX_NONE;
 	}
 }
 
@@ -999,12 +1036,6 @@ void FChaosVDPlaybackController::SyncTracks(const TSharedRef<const FChaosVDTrack
 
 void FChaosVDPlaybackController::SyncTracks_AssumesLocked(const TSharedRef<const FChaosVDTrackInfo>& FromTrack, EChaosVDSyncTimelinesMode TrackSyncMode)
 {
-	if (!ensure(TrackSyncMode == EChaosVDSyncTimelinesMode::RecordedTimestamp))
-	{
-		UE_LOG(LogChaosVDEditor, Error, TEXT("[%s] Attempted to Sync tracks using an unsuported mode. |[%s] is not supported yet."), ANSI_TO_TCHAR(__FUNCTION__), *(UEnum::GetDisplayValueAsText(TrackSyncMode).ToString()));
-		return;
-	}
-
 	if (!FromTrack->bTrackSyncEnabled)
 	{
 		return;
@@ -1029,14 +1060,13 @@ void FChaosVDPlaybackController::SyncTracks_AssumesLocked(const TSharedRef<const
 
 				GoToTrackFrame_AssumesLocked(PlaybackSelfInstigatorID, EChaosVDTrackType::Solver, SolverTrack->TrackID, SolverFrameNumber, StageNumber);
 			}
-
 			break;
 		}	
 	case EChaosVDTrackType::Solver:
 		{
 			if (const TSharedPtr<const FChaosVDTrackInfo>& GameFramesTrackInfo = GetTrackInfo(EChaosVDTrackType::Game, GameTrackID))
 			{
-				const int32 TargetGameFrameNumber = ConvertCurrentFrameToOtherTrackFrame_AssumesLocked(FromTrack, GameFramesTrackInfo.ToSharedRef());
+				const int32 TargetGameFrameNumber = ConvertCurrentFrameToOtherTrackFrame_AssumesLocked(FromTrack, GameFramesTrackInfo.ToSharedRef(), TrackSyncMode);
 				constexpr int32 StageNumber = 0;
 				GoToTrackFrame_AssumesLocked(PlaybackSelfInstigatorID, GameFramesTrackInfo->TrackType, GameFramesTrackInfo->TrackID, TargetGameFrameNumber, StageNumber);
 			}
@@ -1053,7 +1083,7 @@ void FChaosVDPlaybackController::SyncTracks_AssumesLocked(const TSharedRef<const
 					continue;
 				}
 
-				const int32 SolverFrameNumber = ConvertCurrentFrameToOtherTrackFrame_AssumesLocked(FromTrack, SolverTrack.ToSharedRef());
+				const int32 SolverFrameNumber = ConvertCurrentFrameToOtherTrackFrame_AssumesLocked(FromTrack, SolverTrack.ToSharedRef(), TrackSyncMode);
 				const int32 StageNumber = GetTrackLastStageAtFrame_AssumesLocked(EChaosVDTrackType::Solver, SolverTrack->TrackID, SolverFrameNumber);
 
 				GoToTrackFrame_AssumesLocked(PlaybackSelfInstigatorID, EChaosVDTrackType::Solver, SolverTrack->TrackID, SolverFrameNumber, StageNumber);
@@ -1106,6 +1136,11 @@ void FChaosVDPlaybackController::UpdateSolverTracksData()
 		SolverTrackInfo->MaxFrames = GetTrackFramesNumber(EChaosVDTrackType::Solver, SolverIDPair.Key);
 		SolverTrackInfo->TrackName = LoadedRecording->GetSolverFName(SolverIDPair.Key);
 		SolverTrackInfo->TrackType = EChaosVDTrackType::Solver;
+		SolverTrackInfo->bIsServer = LoadedRecording->IsServerSolver_AssumesLocked(SolverIDPair.Key);
+		if (SolverTrackInfo->bIsServer)
+		{
+			CachedServerTrack = SolverTrackInfo;
+		}
 	}
 }
 
diff --git a/Engine/Plugins/ChaosVD/Source/ChaosVD/Private/ChaosVDPlaybackController.h b/Engine/Plugins/ChaosVD/Source/ChaosVD/Private/ChaosVDPlaybackController.h
--- a/Engine/Plugins/ChaosVD/Source/ChaosVD/Private/ChaosVDPlaybackController.h
+++ b/Engine/Plugins/ChaosVD/Source/ChaosVD/Private/ChaosVDPlaybackController.h
@@ -46,6 +46,8 @@ struct FChaosVDTrackInfo
 	bool bIsReSimulated = false;
 	bool bIsPlaying = false;
 	bool bTrackSyncEnabled = true;
+	bool bIsServer = false;
+	bool bHasNetworkSyncData = false;
 
 	bool operator==(const FChaosVDTrackInfo& Other) const;
 
@@ -76,11 +78,10 @@ ENUM_CLASS_FLAGS(EChaosVDUnloadRecordingFlags)
 UENUM()
 enum class EChaosVDSyncTimelinesMode : uint8
 {
-	None = 0,
-	RecordedTimestamp = 1 << 0,
-	NetworkTick = 1 << 1
+	None UMETA(Hidden),
+	RecordedTimestamp,
+	NetworkTick
 };
-ENUM_CLASS_FLAGS(EChaosVDSyncTimelinesMode)
 
 typedef TMap<int32, TSharedPtr<FChaosVDTrackInfo>> TrackInfoByIDMap;
 
@@ -278,6 +279,9 @@ public:
 
 	bool IsPlaying() const;
 
+	EChaosVDSyncTimelinesMode GetTimelineSyncMode() const { return CurrentSyncMode; }
+	void SetTimelineSyncMode(EChaosVDSyncTimelinesMode SyncMode) { CurrentSyncMode = SyncMode ; }
+
 protected:
 
 	/** Updates (or adds) solvers data from the loaded recording to the solver tracks */
@@ -309,6 +313,8 @@ protected:
 	/** Map containing all track info, by track type*/
 	TMap<EChaosVDTrackType, TrackInfoByIDMap> TrackInfoPerType;
 
+	TWeakPtr<FChaosVDTrackInfo> CachedServerTrack;
+
 	/** Ptr to the loaded recording */
 	TSharedPtr<FChaosVDRecording> LoadedRecording;
 
@@ -349,6 +355,8 @@ protected:
 	TSharedPtr<FChaosVDTrackInfo> CurrentPlayingTrack;
 
 	float CurrentPlaybackTime = 0.0f;
+
+	EChaosVDSyncTimelinesMode CurrentSyncMode = EChaosVDSyncTimelinesMode::RecordedTimestamp;
 };
 
 template <typename TVisitorCallback>
diff --git a/Engine/Plugins/ChaosVD/Source/ChaosVD/Private/ChaosVDScene.cpp b/Engine/Plugins/ChaosVD/Source/ChaosVD/Private/ChaosVDScene.cpp
--- a/Engine/Plugins/ChaosVD/Source/ChaosVD/Private/ChaosVDScene.cpp
+++ b/Engine/Plugins/ChaosVD/Source/ChaosVD/Private/ChaosVDScene.cpp
@@ -378,9 +378,13 @@ void FChaosVDScene::HandleEnterNewGameFrame(int32 FrameNumber, const TArray<int3
 
 void FChaosVDScene::HandleEnterNewSolverFrame(int32 FrameNumber, const FChaosVDSolverFrameData& InFrameData)
 {
-	if (AChaosVDSolverInfoActor* SolverDataInfoContainer = SolverDataContainerBySolverID.FindChecked(InFrameData.SolverID))
+	if (AChaosVDSolverInfoActor** SolverDataInfoContainerPtrPtr = SolverDataContainerBySolverID.Find(InFrameData.SolverID))
 	{
-		if (UChaosVDSolverCharacterGroundConstraintDataComponent* DataContainer = SolverDataInfoContainer->GetCharacterGroundConstraintDataComponent())
+		UChaosVDSolverCharacterGroundConstraintDataComponent* DataContainer = SolverDataInfoContainerPtrPtr ? (*SolverDataInfoContainerPtrPtr)->GetCharacterGroundConstraintDataComponent() : nullptr;
+
+		// TODO: Some times when playback is stopped, we might not have all the solver info actors ready when we go to the first frame
+		// This should not happen. For now I am making a change to avoid the crash and ensure. I created UE-217610 to find the issue and prepare a proper fix
+		if (ensure(DataContainer))
 		{
 			DataContainer->UpdateConstraintData(InFrameData.RecordedCharacterGroundConstraints);
 		}
diff --git a/Engine/Plugins/ChaosVD/Source/ChaosVD/Private/Settings/ChaosVDSolverTrackSettings.h b/Engine/Plugins/ChaosVD/Source/ChaosVD/Private/Settings/ChaosVDSolverTrackSettings.h
new file mode 100644
--- /dev/null
+++ b/Engine/Plugins/ChaosVD/Source/ChaosVD/Private/Settings/ChaosVDSolverTrackSettings.h
@@ -0,0 +1,23 @@
+// Copyright Epic Games, Inc. All Rights Reserved.
+
+#pragma once
+
+#include "ChaosVDCoreSettings.h"
+#include "ChaosVDPlaybackController.h"
+#include "ChaosVDSolverTrackSettings.generated.h"
+
+/**
+ * Settings object for Solver Tracks configuration
+ */
+UCLASS(config=ChaosVD, PerObjectConfig)
+class UChaosVDSolverTrackSettings : public UChaosVDSettingsObjectBase
+{
+	GENERATED_BODY()
+
+public:
+	/** Sync mode used to keep all timelines in sync during playback.
+	 * Not: Not all solver tracks support all modes. When an unsupported mode is selected, the default mode, RecordedTimestamp, will be used
+	 */
+	UPROPERTY(config, EditAnywhere, Category="Timeline Sync")
+	EChaosVDSyncTimelinesMode SyncMode = EChaosVDSyncTimelinesMode::RecordedTimestamp;
+};
diff --git a/Engine/Plugins/ChaosVD/Source/ChaosVD/Private/Trace/ChaosVDTraceAnalyzer.cpp b/Engine/Plugins/ChaosVD/Source/ChaosVD/Private/Trace/ChaosVDTraceAnalyzer.cpp
--- a/Engine/Plugins/ChaosVD/Source/ChaosVD/Private/Trace/ChaosVDTraceAnalyzer.cpp
+++ b/Engine/Plugins/ChaosVD/Source/ChaosVD/Private/Trace/ChaosVDTraceAnalyzer.cpp
@@ -25,6 +25,7 @@ void FChaosVDTraceAnalyzer::OnAnalysisBegin(const FOnAnalysisContext& Context)
 	
 	Builder.RouteEvent(RouteId_ChaosVDNonSolverLocation, "ChaosVDLogger", "ChaosVDNonSolverLocation");
 	Builder.RouteEvent(RouteId_ChaosVDNonSolverTransform, "ChaosVDLogger", "ChaosVDNonSolverTransform");
+	Builder.RouteEvent(RouteId_ChaosVDNetworkTickOffset, "ChaosVDLogger", "ChaosVDNetworkTickOffset");
 
 	Builder.RouteEvent(RouteId_BeginFrame, "Misc", "BeginFrame");
 	Builder.RouteEvent(RouteId_EndFrame, "Misc", "EndFrame");
@@ -58,6 +59,7 @@ bool FChaosVDTraceAnalyzer::OnEvent(uint16 RouteId, EStyle Style, const FOnEvent
 				TSharedPtr<FChaosVDGameFrameData> FrameData = MakeShared<FChaosVDGameFrameData>();
 				FrameData->FirstCycle = EventData.GetValue<uint64>("Cycle");
 				FrameData->StartTime = Context.EventTime.AsSeconds(FrameData->FirstCycle);
+
 				ChaosVDTraceProvider->StartGameFrame(FrameData);
 			}
 
@@ -83,10 +85,16 @@ bool FChaosVDTraceAnalyzer::OnEvent(uint16 RouteId, EStyle Style, const FOnEvent
 
 			NewFrameData.SolverID = EventData.GetValue<int32>("SolverID");
 			NewFrameData.FrameCycle = EventData.GetValue<uint64>("Cycle");
+			NewFrameData.InternalFrameNumber = EventData.GetValue<int32>("CurrentFrameNumber", INDEX_NONE);
 			NewFrameData.bIsKeyFrame = EventData.GetValue<bool>("IsKeyFrame");
 			NewFrameData.bIsResimulated = EventData.GetValue<bool>("IsReSimulated");
 			NewFrameData.StartTime = Context.EventTime.AsSeconds(NewFrameData.FrameCycle);
 
+			if (int32* TickOffsetPtr = ChaosVDTraceProvider->GetCurrentTickOffsetsBySolverID().Find(NewFrameData.SolverID))
+			{
+				NewFrameData.NetworkTickOffset = *TickOffsetPtr;
+			}
+
 			FWideStringView DebugNameView;
 			EventData.GetString("DebugName", DebugNameView);
 			NewFrameData.DebugFName = FName(DebugNameView);
@@ -232,6 +240,17 @@ bool FChaosVDTraceAnalyzer::OnEvent(uint16 RouteId, EStyle Style, const FOnEvent
 				CurrentFrameData->RecordedNonSolverTransformsByID.Add(FName(TrackedTransform.DebugName), MoveTemp(TrackedTransform));
 			}
 
+			break;
+		}
+	case RouteId_ChaosVDNetworkTickOffset:
+		{
+			FChaosVDTrackedTransform TrackedTransform;
+
+			const int32 TickOffset = EventData.GetValue<int32>("Offset");
+			const int32 SolverID = EventData.GetValue<int32>("SolverID");
+
+			ChaosVDTraceProvider->GetCurrentTickOffsetsBySolverID().FindOrAdd(SolverID, TickOffset);
+
 			break;
 		}
 
diff --git a/Engine/Plugins/ChaosVD/Source/ChaosVD/Private/Trace/ChaosVDTraceAnalyzer.h b/Engine/Plugins/ChaosVD/Source/ChaosVD/Private/Trace/ChaosVDTraceAnalyzer.h
--- a/Engine/Plugins/ChaosVD/Source/ChaosVD/Private/Trace/ChaosVDTraceAnalyzer.h
+++ b/Engine/Plugins/ChaosVD/Source/ChaosVD/Private/Trace/ChaosVDTraceAnalyzer.h
@@ -50,6 +50,7 @@ private:
 		RouteId_ChaosVDParticleDestroyed,
 		RouteId_ChaosVDNonSolverLocation,
 		RouteId_ChaosVDNonSolverTransform,
+		RouteId_ChaosVDNetworkTickOffset,
 	};
 
 	TraceServices::IAnalysisSession& Session;
diff --git a/Engine/Plugins/ChaosVD/Source/ChaosVD/Private/Trace/ChaosVDTraceProvider.h b/Engine/Plugins/ChaosVD/Source/ChaosVD/Private/Trace/ChaosVDTraceProvider.h
--- a/Engine/Plugins/ChaosVD/Source/ChaosVD/Private/Trace/ChaosVDTraceProvider.h
+++ b/Engine/Plugins/ChaosVD/Source/ChaosVD/Private/Trace/ChaosVDTraceProvider.h
@@ -79,6 +79,8 @@ public:
 
 	void HandleAnalysisComplete();
 
+	TMap<int32,int32>& GetCurrentTickOffsetsBySolverID() {return CurrentNetworkTickOffsets; };
+
 private:
 
 	void RegisterDefaultDataProcessorsIfNeeded();
@@ -107,4 +109,6 @@ private:
 	bool bDefaultDataProcessorsRegistered = false;
 
 	double StartLastCommitedFrameTimeSeconds = 0.0;
+
+	TMap<int32, int32> CurrentNetworkTickOffsets;
 };
diff --git a/Engine/Plugins/ChaosVD/Source/ChaosVD/Private/Utils/ChaosVDUserInterfaceUtils.h b/Engine/Plugins/ChaosVD/Source/ChaosVD/Private/Utils/ChaosVDUserInterfaceUtils.h
--- a/Engine/Plugins/ChaosVD/Source/ChaosVD/Private/Utils/ChaosVDUserInterfaceUtils.h
+++ b/Engine/Plugins/ChaosVD/Source/ChaosVD/Private/Utils/ChaosVDUserInterfaceUtils.h
@@ -3,6 +3,7 @@
 #pragma once
 #include "ChaosVDSettingsManager.h"
 #include "IStructureDetailsView.h"
+#include "SEnumCombo.h"
 
 class UToolMenu;
 
@@ -12,6 +13,29 @@ namespace Chaos::VisualDebugger::Utils
 
 	TSharedRef<IDetailsView> MakeObjectDetailsViewForMenu();
 
+	template <typename EnumType>
+	TSharedRef<SWidget> MakeEnumMenuEntryWidget(const FText& MenuEntryLabel, const SEnumComboBox::FOnEnumSelectionChanged&& EnumValueChanged, const TAttribute<int32>&& CurrentValueAttribute)
+	{
+		return SNew(SHorizontalBox)
+				+SHorizontalBox::Slot()
+				.AutoWidth()
+				.VAlign(VAlign_Center)
+				.Padding(4.0f, 0.f)
+				[
+					SNew(STextBlock)
+					.Text(MenuEntryLabel)
+					.Font(FAppStyle::GetFontStyle(TEXT("MenuItem.Font")))
+				]
+				+SHorizontalBox::Slot()
+				.AutoWidth()
+				.VAlign(VAlign_Center)
+				[
+					SNew(SEnumComboBox, StaticEnum<EnumType>())
+					.CurrentValue(CurrentValueAttribute)
+					.OnEnumSelectionChanged(EnumValueChanged)
+				];
+	}
+
 	UENUM()
 	enum class EChaosVDSaveSettingsOptions
 	{
diff --git a/Engine/Plugins/ChaosVD/Source/ChaosVD/Private/Widgets/SChaosVDGameFramesPlaybackControls.cpp b/Engine/Plugins/ChaosVD/Source/ChaosVD/Private/Widgets/SChaosVDGameFramesPlaybackControls.cpp
--- a/Engine/Plugins/ChaosVD/Source/ChaosVD/Private/Widgets/SChaosVDGameFramesPlaybackControls.cpp
+++ b/Engine/Plugins/ChaosVD/Source/ChaosVD/Private/Widgets/SChaosVDGameFramesPlaybackControls.cpp
@@ -67,12 +67,34 @@ bool SChaosVDGameFramesPlaybackControls::CanPlayback() const
 	bool bCanControlPlayback = false;
 	if (const TSharedPtr<FChaosVDPlaybackController> PlaybackControllerPtr = PlaybackController.Pin())
 	{
-		bCanControlPlayback = PlaybackControllerPtr->IsRecordingLoaded();
-
-		if (bCanControlPlayback)
+		if (PlaybackControllerPtr->IsRecordingLoaded() && GameTrackInfoRef->TrackType != EChaosVDTrackType::Invalid)
 		{
+			// GameFrames playback controls can only be enabled on the Recorded timestamp sync mode
+			bool bIsCompatibleSyncMode = true;
+
+			if (PlaybackControllerPtr->GetTimelineSyncMode() == EChaosVDSyncTimelinesMode::NetworkTick)
+			{
+				static TArray<TSharedPtr<const FChaosVDTrackInfo>> OutTrackInfo;
+				OutTrackInfo.Reset();
+
+				bool bHasTracksWithNetworkSyncData = false;
+				PlaybackControllerPtr->GetAvailableTrackInfosAtTrackFrame(EChaosVDTrackType::Solver, GameTrackInfoRef, OutTrackInfo);
+				for (const TSharedPtr<const FChaosVDTrackInfo>& TrackInfo : OutTrackInfo)
+				{
+					if (TrackInfo && TrackInfo->bHasNetworkSyncData)
+					{
+						bHasTracksWithNetworkSyncData = true;
+						break;
+					}
+				}
+
+				// If we are on Network tick sync mode, but we don't have ant solver track that supports that mode loaded yet (or ever if this was an old file or recorded a standalone game)
+				// allow the game track playback control to work. The solver track widgets themselves will show the corresponding compatibility warning
+				bIsCompatibleSyncMode &= !bHasTracksWithNetworkSyncData;
+			}
+
 			TSharedPtr<FChaosVDTrackInfo> CurrentTrackBeingPlayed = PlaybackControllerPtr->GetCurrentPlayingTrackInfo();
-			bCanControlPlayback = !CurrentTrackBeingPlayed || (CurrentTrackBeingPlayed && CurrentTrackBeingPlayed->TrackType == EChaosVDTrackType::Game);
+			bCanControlPlayback = (!CurrentTrackBeingPlayed || (CurrentTrackBeingPlayed && CurrentTrackBeingPlayed->TrackType == EChaosVDTrackType::Game)) && bIsCompatibleSyncMode;
 		}
 	}
 
diff --git a/Engine/Plugins/ChaosVD/Source/ChaosVD/Private/Widgets/SChaosVDSolverPlaybackControls.cpp b/Engine/Plugins/ChaosVD/Source/ChaosVD/Private/Widgets/SChaosVDSolverPlaybackControls.cpp
--- a/Engine/Plugins/ChaosVD/Source/ChaosVD/Private/Widgets/SChaosVDSolverPlaybackControls.cpp
+++ b/Engine/Plugins/ChaosVD/Source/ChaosVD/Private/Widgets/SChaosVDSolverPlaybackControls.cpp
@@ -27,6 +27,8 @@ void SChaosVDSolverPlaybackControls::Construct(const FArguments& InArgs, const T
 
 	SolverTrackSyncEnabledBrush = FChaosVDStyle::Get().GetBrush(NAME_TrackSyncEnabledBrush);
 	SolverTrackSyncDisabledBrush = FChaosVDStyle::Get().GetBrush(NAME_TrackSyncDisabledBrush);
+	
+	ResimBadgeButtonStyle = FAppStyle::Get().GetWidgetStyle<FButtonStyle>("Menu.Button");
 
 	ChildSlot
 	[
@@ -142,11 +144,19 @@ bool SChaosVDSolverPlaybackControls::CanPlayback() const
 
 	if (const TSharedPtr<FChaosVDPlaybackController> CurrentPlaybackControllerPtr = PlaybackController.Pin())
 	{
+
 		bCanControlPlayback = !CurrentPlaybackControllerPtr->IsPlayingLiveSession();
 
 		// When it is not a live session, the Game Frames timeline follows the same rule as other timelines. The controls are locked unless we are who started a Play action
 		if (bCanControlPlayback)
 		{
+			bool bIsCompatibleSyncMode = CurrentPlaybackControllerPtr->GetTimelineSyncMode() == EChaosVDSyncTimelinesMode::NetworkTick ? SolverTrackInfoRef->bHasNetworkSyncData : true;
+
+			if(!bIsCompatibleSyncMode)
+			{
+				return false;
+			}
+
 			TSharedPtr<FChaosVDTrackInfo> CurrentTrackBeingPlayed = CurrentPlaybackControllerPtr->GetCurrentPlayingTrackInfo();
 			bCanControlPlayback = !CurrentTrackBeingPlayed || (CurrentTrackBeingPlayed && CurrentTrackBeingPlayed->TrackType == EChaosVDTrackType::Solver && CurrentTrackBeingPlayed->TrackID == SolverTrackInfoRef->TrackID);
 		}
@@ -225,8 +235,7 @@ void SChaosVDSolverPlaybackControls::HandleSolverStagePlaybackButtonClicked(ECha
 
 const FSlateBrush* SChaosVDSolverPlaybackControls::GetFrameTypeBadgeBrush() const
 {
-	const FButtonStyle& ButtonStyle = FAppStyle::Get().GetWidgetStyle<FButtonStyle>("Menu.Button");
-	return bIsReSimFrame ? &ButtonStyle.Pressed : FCoreStyle::Get().GetBrush("Border");
+	return SolverTrackInfoRef->bIsReSimulated ? &ResimBadgeButtonStyle.Pressed : FCoreStyle::Get().GetBrush("Border");
 }
 
 TSharedPtr<SWidget> SChaosVDSolverPlaybackControls::CreateVisibilityWidget()
diff --git a/Engine/Plugins/ChaosVD/Source/ChaosVD/Private/Widgets/SChaosVDSolverPlaybackControls.h b/Engine/Plugins/ChaosVD/Source/ChaosVD/Private/Widgets/SChaosVDSolverPlaybackControls.h
--- a/Engine/Plugins/ChaosVD/Source/ChaosVD/Private/Widgets/SChaosVDSolverPlaybackControls.h
+++ b/Engine/Plugins/ChaosVD/Source/ChaosVD/Private/Widgets/SChaosVDSolverPlaybackControls.h
@@ -7,6 +7,7 @@
 #include "ChaosVDPlaybackControllerObserver.h"
 #include "SChaosVDTimelineWidget.h"
 #include "Styling/ISlateStyle.h"
+#include "Styling/SlateTypes.h"
 #include "Templates/SharedPointer.h"
 #include "Widgets/SCompoundWidget.h"
 
@@ -99,4 +100,6 @@ private:
 	
 	const FSlateBrush* SolverTrackSyncEnabledBrush = nullptr;
 	const FSlateBrush* SolverTrackSyncDisabledBrush = nullptr;
+
+	FButtonStyle ResimBadgeButtonStyle;
 };
diff --git a/Engine/Plugins/ChaosVD/Source/ChaosVD/Private/Widgets/SChaosVDSolverTracks.cpp b/Engine/Plugins/ChaosVD/Source/ChaosVD/Private/Widgets/SChaosVDSolverTracks.cpp
--- a/Engine/Plugins/ChaosVD/Source/ChaosVD/Private/Widgets/SChaosVDSolverTracks.cpp
+++ b/Engine/Plugins/ChaosVD/Source/ChaosVD/Private/Widgets/SChaosVDSolverTracks.cpp
@@ -4,19 +4,48 @@
 
 #include "ChaosVDModule.h"
 #include "ChaosVDPlaybackController.h"
+#include "SChaosVDWarningMessageBox.h"
+#include "Settings/ChaosVDSolverTrackSettings.h"
+#include "ToolMenu.h"
+#include "ToolMenus.h"
+#include "ToolMenuEntry.h"
+#include "ToolMenuSection.h"
+#include "Utils/ChaosVDUserInterfaceUtils.h"
 #include "Widgets/SChaosVDPlaybackViewport.h"
 #include "Widgets/SChaosVDSolverPlaybackControls.h"
 #include "Widgets/Layout/SExpandableArea.h"
+#include "Widgets/Layout/SSeparator.h"
+
+#define LOCTEXT_NAMESPACE "ChaosVisualDebugger"
 
 void SChaosVDSolverTracks::Construct(const FArguments& InArgs, TWeakPtr<FChaosVDPlaybackController> InPlaybackController)
 {
+	MenuName = FName(TEXT("ChaosVD.SolverTracks.MenuToolbar"));
+
 	ChildSlot
 	[
-		SAssignNew(SolverTracksListWidget, SListView<TSharedPtr<const FChaosVDTrackInfo>>)
-		.ListItemsSource(&CachedTrackInfoArray)
-		.SelectionMode( ESelectionMode::None )
-		.ListViewStyle(&FAppStyle::Get().GetWidgetStyle<FTableViewStyle>("SimpleListView"))
-		.OnGenerateRow(this, &SChaosVDSolverTracks::MakeSolverTrackControlsFromTrackInfo)
+		SNew(SVerticalBox)
+		+SVerticalBox::Slot()
+		.AutoHeight()
+		.Padding(0, 0.0f, 0, 0.0f)
+		[
+			SNew(SHorizontalBox)
+			+SHorizontalBox::Slot()
+			.HAlign(HAlign_Left)
+			[
+				GenerateToolbarWidget()
+			]
+		]
+		+SVerticalBox::Slot()
+		.Padding(0, 0.0f, 0, 2.0f)
+		[
+			SAssignNew(SolverTracksListWidget, SListView<TSharedPtr<const FChaosVDTrackInfo>>)
+			.ListItemsSource(&CachedTrackInfoArray)
+			.SelectionMode( ESelectionMode::None )
+			.ListViewStyle(&FAppStyle::Get().GetWidgetStyle<FTableViewStyle>("SimpleListView"))
+			.OnGenerateRow(this, &SChaosVDSolverTracks::MakeSolverTrackControlsFromTrackInfo)
+		]
+	
 	];
 
 	ensure(InPlaybackController.IsValid());
@@ -34,6 +63,21 @@ void SChaosVDSolverTracks::Construct(const FArguments& InArgs, TWeakPtr<FChaosVD
 	{
 		UE_LOG(LogChaosVDEditor, Error, TEXT("[%s] Solver tracks contructed with an invalid player controler. The solver tracks widget will not be functional"), ANSI_TO_TCHAR(__FUNCTION__))	
 	}
+	
+	if (UChaosVDSolverTrackSettings* Settings = FChaosVDSettingsManager::Get().GetSettingsObject<UChaosVDSolverTrackSettings>())
+	{
+		Settings->OnSettingsChanged().AddSP(this, &SChaosVDSolverTracks::HandleSettingsChanged);
+		
+		HandleSettingsChanged(Settings);
+	}
+}
+
+SChaosVDSolverTracks::~SChaosVDSolverTracks()
+{
+	if (UChaosVDSolverTrackSettings* Settings = FChaosVDSettingsManager::Get().GetSettingsObject<UChaosVDSolverTrackSettings>())
+	{
+		Settings->OnSettingsChanged().RemoveAll(this);
+	}
 }
 
 void SChaosVDSolverTracks::HandlePlaybackControllerDataUpdated(TWeakPtr<FChaosVDPlaybackController> InPlaybackController)
@@ -106,7 +150,7 @@ TSharedRef<ITableRow> SChaosVDSolverTracks::MakeSolverTrackControlsFromTrackInfo
 		RowWidget = SNew(SVerticalBox)
 				+SVerticalBox::Slot()
 				.AutoHeight()
-				.Padding(10.0f, 2.0f, 10.0f, 0.0f)
+				.Padding(10.0f, 4.0f, 10.0f, 0.0f)
 				[
 					SNew(SExpandableArea)
 						.InitiallyCollapsed(false)
@@ -125,6 +169,34 @@ TSharedRef<ITableRow> SChaosVDSolverTracks::MakeSolverTrackControlsFromTrackInfo
 								.Text(FText::FromName(TrackInfo->TrackName))
 								.Font(FCoreStyle::Get().GetFontStyle("ExpandableArea.TitleFont"))
 							]
+							+SHorizontalBox::Slot()
+							.AutoWidth()
+							.Padding(8.0f, 4.0f, 8.0f, 4.0f)
+							[
+								SNew(SSeparator)
+								.Orientation(Orient_Vertical)
+								.Thickness(1.f)
+							]
+							+SHorizontalBox::Slot()
+							.HAlign(HAlign_Left)
+                            .VAlign(VAlign_Center)
+                            .AutoWidth()
+                            .Padding(0.f, 0.f, 0.f, 0.f)
+                            [
+								SNew(SChaosVDWarningMessageBox)
+								.Visibility_Lambda([TrackInfoAsWeak = TrackInfo.ToWeakPtr(), WeakPlaybackController = PlaybackController]()
+								{
+									TSharedPtr<const FChaosVDTrackInfo> TrackInfoPtr = TrackInfoAsWeak.Pin();
+									TSharedPtr<FChaosVDPlaybackController> PlaybackControllerPtr = WeakPlaybackController.Pin();
+									if (TrackInfoPtr && PlaybackControllerPtr)
+									{
+										bool bIsCompatibleMode = PlaybackControllerPtr->GetTimelineSyncMode() == EChaosVDSyncTimelinesMode::NetworkTick ? TrackInfoPtr->bHasNetworkSyncData : true;
+										return bIsCompatibleMode ? EVisibility::Collapsed : EVisibility::Visible;
+									}
+									return EVisibility::Collapsed;
+								})
+								.WarningText(LOCTEXT("IncomatibleSyncModeWarning", " Incompatible sync mode selected | Attempting to fallback to TimeStamp sync mode for this track | Controls disabled"))
+                            ]
 						]
 						.BodyContent()
 						[
@@ -153,3 +225,84 @@ TSharedRef<ITableRow> SChaosVDSolverTracks::MakeSolverTrackControlsFromTrackInfo
 			RowWidget.ToSharedRef()
 		];
 }
+
+TSharedRef<SWidget> SChaosVDSolverTracks::GenerateToolbarWidget()
+{
+	RegisterMenus();
+
+	FToolMenuContext MenuContext;
+
+	UChaosVDSolverTracksToolbarMenuContext* CommonContextObject = NewObject<UChaosVDSolverTracksToolbarMenuContext>();
+	CommonContextObject->SolverTracksWidget = SharedThis(this);
+
+	MenuContext.AddObject(CommonContextObject);
+
+	return UToolMenus::Get()->GenerateWidget(MenuName, MenuContext);
+}
+
+TSharedRef<SWidget> SChaosVDSolverTracks::GenerateSyncModeMenuWidget()
+{
+	TAttribute<int32> GetCurrentMode;
+	GetCurrentMode.BindLambda([]()
+	{
+		if (UChaosVDSolverTrackSettings* Settings = FChaosVDSettingsManager::Get().GetSettingsObject<UChaosVDSolverTrackSettings>())
+		{
+			return static_cast<int32>(Settings->SyncMode);
+		}
+		return 0;
+	});
+
+	SEnumComboBox::FOnEnumSelectionChanged ValueChangedDelegate;
+	ValueChangedDelegate.BindLambda([](int32 NewValue, ESelectInfo::Type)
+	{
+		if (UChaosVDSolverTrackSettings* Settings = FChaosVDSettingsManager::Get().GetSettingsObject<UChaosVDSolverTrackSettings>())
+		{
+			Settings->SyncMode = static_cast<EChaosVDSyncTimelinesMode>(NewValue);
+			Settings->OnSettingsChanged().Broadcast(Settings);
+			Settings->SaveConfig();
+		}
+	});
+
+	return Chaos::VisualDebugger::Utils::MakeEnumMenuEntryWidget<EChaosVDSyncTimelinesMode>(LOCTEXT("SyncTimelineModeMenuLabel", "Timeline Sync Mode"), MoveTemp(ValueChangedDelegate), MoveTemp(GetCurrentMode));
+}
+
+void SChaosVDSolverTracks::HandleSettingsChanged(UObject* SettingsObject)
+{
+	if (UChaosVDSolverTrackSettings* Settings = Cast<UChaosVDSolverTrackSettings>(SettingsObject))
+	{
+		if (TSharedPtr<FChaosVDPlaybackController> PlaybackControllerPtr = PlaybackController.Pin())
+		{
+			PlaybackControllerPtr->SetTimelineSyncMode(Settings->SyncMode);
+		}
+	}
+}
+
+void SChaosVDSolverTracks::RegisterMenus()
+{
+	const UToolMenus* ToolMenus = UToolMenus::Get();
+	if (ToolMenus->IsMenuRegistered(MenuName))
+	{
+		return;
+	}
+
+	UToolMenu* ToolBar = UToolMenus::Get()->RegisterMenu(MenuName, NAME_None, EMultiBoxType::SlimHorizontalToolBar);
+
+	FToolMenuSection& Section = ToolBar->AddSection("MainToolbar");
+	Section.AddDynamicEntry("MainToolbarEntry", FNewToolMenuSectionDelegate::CreateLambda([](FToolMenuSection& InSection)
+	{
+		const UChaosVDSolverTracksToolbarMenuContext* Context = InSection.FindContext<UChaosVDSolverTracksToolbarMenuContext>();
+		TSharedRef<SChaosVDSolverTracks> SolverTracksWidget = Context->SolverTracksWidget.Pin().ToSharedRef();
+		
+		InSection.AddEntry(
+		FToolMenuEntry::InitWidget(
+			"SyncModeButton",
+			SolverTracksWidget->GenerateSyncModeMenuWidget(),
+			FText::GetEmpty(),
+			false,
+			false
+		));
+		
+	}));
+}
+
+#undef LOCTEXT_NAMESPACE
diff --git a/Engine/Plugins/ChaosVD/Source/ChaosVD/Private/Widgets/SChaosVDSolverTracks.h b/Engine/Plugins/ChaosVD/Source/ChaosVD/Private/Widgets/SChaosVDSolverTracks.h
--- a/Engine/Plugins/ChaosVD/Source/ChaosVD/Private/Widgets/SChaosVDSolverTracks.h
+++ b/Engine/Plugins/ChaosVD/Source/ChaosVD/Private/Widgets/SChaosVDSolverTracks.h
@@ -8,9 +8,21 @@
 #include "Widgets/SCompoundWidget.h"
 #include "Widgets/Views/SListView.h"
 
+#include "SChaosVDSolverTracks.generated.h"
+
 struct FChaosVDTrackInfo;
 class FChaosVDPlaybackController;
 
+UCLASS()
+class CHAOSVD_API UChaosVDSolverTracksToolbarMenuContext : public UObject
+{
+	GENERATED_BODY()
+
+public:
+
+	TWeakPtr<class SChaosVDSolverTracks> SolverTracksWidget;
+};
+
 /** Widget that Generates a expandable list of solver controls, based on the existing solver data
  * on the ChaosVDPlaybackController
  */
@@ -22,6 +34,8 @@ public:
 	SLATE_END_ARGS()
 
 	void Construct(const FArguments& InArgs, TWeakPtr<FChaosVDPlaybackController> InPlaybackController);
+	
+	~SChaosVDSolverTracks();
 
 private:
 
@@ -34,4 +48,13 @@ private:
 	TSharedPtr<SListView<TSharedPtr<const FChaosVDTrackInfo>>> SolverTracksListWidget;
 
 	TArray<TSharedPtr<const FChaosVDTrackInfo>> CachedTrackInfoArray;
+
+	TSharedRef<SWidget> GenerateToolbarWidget();
+	TSharedRef<SWidget> GenerateSyncModeMenuWidget();
+
+	void HandleSettingsChanged(UObject* SettingsObject);
+
+	void RegisterMenus();
+
+	FName MenuName;
 };
diff --git a/Engine/Source/Runtime/Engine/Private/PlayerController.cpp b/Engine/Source/Runtime/Engine/Private/PlayerController.cpp
--- a/Engine/Source/Runtime/Engine/Private/PlayerController.cpp
+++ b/Engine/Source/Runtime/Engine/Private/PlayerController.cpp
@@ -2325,6 +2325,15 @@ void APlayerController::SetMouseLocation(const int X, const int Y)
 /* PlayerTick is only called if the PlayerController has a PlayerInput object.  Therefore, it will not be called on servers for non-locally controlled playercontrollers. */
 void APlayerController::PlayerTick( float DeltaTime )
 {
+#if WITH_CHAOS_VISUAL_DEBUGGER
+
+	if (bNetworkPhysicsTickOffsetAssigned)
+	{
+		CVD_TRACE_NETWORK_TICK_OFFSET(NetworkPhysicsTickOffset, CVD_TRACE_GET_SOLVER_ID_FROM_WORLD(GetWorld()));
+	}
+
+#endif
+
 	if (!bShortConnectTimeOut)
 	{
 		bShortConnectTimeOut = true;
diff --git a/Engine/Source/Runtime/Engine/Public/Physics/ImmediatePhysics/ImmediatePhysicsChaos/ImmediatePhysicsSimulation_Chaos.h b/Engine/Source/Runtime/Engine/Public/Physics/ImmediatePhysics/ImmediatePhysicsChaos/ImmediatePhysicsSimulation_Chaos.h
--- a/Engine/Source/Runtime/Engine/Public/Physics/ImmediatePhysics/ImmediatePhysicsChaos/ImmediatePhysicsSimulation_Chaos.h
+++ b/Engine/Source/Runtime/Engine/Public/Physics/ImmediatePhysics/ImmediatePhysicsChaos/ImmediatePhysicsSimulation_Chaos.h
@@ -147,11 +147,15 @@ namespace ImmediatePhysics_Chaos
 	FName DebugName;
 #endif
 
+	private:
+
 #if WITH_CHAOS_VISUAL_DEBUGGER
 	private:
 		FChaosVDContext CVDContextData;
 
 	public:
+		int32 GetCVDFrameNumber() const { return INDEX_NONE; }
+
 		FChaosVDContext& GetChaosVDContextData()
 		{
 			return CVDContextData;
diff --git a/Engine/Source/Runtime/Experimental/Chaos/Private/ChaosVisualDebugger/ChaosVisualDebuggerTrace.cpp b/Engine/Source/Runtime/Experimental/Chaos/Private/ChaosVisualDebugger/ChaosVisualDebuggerTrace.cpp
--- a/Engine/Source/Runtime/Experimental/Chaos/Private/ChaosVisualDebugger/ChaosVisualDebuggerTrace.cpp
+++ b/Engine/Source/Runtime/Experimental/Chaos/Private/ChaosVisualDebugger/ChaosVisualDebuggerTrace.cpp
@@ -42,6 +42,7 @@ UE_TRACE_EVENT_DEFINE(ChaosVDLogger, ChaosVDSolverSimulationSpace)
 UE_TRACE_EVENT_DEFINE(ChaosVDLogger, ChaosVDDummyEvent)
 UE_TRACE_EVENT_DEFINE(ChaosVDLogger, ChaosVDNonSolverLocation)
 UE_TRACE_EVENT_DEFINE(ChaosVDLogger, ChaosVDNonSolverTransform)
+UE_TRACE_EVENT_DEFINE(ChaosVDLogger, ChaosVDNetworkTickOffset)
 
 namespace Chaos::VisualDebugger::Cvars
 {
@@ -262,6 +263,11 @@ void FChaosVisualDebuggerTrace::SetupForFullCaptureIfNeeded(int32 SolverID, bool
 	}
 }
 
+int32 FChaosVisualDebuggerTrace::GetSolverID(Chaos::FPhysicsSolverBase& Solver)
+{
+	return Solver.GetChaosVDContextData().Id;
+}
+
 bool FChaosVisualDebuggerTrace::ShouldPerformFullCapture(int32 SolverID)
 {
 	FReadScopeLock ReadLock(DeltaRecordingStatesLock);
@@ -463,7 +469,7 @@ void FChaosVisualDebuggerTrace::TraceConstraintsContainer(TConstArrayView<Chaos:
 	}
 }
 
-void FChaosVisualDebuggerTrace::TraceSolverFrameStart(const FChaosVDContext& ContextData, const FString& InDebugName)
+void FChaosVisualDebuggerTrace::TraceSolverFrameStart(const FChaosVDContext& ContextData, const FString& InDebugName, int32 FrameNumber)
 {
 	if (!IsTracing())
 	{
@@ -493,7 +499,8 @@ void FChaosVisualDebuggerTrace::TraceSolverFrameStart(const FChaosVDContext& Con
 		<< ChaosVDSolverFrameStart.Cycle(FPlatformTime::Cycles64())
 		<< ChaosVDSolverFrameStart.DebugName(*InDebugName, InDebugName.Len())
 		<< ChaosVDSolverFrameStart.IsKeyFrame(bOutIsFullCaptureRequested)
-		<< ChaosVDSolverFrameStart.IsReSimulated(bIsReSimulatedFrame);
+		<< ChaosVDSolverFrameStart.IsReSimulated(bIsReSimulatedFrame)
+		<< ChaosVDSolverFrameStart.CurrentFrameNumber(FrameNumber);
 }
 
 void FChaosVisualDebuggerTrace::TraceSolverFrameEnd(const FChaosVDContext& ContextData)
@@ -829,6 +836,18 @@ void FChaosVisualDebuggerTrace::TraceSceneAccelerationStructures(const Chaos::IS
 	}
 }
 
+void FChaosVisualDebuggerTrace::TraceNetworkTickOffset(int32 TickOffset, int32 SolverID)
+{
+	if (!IsTracing())
+	{
+		return;
+	}
+	
+	UE_TRACE_LOG(ChaosVDLogger, ChaosVDNetworkTickOffset, ChaosVDChannel)
+		<< ChaosVDNetworkTickOffset.Offset(TickOffset)
+		<< ChaosVDNetworkTickOffset.SolverID(SolverID);
+}
+
 bool FChaosVisualDebuggerTrace::IsTracing()
 {
 	return bIsTracing;
diff --git a/Engine/Source/Runtime/Experimental/Chaos/Public/Chaos/Framework/PhysicsSolverBase.h b/Engine/Source/Runtime/Experimental/Chaos/Public/Chaos/Framework/PhysicsSolverBase.h
--- a/Engine/Source/Runtime/Experimental/Chaos/Public/Chaos/Framework/PhysicsSolverBase.h
+++ b/Engine/Source/Runtime/Experimental/Chaos/Public/Chaos/Framework/PhysicsSolverBase.h
@@ -535,6 +535,10 @@ namespace Chaos
 		virtual void ConditionalApplyRewind_Internal(){}
 		virtual bool IsResimming() const {return false;}
 
+#if WITH_CHAOS_VISUAL_DEBUGGER
+		virtual int32 GetCVDFrameNumber() const { return INDEX_NONE; }
+#endif
+
 		FChaosMarshallingManager& GetMarshallingManager() { return MarshallingManager; }
 		FChaosResultsManager& GetResultsManager() { return *PullResultsManager; }
 
diff --git a/Engine/Source/Runtime/Experimental/Chaos/Public/ChaosVisualDebugger/ChaosVDTraceMacros.h b/Engine/Source/Runtime/Experimental/Chaos/Public/ChaosVisualDebugger/ChaosVDTraceMacros.h
--- a/Engine/Source/Runtime/Experimental/Chaos/Public/ChaosVisualDebugger/ChaosVDTraceMacros.h
+++ b/Engine/Source/Runtime/Experimental/Chaos/Public/ChaosVisualDebugger/ChaosVDTraceMacros.h
@@ -58,7 +58,7 @@
 		#define CVD_TRACE_SOLVER_START_FRAME(SolverType, SolverRef) \
 			FChaosVDContext StartContextData; \
 			FChaosVisualDebuggerTrace::GetCVDContext<SolverType>(SolverRef, StartContextData); \
-			FChaosVisualDebuggerTrace::TraceSolverFrameStart(StartContextData, FChaosVisualDebuggerTrace::GetDebugName<SolverType>(SolverRef));
+			FChaosVisualDebuggerTrace::TraceSolverFrameStart(StartContextData, FChaosVisualDebuggerTrace::GetDebugName<SolverType>(SolverRef), SolverRef.GetCVDFrameNumber());
 	#endif
 
 	#ifndef CVD_TRACE_SOLVER_END_FRAME
@@ -177,6 +177,18 @@
 			}
 	#endif
 
+#ifndef CVD_TRACE_NETWORK_TICK_OFFSET
+	#define CVD_TRACE_NETWORK_TICK_OFFSET(TickOffset, SolverID) \
+	{ \
+		FChaosVisualDebuggerTrace::TraceNetworkTickOffset(TickOffset, SolverID); \
+	}
+#endif
+
+#ifndef CVD_TRACE_GET_SOLVER_ID_FROM_WORLD
+	#define CVD_TRACE_GET_SOLVER_ID_FROM_WORLD(World) \
+			FChaosVisualDebuggerTrace::GetSolverIDFromWorld(World)
+#endif
+
 #else // WITH_CHAOS_VISUAL_DEBUGGER
 
 	#ifndef CVD_TRACE_PARTICLE
@@ -277,4 +289,12 @@
 			#define CVD_TRACE_ACCELERATION_STRUCTURES(AccelerationStructuresCollections, SolverType, SolverRef, DataChannel)
 	#endif
 
+	#ifndef CVD_TRACE_NETWORK_TICK_OFFSET
+		#define CVD_TRACE_NETWORK_TICK_OFFSET(TickOffset, SolverID)
+	#endif
+
+#ifndef CVD_TRACE_GET_SOLVER_ID_FROM_WORLD
+	#define CVD_TRACE_GET_SOLVER_ID_FROM_WORLD(World)
+#endif
+
 #endif // WITH_CHAOS_VISUAL_DEBUGGER
\ No newline at end of file
diff --git a/Engine/Source/Runtime/Experimental/Chaos/Public/ChaosVisualDebugger/ChaosVisualDebuggerTrace.h b/Engine/Source/Runtime/Experimental/Chaos/Public/ChaosVisualDebugger/ChaosVisualDebuggerTrace.h
--- a/Engine/Source/Runtime/Experimental/Chaos/Public/ChaosVisualDebugger/ChaosVisualDebuggerTrace.h
+++ b/Engine/Source/Runtime/Experimental/Chaos/Public/ChaosVisualDebugger/ChaosVisualDebuggerTrace.h
@@ -16,6 +16,7 @@
 	#define CHAOS_VISUAL_DEBUGGER_ENABLED WITH_CHAOS_VISUAL_DEBUGGER
 #endif
 
+#include "Chaos/Framework/PhysicsSolverBase.h"
 #include "ChaosVisualDebugger/ChaosVDTraceMacros.h"
 
 #if WITH_CHAOS_VISUAL_DEBUGGER
@@ -42,6 +43,7 @@ UE_TRACE_EVENT_BEGIN_EXTERN(ChaosVDLogger, ChaosVDSolverFrameStart)
 	UE_TRACE_EVENT_FIELD(UE::Trace::WideString, DebugName)
 	UE_TRACE_EVENT_FIELD(bool, IsKeyFrame)
 	UE_TRACE_EVENT_FIELD(bool, IsReSimulated)
+	UE_TRACE_EVENT_FIELD(int32, CurrentFrameNumber)
 UE_TRACE_EVENT_END()
 
 UE_TRACE_EVENT_BEGIN_EXTERN(ChaosVDLogger, ChaosVDSolverFrameEnd)
@@ -114,6 +116,11 @@ UE_TRACE_EVENT_BEGIN_EXTERN(ChaosVDLogger, ChaosVDNonSolverTransform)
 	UE_TRACE_EVENT_FIELD(UE::Trace::WideString, DebugName)
 UE_TRACE_EVENT_END()
 
+UE_TRACE_EVENT_BEGIN_EXTERN(ChaosVDLogger, ChaosVDNetworkTickOffset)
+	UE_TRACE_EVENT_FIELD(int32, Offset)
+	UE_TRACE_EVENT_FIELD(int32, SolverID)
+UE_TRACE_EVENT_END()
+
 UE_TRACE_EVENT_BEGIN_EXTERN(ChaosVDLogger, ChaosVDDummyEvent)
 	UE_TRACE_EVENT_FIELD(int32, SolverID)
 UE_TRACE_EVENT_END()
@@ -216,7 +223,7 @@ public:
 	static CHAOS_API void TraceConstraintsContainer(TConstArrayView<Chaos::FPBDConstraintContainer*> ConstraintContainersView);
 
 	/** Traces the start of a solver frame and it pushes its context data to the CVD TLS context stack */
-	static CHAOS_API void TraceSolverFrameStart(const FChaosVDContext& ContextData, const FString& InDebugName);
+	static CHAOS_API void TraceSolverFrameStart(const FChaosVDContext& ContextData, const FString& InDebugName, int32 FrameNumber = INDEX_NONE);
 	
 	/** Traces the end of a solver frame and removes its context data to the CVD TLS context stack */
 	static CHAOS_API void TraceSolverFrameEnd(const FChaosVDContext& ContextData);
@@ -258,8 +265,15 @@ public:
 
 	static CHAOS_API void TraceSceneAccelerationStructures(const Chaos::ISpatialAccelerationCollection<Chaos::FAccelerationStructureHandle, Chaos::FReal, 3>* InAccelerationCollection);
 
+	static CHAOS_API void TraceNetworkTickOffset(int32 TickOffset, int32 SolverID);
+
+	template<typename WorldType>
+	static int32 GetSolverIDFromWorld(WorldType* World);
+
+	static int32 CHAOS_API GetSolverID(Chaos::FPhysicsSolverBase& Solver);
+
 	/** Returns true if the provided solver ID needs a Full Capture */
-	static bool ShouldPerformFullCapture(int32 SolverID);
+	static CHAOS_API bool ShouldPerformFullCapture(int32 SolverID);
 
 	/**
 	 * Gets the CVD Context data form an object that has such data. Usually Solvers
@@ -363,6 +377,21 @@ void FChaosVisualDebuggerTrace::TraceParticlesView(const Chaos::TParticleView<Pa
 	});
 }
 
+template <typename WorldType>
+int32 FChaosVisualDebuggerTrace::GetSolverIDFromWorld(WorldType* World)
+{
+	int32 SolverID = INDEX_NONE;
+	if (World)
+	{
+		if (Chaos::FPhysicsSolverBase* Solver = World->GetPhysicsScene() ? World->GetPhysicsScene()->GetSolver() : nullptr)
+		{
+			SolverID = GetSolverID(*Solver);
+		}
+	}
+
+	return SolverID;
+}
+
 template <typename T>
 void FChaosVisualDebuggerTrace::GetCVDContext(T& ObjectWithContext, FChaosVDContext& OutCVDContext)
 {
diff --git a/Engine/Source/Runtime/Experimental/Chaos/Public/PBDRigidsSolver.h b/Engine/Source/Runtime/Experimental/Chaos/Public/PBDRigidsSolver.h
--- a/Engine/Source/Runtime/Experimental/Chaos/Public/PBDRigidsSolver.h
+++ b/Engine/Source/Runtime/Experimental/Chaos/Public/PBDRigidsSolver.h
@@ -174,6 +174,10 @@ namespace Chaos
 		void SetCurrentFrame(const int32 CurrentFrameIn) { CurrentFrame = CurrentFrameIn; }
 		int32& GetCurrentFrame() { return CurrentFrame; }
 
+#if WITH_CHAOS_VISUAL_DEBUGGER
+		virtual int32 GetCVDFrameNumber() const override { return CurrentFrame; }
+#endif
+
 		/**/
 		void SetPositionIterations(const int32 InNumIterations) { GetEvolution()->SetNumPositionIterations(InNumIterations); }
 		void SetVelocityIterations(const int32 InNumIterations) { GetEvolution()->SetNumVelocityIterations(InNumIterations); }
diff --git a/Engine/Source/Runtime/Experimental/ChaosVDData/Private/ChaosVDRecording.cpp b/Engine/Source/Runtime/Experimental/ChaosVDData/Private/ChaosVDRecording.cpp
--- a/Engine/Source/Runtime/Experimental/ChaosVDData/Private/ChaosVDRecording.cpp
+++ b/Engine/Source/Runtime/Experimental/ChaosVDData/Private/ChaosVDRecording.cpp
@@ -65,6 +65,17 @@ FName FChaosVDRecording::GetSolverFName_AssumedLocked(int32 SolverID)
 	return DefaultName;
 }
 
+bool FChaosVDRecording::IsServerSolver_AssumesLocked(int32 SolverID)
+{
+	return GetSolverFName_AssumedLocked(SolverID).ToString().Contains(TEXT("Server"));
+}
+
+bool FChaosVDRecording::IsServerSolver(int32 SolverID)
+{
+	FReadScopeLock ReadLock(RecordingDataLock);
+	return IsServerSolver_AssumesLocked(SolverID);
+}
+
 FChaosVDSolverFrameData* FChaosVDRecording::GetSolverFrameData_AssumesLocked(const int32 SolverID, const int32 FrameNumber, bool bKeyFrameOnly)
 {
 	if (TArray<FChaosVDSolverFrameData>* SolverFrames = RecordedFramesDataPerSolver.Find(SolverID))
@@ -123,6 +134,19 @@ int32 FChaosVDRecording::GetLowestSolverFrameNumberAtCycle_AssumesLocked(int32 S
 	return INDEX_NONE;
 }
 
+
+int32 FChaosVDRecording::GetLowestSolverFrameNumberAtNetworkFrameNumber_AssumesLocked(int32 SolverID, int32 NetworkFrameNumber)
+{
+	if (TArray<FChaosVDSolverFrameData>* SolverFramesPtr = RecordedFramesDataPerSolver.Find(SolverID))
+	{
+		TArray<FChaosVDSolverFrameData>& SolverFrames = *SolverFramesPtr;
+		
+		return Algo::LowerBoundBy(SolverFrames, NetworkFrameNumber, &FChaosVDSolverFrameData::InternalFrameNumber);
+	}
+
+	return INDEX_NONE;
+}
+
 int32 FChaosVDRecording::FindFirstSolverKeyFrameNumberFromFrame_AssumesLocked(int32 SolverID, int32 StartFrameNumber)
 {
 	if (TArray<int32>* KeyFrameNumbersPtr = RecordedKeyFramesNumberPerSolver.Find(SolverID))
diff --git a/Engine/Source/Runtime/Experimental/ChaosVDData/Public/ChaosVDRecording.h b/Engine/Source/Runtime/Experimental/ChaosVDData/Public/ChaosVDRecording.h
--- a/Engine/Source/Runtime/Experimental/ChaosVDData/Public/ChaosVDRecording.h
+++ b/Engine/Source/Runtime/Experimental/ChaosVDData/Public/ChaosVDRecording.h
@@ -49,6 +49,17 @@ struct FChaosVDTrackedTransform
 	FTransform Transform;
 };
 
+enum class EChaosVDNetworkSyncDataRequirements
+{
+	None = 0,
+	InternalFrameNumber  = 1 << 0,
+	NetworkTickOffset = 1 << 1,
+
+	All = InternalFrameNumber | NetworkTickOffset
+};
+
+ENUM_CLASS_FLAGS(EChaosVDNetworkSyncDataRequirements)
+
 typedef TArray<FChaosVDStepData, TInlineAllocator<16>> FChaosVDStepsContainer;
 
 struct CHAOSVDDATA_API FChaosVDSolverFrameData
@@ -66,6 +77,8 @@ PRAGMA_ENABLE_DEPRECATION_WARNINGS
 	UE_DEPRECATED(5.5, "Please use the DebugFName instead")
 	FString DebugName;
 	int32 SolverID = INDEX_NONE;
+	int32 InternalFrameNumber = INDEX_NONE;
+	int32 NetworkTickOffset = INDEX_NONE;
 	uint64 FrameCycle = 0;
 	Chaos::FRigidTransform3 SimulationTransform;
 	bool bIsKeyFrame = false;
@@ -88,6 +101,30 @@ PRAGMA_ENABLE_DEPRECATION_WARNINGS
 
 		return EndTime - StartTime;
 	}
+
+	/** Returns true if we have the necessary data to sync this frame with other frame based on network ticks offsets */
+	bool HasNetworkSyncData(EChaosVDNetworkSyncDataRequirements Requirements = EChaosVDNetworkSyncDataRequirements::All) const
+	{
+		bool bHasRequiredSyncData = true;
+		if (EnumHasAnyFlags(Requirements, EChaosVDNetworkSyncDataRequirements::InternalFrameNumber))
+		{
+			bHasRequiredSyncData &= InternalFrameNumber != INDEX_NONE;
+		}
+		
+		if (EnumHasAnyFlags(Requirements, EChaosVDNetworkSyncDataRequirements::NetworkTickOffset))
+		{
+			bHasRequiredSyncData &= NetworkTickOffset != INDEX_NONE;
+		}
+
+		return bHasRequiredSyncData;
+	}
+
+	/** Returns the current network tick offset. If we didn't have recorded a network tick, we will still return 0 to keep compatibility with other files
+	 */
+	int32 GetClampedNetworkTickOffset() const
+	{
+		return NetworkTickOffset >= 0 ? NetworkTickOffset : 0;
+	}
 };
 
 struct FChaosVDGameFrameData
@@ -163,6 +200,9 @@ struct CHAOSVDDATA_API FChaosVDRecording
 	 */
 	FName GetSolverFName_AssumedLocked(int32 SolverID);
 
+	bool IsServerSolver_AssumesLocked(int32 SolverID);
+	bool IsServerSolver(int32 SolverID);
+
 	UE_DEPRECATED(5.5, "Please use the GetSolverFName_AssumedLocked instead")
 	FString GetSolverName_AssumedLocked(int32 SolverID) { return TEXT(""); }
 
@@ -192,6 +232,8 @@ struct CHAOSVDDATA_API FChaosVDRecording
 	int32 GetLowestSolverFrameNumberAtCycle(int32 SolverID, uint64 Cycle);
 	int32 GetLowestSolverFrameNumberAtCycle_AssumesLocked(int32 SolverID, uint64 Cycle);
 
+	int32 GetLowestSolverFrameNumberAtNetworkFrameNumber_AssumesLocked(int32 SolverID, int32 NetworkFrameNumber);
+
 	int32 FindFirstSolverKeyFrameNumberFromFrame_AssumesLocked(int32 SolverID, int32 StartFrameNumber);
 	
 	/**
