[ChaosVD] Adding support for recording and visualization of generic debug draw shapes

From Code, shapes can be recorded using the macros

CVD_TRACE_DEBUG_DRAW_BOX
CVD_TRACE_DEBUG_DRAW_LINE
CVD_TRACE_DEBUG_DRAW_SPHERE
CVD_TRACE_DEBUG_DRAW_IMPLICIT_OBJECT

From Blueprint, there is a blueprint function library that adds a new Chaos Visual Debugger Category, that allows you to draw Box, Lines and Spheres.

The macros can tag a Tag that will be shown as debug text (and in the future will be used for search and filtering), color and Solver ID (there are CVD macros to get them using CVD's context system from C++).

With a solver ID, CVD will add the data to the correct solver bucket so it can be hide/shown from the solver timeline controls (and it will work for future filtering features). If no Solver ID is set, CVD will load the debug draw data into a "GameFrames" data bucket which will always be visible.
When the BP function library is used, we get the correct SolverID from the object's world. Which means that in PIE, debug draw things from server or client will be automatically assigned to the correct visibility control

CVD_TRACE_DEBUG_DRAW_IMPLICIT_OBJECT can only be used from C++. This add flexibility to debug draw a specific particle geometry or a shape not currently supported, but CVD_TRACE_DEBUG_DRAW_IMPLICIT_OBJECT is more expensive to execute than the macros for the three basic shapes.

In addition to this, I updated CVD to have a common class for Data Container actors, so this new Game thread bucket can share the same interface and we don't have to manually implement in the CVD Scene or viewport client calls to specific components. This is part of the work to expose a generic way to add external data into CVD via extensions.

diff --git a/Engine/Plugins/ChaosVD/ChaosVD.uplugin b/Engine/Plugins/ChaosVD/ChaosVD.uplugin
--- a/Engine/Plugins/ChaosVD/ChaosVD.uplugin
+++ b/Engine/Plugins/ChaosVD/ChaosVD.uplugin
@@ -20,6 +20,11 @@
 			"Name": "ChaosVD",
 			"Type": "EditorAndProgram",
 			"LoadingPhase": "Default"
+		},
+		{
+			"Name": "ChaosVDBlueprint",
+			"Type": "RuntimeAndProgram",
+			"LoadingPhase": "Default"
 		}
 	],
 	"Plugins": [
diff --git a/Engine/Plugins/ChaosVD/Source/ChaosVD/Private/Actors/ChaosVDDataContainerBaseActor.cpp b/Engine/Plugins/ChaosVD/Source/ChaosVD/Private/Actors/ChaosVDDataContainerBaseActor.cpp
new file mode 100644
--- /dev/null
+++ b/Engine/Plugins/ChaosVD/Source/ChaosVD/Private/Actors/ChaosVDDataContainerBaseActor.cpp
@@ -0,0 +1,36 @@
+// Copyright Epic Games, Inc. All Rights Reserved.
+
+#include "ChaosVDDataContainerBaseActor.h"
+#include "Components/ChaosVDSolverDataComponent.h"
+
+AChaosVDDataContainerBaseActor::AChaosVDDataContainerBaseActor()
+{
+	PrimaryActorTick.bCanEverTick = false;
+}
+
+void AChaosVDDataContainerBaseActor::UpdateFromNewGameFrameData(const FChaosVDGameFrameData& InGameFrameData)
+{
+	TInlineComponentArray<UChaosVDSolverDataComponent*> SolverDataComponents(this);
+	for (UChaosVDSolverDataComponent* Component : SolverDataComponents)
+	{
+		if (Component)
+		{
+			Component->UpdateFromNewGameFrameData(InGameFrameData);
+		}
+	}
+}
+
+void AChaosVDDataContainerBaseActor::Destroyed()
+{
+	TInlineComponentArray<UChaosVDSolverDataComponent*> SolverDataComponents(this);
+	for (UChaosVDSolverDataComponent* Component : SolverDataComponents)
+	{
+		if (Component)
+		{
+			Component->ClearData();
+		}
+	}
+
+	Super::Destroyed();
+}
+
diff --git a/Engine/Plugins/ChaosVD/Source/ChaosVD/Private/Actors/ChaosVDDataContainerBaseActor.h b/Engine/Plugins/ChaosVD/Source/ChaosVD/Private/Actors/ChaosVDDataContainerBaseActor.h
new file mode 100644
--- /dev/null
+++ b/Engine/Plugins/ChaosVD/Source/ChaosVD/Private/Actors/ChaosVDDataContainerBaseActor.h
@@ -0,0 +1,36 @@
+// Copyright Epic Games, Inc. All Rights Reserved.
+
+#pragma once
+
+#include "GameFramework/Actor.h"
+#include "ChaosVDDataContainerBaseActor.generated.h"
+
+class FChaosVDScene;
+
+struct FChaosVDSolverFrameData;
+struct FChaosVDGameFrameData;
+
+/** Base class for any CVD actor that will contain frame related data (either solver frame or game frame) */
+UCLASS(Abstract, NotBlueprintable, NotPlaceable)
+class CHAOSVD_API AChaosVDDataContainerBaseActor : public AActor
+{
+	GENERATED_BODY()
+
+public:
+	AChaosVDDataContainerBaseActor();
+	
+	virtual void UpdateFromNewGameFrameData(const FChaosVDGameFrameData& InGameFrameData);
+
+	virtual void UpdateFromNewSolverFrameData(const FChaosVDSolverFrameData& InSolverFrameData) {};
+
+	virtual void Destroyed() override;
+
+	virtual void SetScene(TWeakPtr<FChaosVDScene> InScene) { SceneWeakPtr = InScene; }
+
+	virtual bool IsVisible() const { return true;}
+	
+	TWeakPtr<FChaosVDScene> GetScene() const { return SceneWeakPtr; }
+	
+protected:
+	TWeakPtr<FChaosVDScene> SceneWeakPtr;
+};
diff --git a/Engine/Plugins/ChaosVD/Source/ChaosVD/Private/Actors/ChaosVDGameFrameInfoActor.cpp b/Engine/Plugins/ChaosVD/Source/ChaosVD/Private/Actors/ChaosVDGameFrameInfoActor.cpp
new file mode 100644
--- /dev/null
+++ b/Engine/Plugins/ChaosVD/Source/ChaosVD/Private/Actors/ChaosVDGameFrameInfoActor.cpp
@@ -0,0 +1,24 @@
+﻿// Copyright Epic Games, Inc. All Rights Reserved.
+
+#include "Actors/ChaosVDGameFrameInfoActor.h"
+
+#include "Components/ChaosVDGenericDebugDrawDataComponent.h"
+#include "Components/ChaosVDSceneQueryDataComponent.h"
+
+AChaosVDGameFrameInfoActor::AChaosVDGameFrameInfoActor()
+{
+	PrimaryActorTick.bCanEverTick = false;
+
+	GenericDebugDrawDataComponent = CreateDefaultSubobject<UChaosVDGenericDebugDrawDataComponent>(TEXT("UChaosVDGenericDebugDrawDataComponent"));
+
+	constexpr int32 GenericGameFrameDataSolverID = INDEX_NONE;
+	GenericDebugDrawDataComponent->SetSolverID(GenericGameFrameDataSolverID);
+}
+
+void AChaosVDGameFrameInfoActor::CleanUp()
+{
+	if (GenericDebugDrawDataComponent)
+	{
+		GenericDebugDrawDataComponent->ClearData();
+	}
+}
diff --git a/Engine/Plugins/ChaosVD/Source/ChaosVD/Private/Actors/ChaosVDGameFrameInfoActor.h b/Engine/Plugins/ChaosVD/Source/ChaosVD/Private/Actors/ChaosVDGameFrameInfoActor.h
new file mode 100644
--- /dev/null
+++ b/Engine/Plugins/ChaosVD/Source/ChaosVD/Private/Actors/ChaosVDGameFrameInfoActor.h
@@ -0,0 +1,29 @@
+﻿// Copyright Epic Games, Inc. All Rights Reserved.
+
+#pragma once
+
+#include "ChaosVDDataContainerBaseActor.h"
+#include "GameFramework/Actor.h"
+#include "UObject/ObjectPtr.h"
+
+#include "ChaosVDGameFrameInfoActor.generated.h"
+
+class UChaosVDGenericDebugDrawDataComponent;
+class UChaosVDSceneQueryDataComponent;
+
+/** Actor that contains game frame related data */
+UCLASS(NotBlueprintable, NotPlaceable)
+class AChaosVDGameFrameInfoActor : public AChaosVDDataContainerBaseActor
+{
+	GENERATED_BODY()
+
+public:
+	AChaosVDGameFrameInfoActor();
+
+	UChaosVDGenericDebugDrawDataComponent* GetGenericDebugDrawDataComponent() const { return GenericDebugDrawDataComponent.Get(); }
+
+	void CleanUp();
+
+	UPROPERTY()
+	TObjectPtr<UChaosVDGenericDebugDrawDataComponent> GenericDebugDrawDataComponent;
+};
diff --git a/Engine/Plugins/ChaosVD/Source/ChaosVD/Private/Actors/ChaosVDSceneQueryDataContainer.cpp b/Engine/Plugins/ChaosVD/Source/ChaosVD/Private/Actors/ChaosVDSceneQueryDataContainer.cpp
deleted file mode 100644
diff --git a/Engine/Plugins/ChaosVD/Source/ChaosVD/Private/Actors/ChaosVDSceneQueryDataContainer.h b/Engine/Plugins/ChaosVD/Source/ChaosVD/Private/Actors/ChaosVDSceneQueryDataContainer.h
deleted file mode 100644
diff --git a/Engine/Plugins/ChaosVD/Source/ChaosVD/Private/Actors/ChaosVDSolverInfoActor.cpp b/Engine/Plugins/ChaosVD/Source/ChaosVD/Private/Actors/ChaosVDSolverInfoActor.cpp
--- a/Engine/Plugins/ChaosVD/Source/ChaosVD/Private/Actors/ChaosVDSolverInfoActor.cpp
+++ b/Engine/Plugins/ChaosVD/Source/ChaosVD/Private/Actors/ChaosVDSolverInfoActor.cpp
@@ -7,8 +7,10 @@
 #include "ChaosVDScene.h"
 #include "ChaosVDSettingsManager.h"
 #include "EditorActorFolders.h"
+#include "Components/ChaosVDGenericDebugDrawDataComponent.h"
 #include "Components/ChaosVDGTAccelerationStructuresDataComponent.h"
 #include "Components/ChaosVDParticleDataComponent.h"
+#include "Components/ChaosVDSceneQueryDataComponent.h"
 #include "Components/ChaosVDSolverCharacterGroundConstraintDataComponent.h"
 #include "Components/ChaosVDSolverCollisionDataComponent.h"
 #include "Components/ChaosVDSolverJointConstraintDataComponent.h"
@@ -19,13 +21,15 @@
 
 #define LOCTEXT_NAMESPACE "ChaosVisualDebugger"
 
-AChaosVDSolverInfoActor::AChaosVDSolverInfoActor(const FObjectInitializer& ObjectInitializer) : Super(ObjectInitializer)
+AChaosVDSolverInfoActor::AChaosVDSolverInfoActor()
 {
 	CollisionDataComponent = CreateDefaultSubobject<UChaosVDSolverCollisionDataComponent>(TEXT("SolverCollisionDataComponent"));
 	ParticleDataComponent = CreateDefaultSubobject<UChaosVDParticleDataComponent>(TEXT("ParticleCollisionDataComponent"));
 	JointsDataComponent = CreateDefaultSubobject<UChaosVDSolverJointConstraintDataComponent>(TEXT("JointDataComponent"));
 	CharacterGroundConstraintDataComponent = CreateDefaultSubobject<UChaosVDSolverCharacterGroundConstraintDataComponent>(TEXT("CharacterGroundConstraintDataComponent"));
 	GTAccelerationStructuresDataComponent = CreateDefaultSubobject<UChaosVDGTAccelerationStructuresDataComponent>(TEXT("GTAccelerationStructuresDataComponent"));
+	SceneQueryDataComponent = CreateDefaultSubobject<UChaosVDSceneQueryDataComponent>(TEXT("ChaosVDSceneQueryDataComponent"));
+	GenericDebugDrawDataComponent = CreateDefaultSubobject<UChaosVDGenericDebugDrawDataComponent>(TEXT("UChaosVDGenericDebugDrawDataComponent"));
 	bIsServer = false;
 
 	if (UChaosVDParticleVisualizationSettings* ParticleVisualizationSettings = FChaosVDSettingsManager::Get().GetSettingsObject<UChaosVDParticleVisualizationSettings>())
@@ -39,6 +43,20 @@ AChaosVDSolverInfoActor::AChaosVDSolverInfoActor(const FObjectInitializer& Objec
 	}
 }
 
+void AChaosVDSolverInfoActor::SetSolverID(int32 InSolverID)
+{
+	SolverID = InSolverID;
+	
+	TInlineComponentArray<UChaosVDSolverDataComponent*> SolverDataComponents(this);
+	for (UChaosVDSolverDataComponent* Component : SolverDataComponents)
+	{
+		if (Component)
+		{
+			Component->SetSolverID(InSolverID);
+		}
+	}
+}
+
 void AChaosVDSolverInfoActor::SetSolverName(const FName& InSolverName)
 {
 	SolverName = InSolverName;
@@ -47,14 +65,14 @@ void AChaosVDSolverInfoActor::SetSolverName(const FName& InSolverName)
 
 void AChaosVDSolverInfoActor::SetScene(TWeakPtr<FChaosVDScene> InScene)
 {
-	FChaosVDSceneObjectBase::SetScene(InScene);
+	Super::SetScene(InScene);
 
 	if (TSharedPtr<FChaosVDScene> ScenePtr = InScene.Pin())
 	{
 		RegisterSelectionSetObject(ScenePtr->GetElementSelectionSet());
 	}
 
-	TInlineComponentArray<UChaosVDSolverDataComponent*> SolverDataComponents;
+	TInlineComponentArray<UChaosVDSolverDataComponent*> SolverDataComponents(this);
 	GetComponents(SolverDataComponents);
 
 	for (UChaosVDSolverDataComponent* Component : SolverDataComponents)
diff --git a/Engine/Plugins/ChaosVD/Source/ChaosVD/Private/Actors/ChaosVDSolverInfoActor.h b/Engine/Plugins/ChaosVD/Source/ChaosVD/Private/Actors/ChaosVDSolverInfoActor.h
--- a/Engine/Plugins/ChaosVD/Source/ChaosVD/Private/Actors/ChaosVDSolverInfoActor.h
+++ b/Engine/Plugins/ChaosVD/Source/ChaosVD/Private/Actors/ChaosVDSolverInfoActor.h
@@ -2,8 +2,8 @@
 
 #pragma once
 
+#include "ChaosVDDataContainerBaseActor.h"
 #include "ChaosVDParticleActor.h"
-#include "ChaosVDSceneObjectBase.h"
 #include "ChaosVDSceneSelectionObserver.h"
 #include "GameFramework/Actor.h"
 #include "ChaosVDSolverInfoActor.generated.h"
@@ -13,21 +13,25 @@ class UChaosVDSolverCharacterGroundConstraintDataComponent;
 class UChaosVDSolverJointConstraintDataComponent;
 struct FChaosVDParticleDataWrapper;
 class AChaosVDParticleActor;
-class UChaosVDSolverCollisionDataComponent;
+class UChaosVDGenericDebugDrawDataComponent;
 class UChaosVDParticleDataComponent;
+class UChaosVDSceneQueryDataComponent;
+class UChaosVDSolverCollisionDataComponent;
+struct FChaosVDGameFrameData;
 
 enum class EChaosVDParticleType : uint8;
 
-UCLASS()
-class AChaosVDSolverInfoActor : public AActor, public FChaosVDSceneObjectBase, public FChaosVDSceneSelectionObserver
+/** Actor that contains all relevant data for the current visualized solver frame */
+UCLASS(NotBlueprintable, NotPlaceable)
+class AChaosVDSolverInfoActor : public AChaosVDDataContainerBaseActor, public FChaosVDSceneSelectionObserver
 {
 	GENERATED_BODY()
 
 public:
 
-	AChaosVDSolverInfoActor(const FObjectInitializer& ObjectInitializer);
+	AChaosVDSolverInfoActor();
 
-	void SetSolverID(int32 InSolverID) { SolverID = InSolverID; }
+	void SetSolverID(int32 InSolverID);
 	int32 GetSolverID() const { return SolverID; }
 
 	void SetSolverName(const FName& InSolverName);
@@ -46,6 +50,8 @@ public:
 	UChaosVDSolverJointConstraintDataComponent* GetJointsDataComponent() { return JointsDataComponent; }
 	UChaosVDSolverCharacterGroundConstraintDataComponent* GetCharacterGroundConstraintDataComponent() { return CharacterGroundConstraintDataComponent; }
 	UChaosVDGTAccelerationStructuresDataComponent* GetGTAccelerationStructuresDataComponent() { return GTAccelerationStructuresDataComponent; }
+	UChaosVDSceneQueryDataComponent* GetSceneQueryDataComponent() const { return SceneQueryDataComponent.Get(); }
+	UChaosVDGenericDebugDrawDataComponent* GetGenericDebugDrawDataComponent() const { return GenericDebugDrawDataComponent.Get(); }
 
 	void RegisterParticleActor(int32 ParticleID, AChaosVDParticleActor* ParticleActor);
 
@@ -66,7 +72,7 @@ public:
 
 	void RemoveSolverFolders(UWorld* World);
 
-	bool IsVisible() const;
+	virtual bool IsVisible() const override;
 
 #if WITH_EDITOR
 	void SetIsTemporarilyHiddenInEditor(bool bIsHidden) override;
@@ -118,6 +124,12 @@ protected:
 
 	UPROPERTY()
 	TObjectPtr<UChaosVDGTAccelerationStructuresDataComponent> GTAccelerationStructuresDataComponent;
+
+	UPROPERTY()
+	TObjectPtr<UChaosVDSceneQueryDataComponent> SceneQueryDataComponent;
+	
+	UPROPERTY()
+	TObjectPtr<UChaosVDGenericDebugDrawDataComponent> GenericDebugDrawDataComponent;
 };
 
 template <typename TCallback>
diff --git a/Engine/Plugins/ChaosVD/Source/ChaosVD/Private/ChaosVDGeometryBuilder.h b/Engine/Plugins/ChaosVD/Source/ChaosVD/Private/ChaosVDGeometryBuilder.h
--- a/Engine/Plugins/ChaosVD/Source/ChaosVD/Private/ChaosVDGeometryBuilder.h
+++ b/Engine/Plugins/ChaosVD/Source/ChaosVD/Private/ChaosVDGeometryBuilder.h
@@ -198,7 +198,8 @@ public:
 	 */
 	TSharedPtr<UE::Geometry::FMeshShapeGenerator> CreateMeshGeneratorForImplicitObject(const Chaos::FImplicitObject* InImplicit, float SimpleShapesComplexityFactor = 1.0f);
 
-private:
+	/** Returns true if the implicit object if of one of the types we need to unpack before generating a mesh for it */
+	bool ImplicitObjectNeedsUnpacking(const Chaos::FImplicitObject* InImplicitObject) const;
 
 	const Chaos::FImplicitObject* UnpackImplicitObject(const Chaos::FImplicitObject* InImplicitObject, Chaos::FRigidTransform3& InOutTransform) const;
 	
@@ -213,9 +214,6 @@ private:
 	 */
 	TSharedPtr<FChaosVDExtractedGeometryDataHandle> ExtractGeometryDataForImplicit(const Chaos::FImplicitObject* InImplicitObject, const Chaos::FRigidTransform3& InTransform);
 
-	/** Returns true if the implicit object if of one of the types we need to unpack before generating a mesh for it */
-	bool ImplicitObjectNeedsUnpacking(const Chaos::FImplicitObject* InImplicitObject) const;
-
 	/**
 	 * Creates a Mesh from the provided Implicit object geometry data. This is a async operation, and the mesh will be assigned to the component once is ready
 	 * @param GeometryCacheKey Key to be used to find this geometry in the cache
diff --git a/Engine/Plugins/ChaosVD/Source/ChaosVD/Private/ChaosVDPlaybackViewportClient.cpp b/Engine/Plugins/ChaosVD/Source/ChaosVD/Private/ChaosVDPlaybackViewportClient.cpp
--- a/Engine/Plugins/ChaosVD/Source/ChaosVD/Private/ChaosVDPlaybackViewportClient.cpp
+++ b/Engine/Plugins/ChaosVD/Source/ChaosVD/Private/ChaosVDPlaybackViewportClient.cpp
@@ -14,7 +14,9 @@
 #include "SEditorViewport.h"
 #include "Selection.h"
 #include "UnrealWidget.h"
+#include "Actors/ChaosVDGameFrameInfoActor.h"
 #include "Actors/ChaosVDSolverInfoActor.h"
+#include "Components/ChaosVDGenericDebugDrawDataComponent.h"
 #include "Components/InstancedStaticMeshComponent.h"
 #include "Visualizers/ChaosVDDebugDrawUtils.h"
 #include "Widgets/SChaosVDMainTab.h"
@@ -238,23 +240,21 @@ void FChaosVDPlaybackViewportClient::Draw(const FSceneView* View, FPrimitiveDraw
 	{
 		//TODO: Currently we can safely assume that any component in these actors is meant to have a visualizer, but we might need a proper interface for these components in the future
 		TInlineComponentArray<const UActorComponent*> ComponentsToVisualize;
-		for (const TPair<int32, AChaosVDSolverInfoActor*>& SolverInfoWithID : ScenePtr->GetSolverInfoActorsMap())
+
+
+		TConstArrayView<TObjectPtr<AChaosVDDataContainerBaseActor>> DataContainerActors = ScenePtr->GetDataContainerActorsView();
+		for (const TObjectPtr<AChaosVDDataContainerBaseActor>& DataContainerActor : DataContainerActors)
 		{
-			if (SolverInfoWithID.Value)
+			if (DataContainerActor)
 			{
 				constexpr bool bIncludeFromChildActors = false;
-				SolverInfoWithID.Value->ForEachComponent(bIncludeFromChildActors, [&ComponentsToVisualize](UActorComponent* Component)
+				DataContainerActor->ForEachComponent(bIncludeFromChildActors, [&ComponentsToVisualize](UActorComponent* Component)
 				{
 					ComponentsToVisualize.Emplace(Component);
 				});
 			}
 		}
 
-		if (const UChaosVDSceneQueryDataComponent* SceneQueryDataComponent = ScenePtr->GetSceneQueryDataContainerComponent())
-		{
-			ComponentsToVisualize.Emplace(SceneQueryDataComponent);
-		}
-
 		for (const UActorComponent* Component : ComponentsToVisualize)
 		{
 			if (!FChaosVDDebugDrawUtils::CanDebugDraw())
diff --git a/Engine/Plugins/ChaosVD/Source/ChaosVD/Private/ChaosVDScene.cpp b/Engine/Plugins/ChaosVD/Source/ChaosVD/Private/ChaosVDScene.cpp
--- a/Engine/Plugins/ChaosVD/Source/ChaosVD/Private/ChaosVDScene.cpp
+++ b/Engine/Plugins/ChaosVD/Source/ChaosVD/Private/ChaosVDScene.cpp
@@ -2,7 +2,6 @@
 
 #include "ChaosVDScene.h"
 
-#include "Actors/ChaosVDSceneQueryDataContainer.h"
 #include "Actors/ChaosVDSolverInfoActor.h"
 #include "ChaosVDGeometryBuilder.h"
 #include "ChaosVDModule.h"
@@ -28,10 +27,10 @@
 #include "Materials/Material.h"
 #include "Misc/ScopedSlowTask.h"
 #include "Selection.h"
+#include "Actors/ChaosVDGameFrameInfoActor.h"
 #include "Settings/ChaosVDCoreSettings.h"
 #include "UObject/Package.h"
 #include "Actors/ChaosVDGeometryContainer.h"
-#include "Components/ChaosVDGTAccelerationStructuresDataComponent.h"
 
 #define LOCTEXT_NAMESPACE "ChaosVisualDebugger"
 
@@ -137,6 +136,7 @@ void FChaosVDScene::AddReferencedObjects(FReferenceCollector& Collector)
 	Collector.AddReferencedObject(ObjectSelection);
 	Collector.AddReferencedObject(ActorSelection);
 	Collector.AddReferencedObject(ComponentSelection);
+	Collector.AddStableReferenceArray(&AvailableDataContainerActors);
 }
 
 void FChaosVDScene::UpdateFromRecordedStepData(const int32 SolverID, const FChaosVDStepData& InRecordedStepData, const FChaosVDSolverFrameData& InFrameData)
@@ -283,31 +283,51 @@ void FChaosVDScene::HandleNewGeometryData(const Chaos::FConstImplicitObjectPtr&
 	}
 }
 
-void FChaosVDScene::CreateSolverInfoActor(int32 SolverID)
+AChaosVDSolverInfoActor* FChaosVDScene::GetOrCreateSolverInfoActor(int32 SolverID)
 {
-	if (!SolverDataContainerBySolverID.Contains(SolverID))
+	if (AChaosVDSolverInfoActor** SolverInfoActorPtrPtr = SolverDataContainerBySolverID.Find(SolverID))
 	{
-		AChaosVDSolverInfoActor* SolverDataInfo = PhysicsVDWorld->SpawnActor<AChaosVDSolverInfoActor>();
-		check(SolverDataInfo);
+		return *SolverInfoActorPtrPtr;
+	}
+	
+	AChaosVDSolverInfoActor* SolverDataInfo = PhysicsVDWorld->SpawnActor<AChaosVDSolverInfoActor>();
+	check(SolverDataInfo);
+
+	FName SolverName = LoadedRecording->GetSolverFName_AssumedLocked(SolverID);
+	FString NameAsString = SolverName.ToString();
+	const bool bIsServer = NameAsString.Contains(TEXT("Server"));
 
-		FName SolverName = LoadedRecording->GetSolverFName_AssumedLocked(SolverID);
-		FString NameAsString = SolverName.ToString();
-		const bool bIsServer = NameAsString.Contains(TEXT("Server"));
+	const FStringFormatOrderedArguments Args {NameAsString, FString::FromInt(SolverID)};
+	const FName FolderPath = *FString::Format(TEXT("Solver {0} | ID {1}"), Args);
 
-		const FStringFormatOrderedArguments Args {NameAsString, FString::FromInt(SolverID)};
-		const FName FolderPath = *FString::Format(TEXT("Solver {0} | ID {1}"), Args);
+	SolverDataInfo->SetFolderPath(FolderPath);
 
-		SolverDataInfo->SetFolderPath(FolderPath);
+	SolverDataInfo->SetSolverID(SolverID);
+	SolverDataInfo->SetSolverName(SolverName);
+	SolverDataInfo->SetScene(AsWeak());
+	SolverDataInfo->SetIsServer(bIsServer);
 
-		SolverDataInfo->SetSolverID(SolverID);
-		SolverDataInfo->SetSolverName(SolverName);
-		SolverDataInfo->SetScene(AsWeak());
-		SolverDataInfo->SetIsServer(bIsServer);
+	SolverDataContainerBySolverID.Add(SolverID, SolverDataInfo);
+	AvailableDataContainerActors.Add(SolverDataInfo);
 
-		SolverDataContainerBySolverID.Add(SolverID, SolverDataInfo);
+	SolverInfoActorCreatedDelegate.Broadcast(SolverDataInfo);
+
+	return SolverDataInfo;
+}
 
-		SolverInfoActorCreatedDelegate.Broadcast(SolverDataInfo);
+AChaosVDGameFrameInfoActor* FChaosVDScene::GetOrCreateGameFrameInfoActor()
+{
+	if (!GameFrameDataInfoActor)
+	{
+		const FName FolderPath("ChaosVisualDebugger/GameFrameData");
+
+		GameFrameDataInfoActor = PhysicsVDWorld->SpawnActor<AChaosVDGameFrameInfoActor>();
+		GameFrameDataInfoActor->SetFolderPath(FolderPath);
+		GameFrameDataInfoActor->SetScene(AsWeak());
+		AvailableDataContainerActors.Add(GameFrameDataInfoActor);
 	}
+
+	return GameFrameDataInfoActor;
 }
 
 void FChaosVDScene::HandleEnterNewGameFrame(int32 FrameNumber, const TArray<int32, TInlineAllocator<16>>& AvailableSolversIds, const FChaosVDGameFrameData& InNewGameFrameData, TArray<int32, TInlineAllocator<16>>& OutRemovedSolversIds)
@@ -325,7 +345,10 @@ void FChaosVDScene::HandleEnterNewGameFrame(int32 FrameNumber, const TArray<int3
 	{
 		AvailableSolversSet.Add(SolverID);
 
-		CreateSolverInfoActor(SolverID);
+		if (AChaosVDSolverInfoActor* SolverInfoActor = GetOrCreateSolverInfoActor(SolverID))
+		{
+			SolverInfoActor->UpdateFromNewGameFrameData(InNewGameFrameData);
+		}
 	}
 
 	int32 AmountRemoved = 0;
@@ -338,6 +361,7 @@ void FChaosVDScene::HandleEnterNewGameFrame(int32 FrameNumber, const TArray<int3
 
 			if (AChaosVDSolverInfoActor* SolverInfoActor = RemoveIterator.Value())
 			{
+				AvailableDataContainerActors.Remove(SolverInfoActor);
 				PhysicsVDWorld->DestroyActor(SolverInfoActor);
 			}
 
@@ -353,26 +377,9 @@ void FChaosVDScene::HandleEnterNewGameFrame(int32 FrameNumber, const TArray<int3
 		CollectGarbage(GARBAGE_COLLECTION_KEEPFLAGS);
 	}
 
-	if (SceneQueriesContainer)
-	{
-		if (UChaosVDSceneQueryDataComponent* QueryDataComponent = SceneQueriesContainer->GetSceneQueryDataComponent())
-		{
-			QueryDataComponent->UpdateQueriesFromFrameData(InNewGameFrameData);
-		}
-	}
-
-	for (const TPair<int32, TArray<TSharedPtr<FChaosVDAABBTreeDataWrapper>>>& AABBTreeDataWithSolverID : InNewGameFrameData.RecordedAABBTreesBySolverID)
+	if (AChaosVDGameFrameInfoActor* GameFrameDataContainer = GetOrCreateGameFrameInfoActor())
 	{
-		if (AChaosVDSolverInfoActor** SolverInfoActorPtrPtr =  SolverDataContainerBySolverID.Find(AABBTreeDataWithSolverID.Key))
-		{
-			if (AChaosVDSolverInfoActor* SolverInfoActorPtr = *SolverInfoActorPtrPtr)
-			{
-				if (UChaosVDGTAccelerationStructuresDataComponent* GTAccelerationStructureDataComponent = SolverInfoActorPtr->GetGTAccelerationStructuresDataComponent())
-				{
-					GTAccelerationStructureDataComponent->UpdateAABBTreeData(AABBTreeDataWithSolverID.Value);
-				}
-			}
-		}
+		GameFrameDataContainer->UpdateFromNewGameFrameData(InNewGameFrameData);
 	}
 }
 
@@ -393,9 +400,13 @@ void FChaosVDScene::HandleEnterNewSolverFrame(int32 FrameNumber, const FChaosVDS
 
 void FChaosVDScene::CleanUpScene(EChaosVDSceneCleanUpOptions Options)
 {
+	// AvailableDataContainerActors should always be at least the number of solver actors created
+	ensure(AvailableDataContainerActors.Num() >= SolverDataContainerBySolverID.Num());
+
+	if (AvailableDataContainerActors.Num() > 0)
 	{
 		constexpr float AmountOfWork = 1.0f;
-		const float PercentagePerElement = 1.0f / SolverDataContainerBySolverID.Num();
+		const float PercentagePerElement = 1.0f / AvailableDataContainerActors.Num();
 
 		FScopedSlowTask CleaningSceneSlowTask(AmountOfWork, LOCTEXT("CleaningupSceneSolverMessage", "Clearing Solver Data ..."));
 		CleaningSceneSlowTask.MakeDialog();
@@ -404,14 +415,20 @@ void FChaosVDScene::CleanUpScene(EChaosVDSceneCleanUpOptions Options)
 
 		if (PhysicsVDWorld)
 		{
-			for (const TPair<int32, AChaosVDSolverInfoActor*>& SolverDataInfoWithID : SolverDataContainerBySolverID)
+			for (TObjectPtr<AChaosVDDataContainerBaseActor>& DataContainerActor : AvailableDataContainerActors)
 			{
-				PhysicsVDWorld->DestroyActor(SolverDataInfoWithID.Value);
+				if (DataContainerActor)
+				{
+					PhysicsVDWorld->DestroyActor(DataContainerActor.Get());
+				}
+
 				CleaningSceneSlowTask.EnterProgressFrame(PercentagePerElement);
 			}
 		}
 
+		AvailableDataContainerActors.Reset();
 		SolverDataContainerBySolverID.Reset();
+		GameFrameDataInfoActor = nullptr;
 	}
 
 	if (Chaos::VisualDebugger::Cvars:: bReInitializeGeometryBuilderOnCleanup && EnumHasAnyFlags(Options, EChaosVDSceneCleanUpOptions::ReInitializeGeometryBuilder))
@@ -427,8 +444,6 @@ void FChaosVDScene::CleanUpScene(EChaosVDSceneCleanUpOptions Options)
 		GeometryGenerator = MakeShared<FChaosVDGeometryBuilder>();
 		GeometryGenerator->Initialize(AsWeak());
 	}
-	
-	SceneQueriesContainer->CleanUp();
 
 	if (EnumHasAnyFlags(Options, EChaosVDSceneCleanUpOptions::CollectGarbage))
 	{
@@ -494,11 +509,6 @@ bool FChaosVDScene::IsSolverForServer(int32 SolverID) const
 	return false;
 }
 
-UChaosVDSceneQueryDataComponent* FChaosVDScene::GetSceneQueryDataContainerComponent() const
-{
-	return SceneQueriesContainer ? SceneQueriesContainer->GetSceneQueryDataComponent() : nullptr;
-}
-
 AChaosVDParticleActor* FChaosVDScene::SpawnParticleFromRecordedData(const TSharedPtr<FChaosVDParticleDataWrapper>& InParticleData, const FChaosVDSolverFrameData& InFrameData)
 {
 	using namespace Chaos;
@@ -567,15 +577,6 @@ void FChaosVDScene::CreateBaseLights(UWorld* TargetWorld) const
 	}
 }
 
-void FChaosVDScene::CreateSceneQueriesContainer(UWorld* TargetWorld)
-{
-	const FName FolderPath("ChaosVisualDebugger/SceneQueries");
-
-	SceneQueriesContainer = TargetWorld->SpawnActor<AChaosVDSceneQueryDataContainer>();
-	SceneQueriesContainer->SetFolderPath(FolderPath);
-	SceneQueriesContainer->SetScene(AsWeak());
-}
-
 AActor* FChaosVDScene::CreateMeshComponentsContainer(UWorld* TargetWorld)
 {
 	const FName GeometryFolderPath("ChaosVisualDebugger/GeneratedMeshComponents");
@@ -612,7 +613,6 @@ UWorld* FChaosVDScene::CreatePhysicsVDWorld()
 	}
 
 	CreateBaseLights(NewWorld);
-	CreateSceneQueriesContainer(NewWorld);
 	CreateMeshComponentsContainer(NewWorld);
 
 	ActorDestroyedHandle = NewWorld->AddOnActorDestroyedHandler(FOnActorDestroyed::FDelegate::CreateRaw(this, &FChaosVDScene::HandleActorDestroyed));
diff --git a/Engine/Plugins/ChaosVD/Source/ChaosVD/Private/ChaosVDScene.h b/Engine/Plugins/ChaosVD/Source/ChaosVD/Private/ChaosVDScene.h
--- a/Engine/Plugins/ChaosVD/Source/ChaosVD/Private/ChaosVDScene.h
+++ b/Engine/Plugins/ChaosVD/Source/ChaosVD/Private/ChaosVDScene.h
@@ -13,12 +13,13 @@
 #include "UObject/ObjectMacros.h"
 #include "UObject/ObjectPtr.h"
 
+class AChaosVDDataContainerBaseActor;
 class AChaosVDGeometryContainer;
 class IChaosVDGeometryOwnerInterface;
 class FChaosVDSelectionCustomization;
 class ITypedElementSelectionInterface;
 struct FTypedElementSelectionOptions;
-class AChaosVDSceneQueryDataContainer;
+class AChaosVDGameFrameInfoActor;
 class AChaosVDSolverInfoActor;
 class AChaosVDSceneCollisionContainer;
 class UChaosVDCoreSettings;
@@ -131,8 +132,6 @@ public:
 
 	AActor* GetMeshComponentsContainerActor() const { return MeshComponentContainerActor; }
 
-	UChaosVDSceneQueryDataComponent* GetSceneQueryDataContainerComponent() const;
-
 	FChaosVDActorUpdatedDelegate& OnActorActiveStateChanged() { return ParticleActorUpdateDelegate; }
 	FChaosVDActorUpdatedDelegate& OnActorLabelChanged() { return ParticleLabelUpdateDelegate; }
 
@@ -145,6 +144,8 @@ public:
 
 	TWeakPtr<FChaosVDSolverDataSelection> GetSolverDataSelectionObject() { return SolverDataSelectionObject ? SolverDataSelectionObject : nullptr;}
 
+	TConstArrayView<TObjectPtr<AChaosVDDataContainerBaseActor>> GetDataContainerActorsView() const { return AvailableDataContainerActors; }
+
 	TSharedPtr<FChaosVDRecording> LoadedRecording;
 
 private:
@@ -159,10 +160,11 @@ private:
 
 	void CreateBaseLights(UWorld* TargetWorld) const;
 
-	void CreateSceneQueriesContainer(UWorld* TargetWorld);
-
 	/** Creates an actor that will contain all solver data for the provided Solver ID*/
-	void CreateSolverInfoActor(int32 SolverID);
+	AChaosVDSolverInfoActor* GetOrCreateSolverInfoActor(int32 SolverID);
+
+	/** Creates an actor that will contain all non-solver data for recorded from any thread*/
+	AChaosVDGameFrameInfoActor* GetOrCreateGameFrameInfoActor();
 
 	AActor* CreateMeshComponentsContainer(UWorld* TargetWorld);
 
@@ -212,7 +214,8 @@ private:
 	mutable AActor* SkySphere = nullptr;
 
 	AActor* MeshComponentContainerActor = nullptr;
-	AChaosVDSceneQueryDataContainer* SceneQueriesContainer = nullptr;
+
+	AChaosVDGameFrameInfoActor* GameFrameDataInfoActor = nullptr;
 
 	bool bIsInitialized = false;
 
@@ -229,5 +232,7 @@ private:
 	
 	TSharedPtr<FChaosVDSolverDataSelection> SolverDataSelectionObject;
 
+	TArray<TObjectPtr<AChaosVDDataContainerBaseActor>> AvailableDataContainerActors;
+
 	friend FChaosVDSelectionCustomization;
 };
diff --git a/Engine/Plugins/ChaosVD/Source/ChaosVD/Private/Components/ChaosVDGTAccelerationStructuresDataComponent.cpp b/Engine/Plugins/ChaosVD/Source/ChaosVD/Private/Components/ChaosVDGTAccelerationStructuresDataComponent.cpp
--- a/Engine/Plugins/ChaosVD/Source/ChaosVD/Private/Components/ChaosVDGTAccelerationStructuresDataComponent.cpp
+++ b/Engine/Plugins/ChaosVD/Source/ChaosVD/Private/Components/ChaosVDGTAccelerationStructuresDataComponent.cpp
@@ -2,6 +2,7 @@
 
 #include "ChaosVDGTAccelerationStructuresDataComponent.h"
 
+#include "ChaosVDRecording.h"
 #include "Algo/Copy.h"
 
 UChaosVDGTAccelerationStructuresDataComponent::UChaosVDGTAccelerationStructuresDataComponent()
@@ -18,6 +19,14 @@ void UChaosVDGTAccelerationStructuresDataComponent::UpdateAABBTreeData(TConstArr
 	Algo::Copy(AABBTreeDataView, RecordedABBTreeData);
 }
 
+void UChaosVDGTAccelerationStructuresDataComponent::UpdateFromNewGameFrameData(const FChaosVDGameFrameData& InGameFrameData)
+{
+	if (const TArray<TSharedPtr<FChaosVDAABBTreeDataWrapper>>* RecordedAABBTreesData = InGameFrameData.RecordedAABBTreesBySolverID.Find(SolverID))
+	{
+		UpdateAABBTreeData(*RecordedAABBTreesData);
+	}
+}
+
 void UChaosVDGTAccelerationStructuresDataComponent::ClearData()
 {
 	RecordedABBTreeData.Reset();
diff --git a/Engine/Plugins/ChaosVD/Source/ChaosVD/Private/Components/ChaosVDGTAccelerationStructuresDataComponent.h b/Engine/Plugins/ChaosVD/Source/ChaosVD/Private/Components/ChaosVDGTAccelerationStructuresDataComponent.h
--- a/Engine/Plugins/ChaosVD/Source/ChaosVD/Private/Components/ChaosVDGTAccelerationStructuresDataComponent.h
+++ b/Engine/Plugins/ChaosVD/Source/ChaosVD/Private/Components/ChaosVDGTAccelerationStructuresDataComponent.h
@@ -5,8 +5,8 @@
 #include "ChaosVDSolverDataComponent.h"
 #include "ChaosVDGTAccelerationStructuresDataComponent.generated.h"
 
-
 struct FChaosVDAABBTreeDataWrapper;
+struct FChaosVDGameFrameData;
 
 UCLASS()
 class UChaosVDGTAccelerationStructuresDataComponent : public UChaosVDSolverDataComponent
@@ -21,6 +21,8 @@ public:
 
 	TConstArrayView<TSharedPtr<FChaosVDAABBTreeDataWrapper>> GetAABBTreeData() const { return RecordedABBTreeData; }
 
+	virtual void UpdateFromNewGameFrameData(const FChaosVDGameFrameData& InGameFrameData) override;
+
 	virtual void ClearData() override;
 protected:
 	
diff --git a/Engine/Plugins/ChaosVD/Source/ChaosVD/Private/Components/ChaosVDGenericDebugDrawDataComponent.cpp b/Engine/Plugins/ChaosVD/Source/ChaosVD/Private/Components/ChaosVDGenericDebugDrawDataComponent.cpp
new file mode 100644
--- /dev/null
+++ b/Engine/Plugins/ChaosVD/Source/ChaosVD/Private/Components/ChaosVDGenericDebugDrawDataComponent.cpp
@@ -0,0 +1,29 @@
+// Copyright Epic Games, Inc. All Rights Reserved.
+
+#include "ChaosVDGenericDebugDrawDataComponent.h"
+
+#include "ChaosVDRecording.h"
+
+UChaosVDGenericDebugDrawDataComponent::UChaosVDGenericDebugDrawDataComponent()
+{
+	SetCanEverAffectNavigation(false);
+	bNavigationRelevant = false;
+	PrimaryComponentTick.bCanEverTick = false;
+}
+
+void UChaosVDGenericDebugDrawDataComponent::UpdateFromNewGameFrameData(const FChaosVDGameFrameData& InGameFrameData)
+{
+	CopyDataFromSourceMap<FChaosVDDebugDrawBoxDataWrapper>(InGameFrameData.RecordedDebugDrawBoxesBySolverID, DebugDrawBoxes, SolverID);
+	CopyDataFromSourceMap<FChaosVDDebugDrawLineDataWrapper>(InGameFrameData.RecordedDebugDrawLinesBySolverID, DebugDrawLines, SolverID);
+	CopyDataFromSourceMap<FChaosVDDebugDrawSphereDataWrapper>(InGameFrameData.RecordedDebugDrawSpheresBySolverID, DebugDrawSpheres, SolverID);
+	CopyDataFromSourceMap<FChaosVDDebugDrawImplicitObjectDataWrapper>(InGameFrameData.RecordedDebugDrawImplicitObjectsBySolverID, DebugDrawImplicitObjects, SolverID);
+}
+
+void UChaosVDGenericDebugDrawDataComponent::ClearData()
+{
+	DebugDrawBoxes.Reset();
+	DebugDrawLines.Reset();
+	DebugDrawSpheres.Reset();
+	DebugDrawLines.Reset();
+	DebugDrawImplicitObjects.Reset();
+}
diff --git a/Engine/Plugins/ChaosVD/Source/ChaosVD/Private/Components/ChaosVDGenericDebugDrawDataComponent.h b/Engine/Plugins/ChaosVD/Source/ChaosVD/Private/Components/ChaosVDGenericDebugDrawDataComponent.h
new file mode 100644
--- /dev/null
+++ b/Engine/Plugins/ChaosVD/Source/ChaosVD/Private/Components/ChaosVDGenericDebugDrawDataComponent.h
@@ -0,0 +1,53 @@
+// Copyright Epic Games, Inc. All Rights Reserved.
+
+#pragma once
+
+#include "ChaosVDSolverDataComponent.h"
+#include "Components/ActorComponent.h"
+#include "DataWrappers/ChaosVDDebugShapeDataWrapper.h"
+#include "ChaosVDGenericDebugDrawDataComponent.generated.h"
+
+class UChaosVDGenericDebugDrawDataComponent;
+struct FChaosVDGameFrameData;
+
+UCLASS()
+class CHAOSVD_API UChaosVDGenericDebugDrawDataComponent : public UChaosVDSolverDataComponent
+{
+	GENERATED_BODY()
+public:
+	UChaosVDGenericDebugDrawDataComponent();
+
+	virtual void UpdateFromNewGameFrameData(const FChaosVDGameFrameData& InGameFrameData) override;
+
+	virtual void ClearData() override;
+
+	TConstArrayView<TSharedPtr<FChaosVDDebugDrawBoxDataWrapper>> GetDebugDrawBoxesDataView() const { return DebugDrawBoxes; };
+	TConstArrayView<TSharedPtr<FChaosVDDebugDrawLineDataWrapper>> GetDebugDrawLinesDataView() const { return DebugDrawLines; };
+	TConstArrayView<TSharedPtr<FChaosVDDebugDrawSphereDataWrapper>> GetDebugDrawSpheresDataView() const { return DebugDrawSpheres; };
+	TConstArrayView<TSharedPtr<FChaosVDDebugDrawImplicitObjectDataWrapper>> GetDebugDrawImplicitObjectsDataView() const { return DebugDrawImplicitObjects; };
+
+private:
+
+	template<typename RecordedDataType>
+	void CopyDataFromSourceMap(const TMap<int32,TArray<TSharedPtr<RecordedDataType>>>& CopyFrom, TArray<TSharedPtr<RecordedDataType>>& CopyTo, int32 SolverID);
+	
+	TArray<TSharedPtr<FChaosVDDebugDrawBoxDataWrapper>> DebugDrawBoxes;
+	TArray<TSharedPtr<FChaosVDDebugDrawLineDataWrapper>> DebugDrawLines;
+	TArray<TSharedPtr<FChaosVDDebugDrawSphereDataWrapper>> DebugDrawSpheres;
+	TArray<TSharedPtr<FChaosVDDebugDrawImplicitObjectDataWrapper>> DebugDrawImplicitObjects;
+};
+
+template <typename RecordedDataType>
+void UChaosVDGenericDebugDrawDataComponent::CopyDataFromSourceMap(const TMap<int32, TArray<TSharedPtr<RecordedDataType>>>& CopyFrom, TArray<TSharedPtr<RecordedDataType>>& CopyTo, int32 SolverID)
+{
+	if (const TArray<TSharedPtr<RecordedDataType>>* RecordedData = CopyFrom.Find(SolverID))
+	{
+		CopyTo.Reset(RecordedData->Num());
+		CopyTo = *RecordedData;
+	}
+	else
+	{
+		CopyTo.Reset();	
+	}
+}
+
diff --git a/Engine/Plugins/ChaosVD/Source/ChaosVD/Private/Components/ChaosVDSceneQueryDataComponent.cpp b/Engine/Plugins/ChaosVD/Source/ChaosVD/Private/Components/ChaosVDSceneQueryDataComponent.cpp
--- a/Engine/Plugins/ChaosVD/Source/ChaosVD/Private/Components/ChaosVDSceneQueryDataComponent.cpp
+++ b/Engine/Plugins/ChaosVD/Source/ChaosVD/Private/Components/ChaosVDSceneQueryDataComponent.cpp
@@ -3,6 +3,8 @@
 #include "Components/ChaosVDSceneQueryDataComponent.h"
 
 #include "ChaosVDRecording.h"
+#include "Actors/ChaosVDSolverInfoActor.h"
+
 
 UChaosVDSceneQueryDataComponent::UChaosVDSceneQueryDataComponent()
 {
@@ -12,15 +14,23 @@ UChaosVDSceneQueryDataComponent::UChaosVDSceneQueryDataComponent()
 	bNavigationRelevant = false;
 }
 
-void UChaosVDSceneQueryDataComponent::UpdateQueriesFromFrameData(const FChaosVDGameFrameData& InGameFrameData)
+void UChaosVDSceneQueryDataComponent::UpdateFromNewGameFrameData(const FChaosVDGameFrameData& InGameFrameData)
 {
-	const int32 RecordedQueriesNum = InGameFrameData.RecordedSceneQueries.Num();
+
+	const TMap<int32, TSharedPtr<FChaosVDQueryDataWrapper>>* RecordedQueriesByQueryID = InGameFrameData.RecordedSceneQueriesBySolverID.Find(SolverID);
+
+	if (!RecordedQueriesByQueryID)
+	{
+		return;
+	}
+	
+	const int32 RecordedQueriesNum = RecordedQueriesByQueryID->Num();
 
 	RecordedQueriesByType.Empty(RecordedQueriesNum);
 	RecordedQueriesByID.Empty(RecordedQueriesNum);
 	RecordedQueries.Empty(RecordedQueriesNum);
 
-	for (const TPair<int32, TSharedPtr<FChaosVDQueryDataWrapper>>& QueryIDPair : InGameFrameData.RecordedSceneQueries)
+	for (const TPair<int32, TSharedPtr<FChaosVDQueryDataWrapper>>& QueryIDPair : (*RecordedQueriesByQueryID))
 	{
 		if (TSharedPtr<FChaosVDQueryDataWrapper> QueryData = QueryIDPair.Value)
 		{
diff --git a/Engine/Plugins/ChaosVD/Source/ChaosVD/Private/Components/ChaosVDSceneQueryDataComponent.h b/Engine/Plugins/ChaosVD/Source/ChaosVD/Private/Components/ChaosVDSceneQueryDataComponent.h
--- a/Engine/Plugins/ChaosVD/Source/ChaosVD/Private/Components/ChaosVDSceneQueryDataComponent.h
+++ b/Engine/Plugins/ChaosVD/Source/ChaosVD/Private/Components/ChaosVDSceneQueryDataComponent.h
@@ -19,7 +19,7 @@ class UChaosVDSceneQueryDataComponent : public UChaosVDSolverDataComponent
 public:
 	UChaosVDSceneQueryDataComponent();
 
-	void UpdateQueriesFromFrameData(const FChaosVDGameFrameData& InGameFrameData);
+	virtual void UpdateFromNewGameFrameData(const FChaosVDGameFrameData& InGameFrameData) override;
 
 	TConstArrayView<TSharedPtr<FChaosVDQueryDataWrapper>> GetQueriesByType(EChaosVDSceneQueryType Type) const;
 	TConstArrayView<TSharedPtr<FChaosVDQueryDataWrapper>> GetAllQueries() const;
diff --git a/Engine/Plugins/ChaosVD/Source/ChaosVD/Private/Components/ChaosVDSolverDataComponent.h b/Engine/Plugins/ChaosVD/Source/ChaosVD/Private/Components/ChaosVDSolverDataComponent.h
--- a/Engine/Plugins/ChaosVD/Source/ChaosVD/Private/Components/ChaosVDSolverDataComponent.h
+++ b/Engine/Plugins/ChaosVD/Source/ChaosVD/Private/Components/ChaosVDSolverDataComponent.h
@@ -6,6 +6,7 @@
 #include "ChaosVDSolverDataComponent.generated.h"
 
 class FChaosVDScene;
+struct FChaosVDGameFrameData;
 /**
  * Base class for all components that stores recorded solver data
  */
@@ -17,8 +18,14 @@ public:
 	virtual void ClearData() PURE_VIRTUAL(UChaosVDSolverDataComponent::ClearData);
 	virtual void SetScene(const TWeakPtr<FChaosVDScene>& InSceneWeakPtr);
 
+	virtual void UpdateFromNewGameFrameData(const FChaosVDGameFrameData& InGameFrameData) {};
+	
+	void SetSolverID(int32 InSolverID) { SolverID = InSolverID; }
+
 protected:
 
 	TWeakPtr<FChaosVDScene> SceneWeakPtr;
+
+	int32 SolverID = INDEX_NONE;
 };
 
diff --git a/Engine/Plugins/ChaosVD/Source/ChaosVD/Private/Settings/ChaosVDAccelerationStructureVisualizationSettings.h b/Engine/Plugins/ChaosVD/Source/ChaosVD/Private/Settings/ChaosVDAccelerationStructureVisualizationSettings.h
--- a/Engine/Plugins/ChaosVD/Source/ChaosVD/Private/Settings/ChaosVDAccelerationStructureVisualizationSettings.h
+++ b/Engine/Plugins/ChaosVD/Source/ChaosVD/Private/Settings/ChaosVDAccelerationStructureVisualizationSettings.h
@@ -24,6 +24,10 @@ enum class EChaosVDAccelerationStructureDataVisualizationFlags : uint32
 	DrawDynamicTrees = 1 << 5,
 	/** If Set, all static trees will be drawn */
 	DrawStaticTrees = 1 << 6,
+	/** If Set, draws the real bounds of each element (not the bounds recorded in the leaf) in the leaves of the tree -
+	 * This should match the bounds recorded in the leaf itself, otherwise it means the AABBtree might have out of sync data */
+	DrawLeavesRealElementBounds = 1 << 7 UMETA(Hidden), // TODO: This will be unhidden in the nex CL with the required object version bump
+
 	/** If set, enabled debug drawing for any recorded acceleration structure available at the current visualized frame */
 	EnableDraw = DrawDynamicTrees | DrawStaticTrees
 };
diff --git a/Engine/Plugins/ChaosVD/Source/ChaosVD/Private/Settings/ChaosVDGenericDebugDrawSettings.cpp b/Engine/Plugins/ChaosVD/Source/ChaosVD/Private/Settings/ChaosVDGenericDebugDrawSettings.cpp
new file mode 100644
--- /dev/null
+++ b/Engine/Plugins/ChaosVD/Source/ChaosVD/Private/Settings/ChaosVDGenericDebugDrawSettings.cpp
@@ -0,0 +1,24 @@
+// Copyright Epic Games, Inc. All Rights Reserved.
+
+#include "ChaosVDGenericDebugDrawSettings.h"
+
+#include "ChaosVDSettingsManager.h"
+
+void UChaosVDGenericDebugDrawSettings::SetDataVisualizationFlags(EChaosVDGenericDebugDrawVisualizationFlags NewFlags)
+{
+	if (UChaosVDGenericDebugDrawSettings* Settings = FChaosVDSettingsManager::Get().GetSettingsObject<UChaosVDGenericDebugDrawSettings>())
+	{
+		Settings->DebugDrawFlags = static_cast<uint32>(NewFlags);
+		Settings->BroadcastSettingsChanged();
+	}
+}
+
+EChaosVDGenericDebugDrawVisualizationFlags UChaosVDGenericDebugDrawSettings::GetDataVisualizationFlags()
+{
+	if (UChaosVDGenericDebugDrawSettings* Settings = FChaosVDSettingsManager::Get().GetSettingsObject<UChaosVDGenericDebugDrawSettings>())
+	{
+		return static_cast<EChaosVDGenericDebugDrawVisualizationFlags>(Settings->DebugDrawFlags);
+	}
+
+	return EChaosVDGenericDebugDrawVisualizationFlags::None;
+}
diff --git a/Engine/Plugins/ChaosVD/Source/ChaosVD/Private/Settings/ChaosVDGenericDebugDrawSettings.h b/Engine/Plugins/ChaosVD/Source/ChaosVD/Private/Settings/ChaosVDGenericDebugDrawSettings.h
new file mode 100644
--- /dev/null
+++ b/Engine/Plugins/ChaosVD/Source/ChaosVD/Private/Settings/ChaosVDGenericDebugDrawSettings.h
@@ -0,0 +1,48 @@
+// Copyright Epic Games, Inc. All Rights Reserved.
+
+#pragma once
+
+#include "ChaosVDCoreSettings.h"
+#include "ChaosVDGenericDebugDrawSettings.generated.h"
+
+/** Set of visualization flags options for Scene Queries */
+UENUM(meta = (Bitflags, UseEnumValuesAsMaskValuesInEditor = "true"))
+enum class EChaosVDGenericDebugDrawVisualizationFlags : uint32
+{
+	None					= 0 UMETA(Hidden),
+	EnableDraw				= 1 << 0,
+	DrawBoxes				= 1 << 1,
+	DrawLines				= 1 << 2,
+	DrawSpheres				= 1 << 3,
+	DrawImplicitObjects		= 1 << 4,
+};
+ENUM_CLASS_FLAGS(EChaosVDGenericDebugDrawVisualizationFlags);
+
+UCLASS(config=ChaosVD, PerObjectConfig)
+class CHAOSVD_API UChaosVDGenericDebugDrawSettings : public UChaosVDSettingsObjectBase
+{
+	GENERATED_BODY()
+
+public:
+
+	/** If true, any debug draw text available will be drawn */
+	UPROPERTY(EditAnywhere, Category=DebugDraw)
+	bool bShowDebugText = false;
+
+	/** The depth priority used for while drawing data. Can be World or Foreground (with this one the shapes will be drawn on top of the geometry and be always visible) */
+	UPROPERTY(EditAnywhere, Category=DebugDraw)
+	TEnumAsByte<ESceneDepthPriorityGroup> DepthPriority = ESceneDepthPriorityGroup::SDPG_Foreground;
+
+	/** Thickness to apply to any debug draw shape controlled by this setting. */
+	UPROPERTY(EditAnywhere, config, Category=DebugDraw)
+	float BaseThickness = 2.0f;
+
+	static void SetDataVisualizationFlags(EChaosVDGenericDebugDrawVisualizationFlags NewFlags);
+	static EChaosVDGenericDebugDrawVisualizationFlags GetDataVisualizationFlags();
+
+private:
+
+	/** Set of flags to enable/disable visualization of debug draw data shapes */
+	UPROPERTY(config, meta = (Bitmask, BitmaskEnum = "/Script/ChaosVD.EChaosVDGenericDebugDrawVisualizationFlags"))
+	uint32 DebugDrawFlags = static_cast<uint32>(EChaosVDGenericDebugDrawVisualizationFlags::DrawBoxes | EChaosVDGenericDebugDrawVisualizationFlags::DrawLines | EChaosVDGenericDebugDrawVisualizationFlags::DrawSpheres);
+};
diff --git a/Engine/Plugins/ChaosVD/Source/ChaosVD/Private/Trace/ChaosVDTraceProvider.cpp b/Engine/Plugins/ChaosVD/Source/ChaosVD/Private/Trace/ChaosVDTraceProvider.cpp
--- a/Engine/Plugins/ChaosVD/Source/ChaosVD/Private/Trace/ChaosVDTraceProvider.cpp
+++ b/Engine/Plugins/ChaosVD/Source/ChaosVD/Private/Trace/ChaosVDTraceProvider.cpp
@@ -10,6 +10,10 @@
 
 #include "Compression/OodleDataCompressionUtil.h"
 #include "DataProcessors/ChaosVDAABBTreeDataProcessor.h"
+#include "DataProcessors/ChaosVDDebugDrawBoxDataProcessor.h"
+#include "DataProcessors/ChaosVDDebugDrawImplicitObjectDataProcessor.h"
+#include "DataProcessors/ChaosVDDebugDrawLineDataProcessor.h"
+#include "DataProcessors/ChaosVDDebugDrawSphereDataProcessor.h"
 #include "Trace/DataProcessors/ChaosVDCharacterGroundConstraintDataProcessor.h"
 #include "Trace/DataProcessors/ChaosVDConstraintDataProcessor.h"
 #include "Trace/DataProcessors/ChaosVDJointConstraintDataProcessor.h"
@@ -279,6 +283,22 @@ void FChaosVDTraceProvider::RegisterDefaultDataProcessorsIfNeeded()
 	AABBTreeDataProcessor->SetTraceProvider(AsShared());
 	RegisterDataProcessor(AABBTreeDataProcessor);
 
+	TSharedPtr<FChaosVDDebugDrawBoxDataProcessor> DebugDrawBoxesDataProcessor = MakeShared<FChaosVDDebugDrawBoxDataProcessor>();
+	DebugDrawBoxesDataProcessor->SetTraceProvider(AsShared());
+	RegisterDataProcessor(DebugDrawBoxesDataProcessor);
+
+	TSharedPtr<FChaosVDDebugDrawLineDataProcessor> DebugDrawLinesDataProcessor = MakeShared<FChaosVDDebugDrawLineDataProcessor>();
+	DebugDrawLinesDataProcessor->SetTraceProvider(AsShared());
+	RegisterDataProcessor(DebugDrawLinesDataProcessor);
+
+	TSharedPtr<FChaosVDDebugDrawSphereDataProcessor> DebugDrawSpheresDataProcessor = MakeShared<FChaosVDDebugDrawSphereDataProcessor>();
+	DebugDrawSpheresDataProcessor->SetTraceProvider(AsShared());
+	RegisterDataProcessor(DebugDrawSpheresDataProcessor);
+
+	TSharedPtr<FChaosVDDebugDrawImplicitObjectDataProcessor> DebugDrawImplicitObjectDataProcessor = MakeShared<FChaosVDDebugDrawImplicitObjectDataProcessor>();
+	DebugDrawImplicitObjectDataProcessor->SetTraceProvider(AsShared());
+	RegisterDataProcessor(DebugDrawImplicitObjectDataProcessor);
+
 	bDefaultDataProcessorsRegistered = true;
 }
 
diff --git a/Engine/Plugins/ChaosVD/Source/ChaosVD/Private/Trace/DataProcessors/ChaosVDDebugDrawBoxDataProcessor.cpp b/Engine/Plugins/ChaosVD/Source/ChaosVD/Private/Trace/DataProcessors/ChaosVDDebugDrawBoxDataProcessor.cpp
new file mode 100644
--- /dev/null
+++ b/Engine/Plugins/ChaosVD/Source/ChaosVD/Private/Trace/DataProcessors/ChaosVDDebugDrawBoxDataProcessor.cpp
@@ -0,0 +1,32 @@
+// Copyright Epic Games, Inc. All Rights Reserved.
+
+#include "ChaosVDDebugDrawBoxDataProcessor.h"
+#include "DataWrappers/ChaosVDDebugShapeDataWrapper.h"
+
+FChaosVDDebugDrawBoxDataProcessor::FChaosVDDebugDrawBoxDataProcessor() : FChaosVDDataProcessorBase(FChaosVDDebugDrawBoxDataWrapper::WrapperTypeName)
+{
+}
+
+bool FChaosVDDebugDrawBoxDataProcessor::ProcessRawData(const TArray<uint8>& InData)
+{
+	FChaosVDDataProcessorBase::ProcessRawData(InData);
+
+	const TSharedPtr<FChaosVDTraceProvider> ProviderSharedPtr = TraceProvider.Pin();
+	if (!ensure(ProviderSharedPtr.IsValid()))
+	{
+		return false;
+	}
+
+	const TSharedPtr<FChaosVDDebugDrawBoxDataWrapper> DebugDrawData = MakeShared<FChaosVDDebugDrawBoxDataWrapper>();
+	const bool bSuccess = Chaos::VisualDebugger::ReadDataFromBuffer(InData, *DebugDrawData, ProviderSharedPtr.ToSharedRef());
+	
+	if (bSuccess)
+	{
+		if (TSharedPtr<FChaosVDGameFrameData> CurrentFrameData = ProviderSharedPtr->GetCurrentGameFrame().Pin())
+		{
+			CurrentFrameData->RecordedDebugDrawBoxesBySolverID.FindOrAdd(DebugDrawData->SolverID).Add(DebugDrawData);
+		}
+	}
+
+	return bSuccess;
+}
diff --git a/Engine/Plugins/ChaosVD/Source/ChaosVD/Private/Trace/DataProcessors/ChaosVDDebugDrawBoxDataProcessor.h b/Engine/Plugins/ChaosVD/Source/ChaosVD/Private/Trace/DataProcessors/ChaosVDDebugDrawBoxDataProcessor.h
new file mode 100644
--- /dev/null
+++ b/Engine/Plugins/ChaosVD/Source/ChaosVD/Private/Trace/DataProcessors/ChaosVDDebugDrawBoxDataProcessor.h
@@ -0,0 +1,15 @@
+// Copyright Epic Games, Inc. All Rights Reserved.
+
+#pragma once
+#include "ChaosVDDataProcessorBase.h"
+
+/**
+ * Data processor implementation that is able to deserialize traced Debug Draw Box shapes
+ */
+class FChaosVDDebugDrawBoxDataProcessor final : public FChaosVDDataProcessorBase
+{
+public:
+	explicit FChaosVDDebugDrawBoxDataProcessor();
+	
+	virtual bool ProcessRawData(const TArray<uint8>& InData) override;
+};
diff --git a/Engine/Plugins/ChaosVD/Source/ChaosVD/Private/Trace/DataProcessors/ChaosVDDebugDrawImplicitObjectDataProcessor.cpp b/Engine/Plugins/ChaosVD/Source/ChaosVD/Private/Trace/DataProcessors/ChaosVDDebugDrawImplicitObjectDataProcessor.cpp
new file mode 100644
--- /dev/null
+++ b/Engine/Plugins/ChaosVD/Source/ChaosVD/Private/Trace/DataProcessors/ChaosVDDebugDrawImplicitObjectDataProcessor.cpp
@@ -0,0 +1,32 @@
+// Copyright Epic Games, Inc. All Rights Reserved.
+
+#include "ChaosVDDebugDrawImplicitObjectDataProcessor.h"
+#include "DataWrappers/ChaosVDDebugShapeDataWrapper.h"
+
+FChaosVDDebugDrawImplicitObjectDataProcessor::FChaosVDDebugDrawImplicitObjectDataProcessor() : FChaosVDDataProcessorBase(FChaosVDDebugDrawImplicitObjectDataWrapper::WrapperTypeName)
+{
+}
+
+bool FChaosVDDebugDrawImplicitObjectDataProcessor::ProcessRawData(const TArray<uint8>& InData)
+{
+	FChaosVDDataProcessorBase::ProcessRawData(InData);
+
+	const TSharedPtr<FChaosVDTraceProvider> ProviderSharedPtr = TraceProvider.Pin();
+	if (!ensure(ProviderSharedPtr.IsValid()))
+	{
+		return false;
+	}
+
+	const TSharedPtr<FChaosVDDebugDrawImplicitObjectDataWrapper> DebugDrawData = MakeShared<FChaosVDDebugDrawImplicitObjectDataWrapper>();
+	const bool bSuccess = Chaos::VisualDebugger::ReadDataFromBuffer(InData, *DebugDrawData, ProviderSharedPtr.ToSharedRef());
+	
+	if (bSuccess)
+	{
+		if (TSharedPtr<FChaosVDGameFrameData> CurrentFrameData = ProviderSharedPtr->GetCurrentGameFrame().Pin())
+		{
+			CurrentFrameData->RecordedDebugDrawImplicitObjectsBySolverID.FindOrAdd(DebugDrawData->SolverID).Add(DebugDrawData);
+		}
+	}
+
+	return bSuccess;
+}
diff --git a/Engine/Plugins/ChaosVD/Source/ChaosVD/Private/Trace/DataProcessors/ChaosVDDebugDrawImplicitObjectDataProcessor.h b/Engine/Plugins/ChaosVD/Source/ChaosVD/Private/Trace/DataProcessors/ChaosVDDebugDrawImplicitObjectDataProcessor.h
new file mode 100644
--- /dev/null
+++ b/Engine/Plugins/ChaosVD/Source/ChaosVD/Private/Trace/DataProcessors/ChaosVDDebugDrawImplicitObjectDataProcessor.h
@@ -0,0 +1,15 @@
+// Copyright Epic Games, Inc. All Rights Reserved.
+
+#pragma once
+#include "ChaosVDDataProcessorBase.h"
+
+/**
+ * Data processor implementation that is able to deserialize traced Debug Draw Implicit Objects
+ */
+class FChaosVDDebugDrawImplicitObjectDataProcessor final : public FChaosVDDataProcessorBase
+{
+public:
+	explicit FChaosVDDebugDrawImplicitObjectDataProcessor();
+	
+	virtual bool ProcessRawData(const TArray<uint8>& InData) override;
+};
diff --git a/Engine/Plugins/ChaosVD/Source/ChaosVD/Private/Trace/DataProcessors/ChaosVDDebugDrawLineDataProcessor.cpp b/Engine/Plugins/ChaosVD/Source/ChaosVD/Private/Trace/DataProcessors/ChaosVDDebugDrawLineDataProcessor.cpp
new file mode 100644
--- /dev/null
+++ b/Engine/Plugins/ChaosVD/Source/ChaosVD/Private/Trace/DataProcessors/ChaosVDDebugDrawLineDataProcessor.cpp
@@ -0,0 +1,32 @@
+// Copyright Epic Games, Inc. All Rights Reserved.
+
+#include "ChaosVDDebugDrawLineDataProcessor.h"
+#include "DataWrappers/ChaosVDDebugShapeDataWrapper.h"
+
+FChaosVDDebugDrawLineDataProcessor::FChaosVDDebugDrawLineDataProcessor() : FChaosVDDataProcessorBase(FChaosVDDebugDrawLineDataWrapper::WrapperTypeName)
+{
+}
+
+bool FChaosVDDebugDrawLineDataProcessor::ProcessRawData(const TArray<uint8>& InData)
+{
+	FChaosVDDataProcessorBase::ProcessRawData(InData);
+
+	const TSharedPtr<FChaosVDTraceProvider> ProviderSharedPtr = TraceProvider.Pin();
+	if (!ensure(ProviderSharedPtr.IsValid()))
+	{
+		return false;
+	}
+
+	const TSharedPtr<FChaosVDDebugDrawLineDataWrapper> DebugDrawData = MakeShared<FChaosVDDebugDrawLineDataWrapper>();
+	const bool bSuccess = Chaos::VisualDebugger::ReadDataFromBuffer(InData, *DebugDrawData, ProviderSharedPtr.ToSharedRef());
+	
+	if (bSuccess)
+	{
+		if (TSharedPtr<FChaosVDGameFrameData> CurrentFrameData = ProviderSharedPtr->GetCurrentGameFrame().Pin())
+		{
+			CurrentFrameData->RecordedDebugDrawLinesBySolverID.FindOrAdd(DebugDrawData->SolverID).Add(DebugDrawData);
+		}
+	}
+
+	return bSuccess;
+}
\ No newline at end of file
diff --git a/Engine/Plugins/ChaosVD/Source/ChaosVD/Private/Trace/DataProcessors/ChaosVDDebugDrawLineDataProcessor.h b/Engine/Plugins/ChaosVD/Source/ChaosVD/Private/Trace/DataProcessors/ChaosVDDebugDrawLineDataProcessor.h
new file mode 100644
--- /dev/null
+++ b/Engine/Plugins/ChaosVD/Source/ChaosVD/Private/Trace/DataProcessors/ChaosVDDebugDrawLineDataProcessor.h
@@ -0,0 +1,16 @@
+// Copyright Epic Games, Inc. All Rights Reserved.
+
+#pragma once
+#include "ChaosVDDataProcessorBase.h"
+
+/**
+ * Data processor implementation that is able to deserialize traced Debug Draw lines
+ */
+class FChaosVDDebugDrawLineDataProcessor final : public FChaosVDDataProcessorBase
+{
+public:
+	explicit FChaosVDDebugDrawLineDataProcessor();
+	
+	virtual bool ProcessRawData(const TArray<uint8>& InData) override;
+};
+
diff --git a/Engine/Plugins/ChaosVD/Source/ChaosVD/Private/Trace/DataProcessors/ChaosVDDebugDrawSphereDataProcessor.cpp b/Engine/Plugins/ChaosVD/Source/ChaosVD/Private/Trace/DataProcessors/ChaosVDDebugDrawSphereDataProcessor.cpp
new file mode 100644
--- /dev/null
+++ b/Engine/Plugins/ChaosVD/Source/ChaosVD/Private/Trace/DataProcessors/ChaosVDDebugDrawSphereDataProcessor.cpp
@@ -0,0 +1,32 @@
+// Copyright Epic Games, Inc. All Rights Reserved.
+
+#include "ChaosVDDebugDrawSphereDataProcessor.h"
+#include "DataWrappers/ChaosVDDebugShapeDataWrapper.h"
+
+FChaosVDDebugDrawSphereDataProcessor::FChaosVDDebugDrawSphereDataProcessor() : FChaosVDDataProcessorBase(FChaosVDDebugDrawSphereDataWrapper::WrapperTypeName)
+{
+}
+
+bool FChaosVDDebugDrawSphereDataProcessor::ProcessRawData(const TArray<uint8>& InData)
+{
+	FChaosVDDataProcessorBase::ProcessRawData(InData);
+
+	const TSharedPtr<FChaosVDTraceProvider> ProviderSharedPtr = TraceProvider.Pin();
+	if (!ensure(ProviderSharedPtr.IsValid()))
+	{
+		return false;
+	}
+
+	const TSharedPtr<FChaosVDDebugDrawSphereDataWrapper> DebugDrawData = MakeShared<FChaosVDDebugDrawSphereDataWrapper>();
+	const bool bSuccess = Chaos::VisualDebugger::ReadDataFromBuffer(InData, *DebugDrawData, ProviderSharedPtr.ToSharedRef());
+	
+	if (bSuccess)
+	{
+		if (TSharedPtr<FChaosVDGameFrameData> CurrentFrameData = ProviderSharedPtr->GetCurrentGameFrame().Pin())
+		{
+			CurrentFrameData->RecordedDebugDrawSpheresBySolverID.FindOrAdd(DebugDrawData->SolverID).Add(DebugDrawData);
+		}
+	}
+
+	return bSuccess;
+}
diff --git a/Engine/Plugins/ChaosVD/Source/ChaosVD/Private/Trace/DataProcessors/ChaosVDDebugDrawSphereDataProcessor.h b/Engine/Plugins/ChaosVD/Source/ChaosVD/Private/Trace/DataProcessors/ChaosVDDebugDrawSphereDataProcessor.h
new file mode 100644
--- /dev/null
+++ b/Engine/Plugins/ChaosVD/Source/ChaosVD/Private/Trace/DataProcessors/ChaosVDDebugDrawSphereDataProcessor.h
@@ -0,0 +1,16 @@
+// Copyright Epic Games, Inc. All Rights Reserved.
+
+#pragma once
+#include "ChaosVDDataProcessorBase.h"
+
+/**
+ * Data processor implementation that is able to deserialize traced Debug Draw sphere shapes
+ */
+class FChaosVDDebugDrawSphereDataProcessor final : public FChaosVDDataProcessorBase
+{
+public:
+	explicit FChaosVDDebugDrawSphereDataProcessor();
+	
+	virtual bool ProcessRawData(const TArray<uint8>& InData) override;
+};
+
diff --git a/Engine/Plugins/ChaosVD/Source/ChaosVD/Private/Trace/DataProcessors/ChaosVDSceneQueryDataProcessor.cpp b/Engine/Plugins/ChaosVD/Source/ChaosVD/Private/Trace/DataProcessors/ChaosVDSceneQueryDataProcessor.cpp
--- a/Engine/Plugins/ChaosVD/Source/ChaosVD/Private/Trace/DataProcessors/ChaosVDSceneQueryDataProcessor.cpp
+++ b/Engine/Plugins/ChaosVD/Source/ChaosVD/Private/Trace/DataProcessors/ChaosVDSceneQueryDataProcessor.cpp
@@ -33,13 +33,17 @@ bool FChaosVDSceneQueryDataProcessor::ProcessRawData(const TArray<uint8>& InData
 			// If ParentQueryID was set, this is a sub query, so find the parent add it to the sub-queries list so we can navigate trough the query "hierarchy" later on
 			if (QueryData->ParentQueryID != INDEX_NONE)
 			{
-				if (const TSharedPtr<FChaosVDQueryDataWrapper>* ParentQueryData = CurrentFrameData->RecordedSceneQueries.Find(QueryData->ParentQueryID))
+				if (TMap<int32, TSharedPtr<FChaosVDQueryDataWrapper>>* ParentQueryDataByQueryIDPtr = CurrentFrameData->RecordedSceneQueriesBySolverID.Find(QueryData->WorldSolverID))
 				{
-					(*ParentQueryData)->SubQueriesIDs.Add(QueryData->ID);
-				}
+					if (const TSharedPtr<FChaosVDQueryDataWrapper>* ParentQueryData = ParentQueryDataByQueryIDPtr->Find(QueryData->ParentQueryID))
+					{
+						(*ParentQueryData)->SubQueriesIDs.Add(QueryData->ID);
+					}
+				}	
 			}
 
-			CurrentFrameData->RecordedSceneQueries.Add(QueryData->ID, QueryData);
+			CurrentFrameData->RecordedSceneQueriesByQueryID.Add(QueryData->ID, QueryData);
+			CurrentFrameData->RecordedSceneQueriesBySolverID.FindOrAdd(QueryData->WorldSolverID).Add(QueryData->ID, QueryData);
 		}
 	}
 
diff --git a/Engine/Plugins/ChaosVD/Source/ChaosVD/Private/Trace/DataProcessors/ChaosVDSceneQueryVisitDataProcessor.cpp b/Engine/Plugins/ChaosVD/Source/ChaosVD/Private/Trace/DataProcessors/ChaosVDSceneQueryVisitDataProcessor.cpp
--- a/Engine/Plugins/ChaosVD/Source/ChaosVD/Private/Trace/DataProcessors/ChaosVDSceneQueryVisitDataProcessor.cpp
+++ b/Engine/Plugins/ChaosVD/Source/ChaosVD/Private/Trace/DataProcessors/ChaosVDSceneQueryVisitDataProcessor.cpp
@@ -29,7 +29,7 @@ bool FChaosVDSceneQueryVisitDataProcessor::ProcessRawData(const TArray<uint8>& I
 	{
 		if (const TSharedPtr<FChaosVDGameFrameData> CurrentFrameData = ProviderSharedPtr->GetCurrentGameFrame().Pin())
 		{	
-			if (TSharedPtr<FChaosVDQueryDataWrapper>* QueryDataPtrPtr = CurrentFrameData->RecordedSceneQueries.Find(VisitStepData.OwningQueryID))
+			if (TSharedPtr<FChaosVDQueryDataWrapper>* QueryDataPtrPtr = CurrentFrameData->RecordedSceneQueriesByQueryID.Find(VisitStepData.OwningQueryID))
 			{
 				TSharedPtr<FChaosVDQueryDataWrapper> QueryDataPtr = *QueryDataPtrPtr;
 				if (QueryDataPtrPtr->IsValid())
diff --git a/Engine/Plugins/ChaosVD/Source/ChaosVD/Private/Visualizers/ChaosVDDebugDrawUtils.cpp b/Engine/Plugins/ChaosVD/Source/ChaosVD/Private/Visualizers/ChaosVDDebugDrawUtils.cpp
--- a/Engine/Plugins/ChaosVD/Source/ChaosVD/Private/Visualizers/ChaosVDDebugDrawUtils.cpp
+++ b/Engine/Plugins/ChaosVD/Source/ChaosVD/Private/Visualizers/ChaosVDDebugDrawUtils.cpp
@@ -252,20 +252,51 @@ void FChaosVDDebugDrawUtils::DrawImplicitObject(FPrimitiveDrawInterface* PDI, co
 	  return;
 	}
 
+	using namespace Chaos;
+	const EImplicitObjectType InnerType = GetInnerType(ImplicitObject->GetType());
+		
+	if (InnerType == ImplicitObjectType::Union || InnerType == ImplicitObjectType::UnionClustered)
+	{
+		if (const FImplicitObjectUnion* Union = ImplicitObject->template AsA<FImplicitObjectUnion>())
+		{
+			for (int32 ObjectIndex = 0; ObjectIndex < Union->GetObjects().Num(); ++ObjectIndex)
+			{
+				const FImplicitObjectPtr& UnionImplicit = Union->GetObjects()[ObjectIndex];
+				DrawImplicitObject(PDI, GeometryGenerator, UnionImplicit.GetReference(), InWorldTransform, InColor, InDebugText, DepthPriority, Thickness);
+			}
+		}
+
+		return;
+	}
+
+	if (InnerType == ImplicitObjectType::Transformed)
+	{
+		if (const TImplicitObjectTransformed<FReal, 3>* Transformed = ImplicitObject->template GetObject<TImplicitObjectTransformed<FReal, 3>>())
+		{
+			DrawImplicitObject(PDI, GeometryGenerator, Transformed->GetTransformedObject(), Transformed->GetTransform() * InWorldTransform, InColor, InDebugText, DepthPriority, Thickness);
+		}
+		
+		return;
+	}
+
 	constexpr float SimpleShapesComplexityFactor = 0.5f;
-	if (const TSharedPtr<UE::Geometry::FMeshShapeGenerator> MeshGenerator = GeometryGenerator->CreateMeshGeneratorForImplicitObject(ImplicitObject, SimpleShapesComplexityFactor))
+
+	FRigidTransform3 ExtractedTransform = InWorldTransform;
+	const bool bNeedsUnpack = GeometryGenerator->ImplicitObjectNeedsUnpacking(ImplicitObject);
+	const FImplicitObject* ImplicitObjectToProcess = bNeedsUnpack ? GeometryGenerator->UnpackImplicitObject(ImplicitObject, ExtractedTransform) : ImplicitObject.GetReference();
+
+	if (const TSharedPtr<UE::Geometry::FMeshShapeGenerator> MeshGenerator = GeometryGenerator->CreateMeshGeneratorForImplicitObject(ImplicitObjectToProcess, SimpleShapesComplexityFactor))
 	{
-		FTransform AdjustedTransform = InWorldTransform;
-		GeometryGenerator->AdjustedTransformForImplicit(ImplicitObject, AdjustedTransform);
+		GeometryGenerator->AdjustedTransformForImplicit(ImplicitObject, ExtractedTransform);
 		MeshGenerator->Generate();
 
 		PDI->AddReserveLines(DepthPriority, MeshGenerator->Triangles.Num() * 3, false, Thickness > SMALL_NUMBER);
 
 		for (const UE::Geometry::FIndex3i& Triangle : MeshGenerator->Triangles)
 		{
-			FVector VertexA = AdjustedTransform.TransformPosition(MeshGenerator->Vertices[Triangle.A]);
-			FVector VertexB = AdjustedTransform.TransformPosition(MeshGenerator->Vertices[Triangle.B]);
-			FVector VertexC = AdjustedTransform.TransformPosition(MeshGenerator->Vertices[Triangle.C]);
+			FVector VertexA = ExtractedTransform.TransformPosition(MeshGenerator->Vertices[Triangle.A]);
+			FVector VertexB = ExtractedTransform.TransformPosition(MeshGenerator->Vertices[Triangle.B]);
+			FVector VertexC = ExtractedTransform.TransformPosition(MeshGenerator->Vertices[Triangle.C]);
 	
 			DrawLine(PDI, VertexA, VertexB, InColor, FText::GetEmpty(), DepthPriority, Thickness);
 			DrawLine(PDI, VertexB, VertexC, InColor, FText::GetEmpty(), DepthPriority, Thickness);
@@ -279,6 +310,55 @@ void FChaosVDDebugDrawUtils::DrawImplicitObject(FPrimitiveDrawInterface* PDI, co
 	}
 }
 
+void FChaosVDDebugDrawUtils::DrawSphere(FPrimitiveDrawInterface* PDI, const FVector& Center, float Radius, int32 Segments, const FColor& InColor, const FText& InDebugText, ESceneDepthPriorityGroup DepthPriority, float Thickness)
+{
+	// TODO: This implementation is taken as is from ULineBatchComponent::DrawSphere. At some point in the future we want to migrate CVD to use the line batcher directly, but we need to add support for hit proxies
+
+	// Need at least 4 segments
+	Segments = FMath::Max(Segments, 4);
+
+	const float AngleInc = 2.f * UE_PI / Segments;
+	int32 NumSegmentsY = Segments;
+	float Latitude = AngleInc;
+	float SinY1 = 0.0f, CosY1 = 1.0f;
+
+	PDI->AddReserveLines(DepthPriority, NumSegmentsY * Segments * 2, false, Thickness > SMALL_NUMBER);
+	while (NumSegmentsY--)
+	{
+		const float SinY2 = FMath::Sin(Latitude);
+		const float CosY2 = FMath::Cos(Latitude);
+
+		FVector Vertex1 = FVector(SinY1, 0.0f, CosY1) * Radius + Center;
+		FVector Vertex3 = FVector(SinY2, 0.0f, CosY2) * Radius + Center;
+		float Longitude = AngleInc;
+
+		int32 NumSegmentsX = Segments;
+		while (NumSegmentsX--)
+		{
+			const float SinX = FMath::Sin(Longitude);
+			const float CosX = FMath::Cos(Longitude);
+
+			const FVector Vertex2 = FVector((CosX * SinY1), (SinX * SinY1), CosY1) * Radius + Center;
+			const FVector Vertex4 = FVector((CosX * SinY2), (SinX * SinY2), CosY2) * Radius + Center;
+
+			DrawLine(PDI, Vertex1, Vertex2, InColor, FText::GetEmpty(), DepthPriority, Thickness);
+			DrawLine(PDI, Vertex1, Vertex3, InColor, FText::GetEmpty(), DepthPriority, Thickness);
+
+			Vertex1 = Vertex2;
+			Vertex3 = Vertex4;
+			Longitude += AngleInc;
+		}
+		SinY1 = SinY2;
+		CosY1 = CosY2;
+		Latitude += AngleInc;
+	}
+
+	if (!InDebugText.IsEmpty())
+	{
+		DrawText(InDebugText, Center, InColor);
+	}
+}
+
 void FChaosVDDebugDrawUtils::DrawCanvas(FViewport& InViewport, FSceneView& View, FCanvas& Canvas)
 {
 	if (!GEngine)
diff --git a/Engine/Plugins/ChaosVD/Source/ChaosVD/Private/Visualizers/ChaosVDDebugDrawUtils.h b/Engine/Plugins/ChaosVD/Source/ChaosVD/Private/Visualizers/ChaosVDDebugDrawUtils.h
--- a/Engine/Plugins/ChaosVD/Source/ChaosVD/Private/Visualizers/ChaosVDDebugDrawUtils.h
+++ b/Engine/Plugins/ChaosVD/Source/ChaosVD/Private/Visualizers/ChaosVDDebugDrawUtils.h
@@ -30,6 +30,8 @@ public:
 	static void DrawBox(FPrimitiveDrawInterface* PDI, const FVector& InExtents, const FColor& InColor, const FTransform& InTransform, const FText& InDebugText, ESceneDepthPriorityGroup DepthPriority = SDPG_World, float Thickness = 2.0f);
 	static void DrawLine(FPrimitiveDrawInterface* PDI, const FVector& InStartPosition, const FVector& InEndPosition, const FColor& InColor, const FText& InDebugText, ESceneDepthPriorityGroup DepthPriority = SDPG_World, float Thickness = 2.0f);
 	static void DrawImplicitObject(FPrimitiveDrawInterface* PDI, const TSharedPtr<FChaosVDGeometryBuilder>& GeometryGenerator, const Chaos::FConstImplicitObjectPtr& ImplicitObject, const FTransform& InWorldTransform, const FColor& InColor, const FText& InDebugText, ESceneDepthPriorityGroup DepthPriority = SDPG_World, float Thickness = 2.0f);
+	
+	static void DrawSphere(FPrimitiveDrawInterface* PDI, const FVector& Center, float Radius, int32 Segments, const FColor& InColor, const FText& InDebugText, ESceneDepthPriorityGroup DepthPriority = SDPG_World, float Thickness = 2.0f);
 
 	static void DrawCanvas(FViewport& InViewport, FSceneView& View, FCanvas& Canvas);
 	
diff --git a/Engine/Plugins/ChaosVD/Source/ChaosVD/Private/Visualizers/ChaosVDGTAccelerationStructureDataComponentVisualizer.cpp b/Engine/Plugins/ChaosVD/Source/ChaosVD/Private/Visualizers/ChaosVDGTAccelerationStructureDataComponentVisualizer.cpp
--- a/Engine/Plugins/ChaosVD/Source/ChaosVD/Private/Visualizers/ChaosVDGTAccelerationStructureDataComponentVisualizer.cpp
+++ b/Engine/Plugins/ChaosVD/Source/ChaosVD/Private/Visualizers/ChaosVDGTAccelerationStructureDataComponentVisualizer.cpp
@@ -357,6 +357,11 @@ void FChaosVDGTAccelerationStructureDataComponentVisualizer::DrawAABBTreeArrayLe
 		{	
 			Chaos::VisualDebugger::Utils::DrawFBoxAtLocation(PDI, TreeArrayLeafElement.Bounds, ColorByDensity, VisualizationContext.DepthPriority, FinalThickness * 0.7f);
 		}
+
+		if (VisualizationContext.IsVisualizationFlagEnabled(EChaosVDAccelerationStructureDataVisualizationFlags::DrawLeavesRealElementBounds))
+		{
+			Chaos::VisualDebugger::Utils::DrawFBoxAtLocation(PDI, TreeArrayLeafElement.ActualBounds, FColor::Red, VisualizationContext.DepthPriority, FinalThickness * 0.7f);
+		}
 	}
 
 	PDI->SetHitProxy(nullptr);
diff --git a/Engine/Plugins/ChaosVD/Source/ChaosVD/Private/Visualizers/ChaosVDGenericDebugDrawDataComponentVisualizer.cpp b/Engine/Plugins/ChaosVD/Source/ChaosVD/Private/Visualizers/ChaosVDGenericDebugDrawDataComponentVisualizer.cpp
new file mode 100644
--- /dev/null
+++ b/Engine/Plugins/ChaosVD/Source/ChaosVD/Private/Visualizers/ChaosVDGenericDebugDrawDataComponentVisualizer.cpp
@@ -0,0 +1,203 @@
+// Copyright Epic Games, Inc. All Rights Reserved.
+
+#include "ChaosVDGenericDebugDrawDataComponentVisualizer.h"
+
+#include "Actors/ChaosVDGameFrameInfoActor.h"
+#include "ChaosVDDebugDrawUtils.h"
+#include "ChaosVDScene.h"
+#include "Components/ChaosVDGenericDebugDrawDataComponent.h"
+#include "SceneView.h"
+#include "Widgets/SChaosVDViewportToolbar.h"
+
+class AChaosVDGameFrameInfoActor;
+
+#define LOCTEXT_NAMESPACE "ChaosVisualDebugger"
+
+FChaosVDGenericDebugDrawDataComponentVisualizer::FChaosVDGenericDebugDrawDataComponentVisualizer()
+{
+	FChaosVDGenericDebugDrawDataComponentVisualizer::RegisterVisualizerMenus();
+}
+
+void FChaosVDGenericDebugDrawDataComponentVisualizer::RegisterVisualizerMenus()
+{
+	FName MenuSection("GenericDebugDrawDataVisualization.Show");
+	FText MenuSectionLabel = LOCTEXT("GenericDebugDrawDataShowMenuLabel", "Generic Debug Draw Data Visualization");
+	FText FlagsMenuLabel = LOCTEXT("GenericDebugDrawDataFlagsMenuLabel", "Generic Debug Draw Data Flags");
+	FText FlagsMenuTooltip = LOCTEXT("GenericDebugDrawDataFlagsMenuToolTip", "Set of flags to enable/disable visibility of specific types of Debug Draw Data that are not solver related");
+	FSlateIcon FlagsMenuIcon = FSlateIcon(FAppStyle::Get().GetStyleSetName(), TEXT("EditorViewport.CollisionVisibility"));
+
+	FText SettingsMenuLabel = LOCTEXT("GenericDebugDrawDataSettingsMenuLabel", "Generic Debug Draw Data Visualization Settings");
+	FText SettingsMenuTooltip = LOCTEXT("GenericDebugDrawDataSettingsMenuToolTip", "Options to change how the recorded Generic Debug Draw Data is debug drawn");
+	
+	CreateGenericVisualizerMenu<UChaosVDGenericDebugDrawSettings, EChaosVDGenericDebugDrawVisualizationFlags>(SChaosVDViewportToolbar::ShowMenuName, MenuSection, MenuSectionLabel, FlagsMenuLabel, FlagsMenuTooltip, FlagsMenuIcon, SettingsMenuLabel, SettingsMenuTooltip);
+}
+
+void FChaosVDGenericDebugDrawDataComponentVisualizer::DrawVisualization(const UActorComponent* Component, const FSceneView* View, FPrimitiveDrawInterface* PDI)
+{
+	const UChaosVDGenericDebugDrawDataComponent* DebugDrawDataComponent = Cast<UChaosVDGenericDebugDrawDataComponent>(Component);
+
+	if (!DebugDrawDataComponent)
+	{
+		return;
+	}
+
+		
+	AChaosVDDataContainerBaseActor* InfoActor = Cast<AChaosVDDataContainerBaseActor>(Component->GetOwner());
+	if (!InfoActor)
+	{
+		return;
+	}
+
+	if (!InfoActor->IsVisible())
+	{
+		return;
+	}
+
+	const TSharedPtr<FChaosVDScene> CVDScene = InfoActor->GetScene().Pin();
+	
+	FChaosVDGenericDebugDrawDataVisualizationSettings VisualizationContext;
+	VisualizationContext.CVDScene = CVDScene;
+	VisualizationContext.SolverDataSelectionObject = CVDScene->GetSolverDataSelectionObject().Pin();
+	VisualizationContext.DataComponent = DebugDrawDataComponent;
+
+	if (const UChaosVDGenericDebugDrawSettings* EditorSettings = FChaosVDSettingsManager::Get().GetSettingsObject<UChaosVDGenericDebugDrawSettings>())
+	{
+		VisualizationContext.VisualizationFlags = static_cast<uint32>(UChaosVDGenericDebugDrawSettings::GetDataVisualizationFlags());
+		VisualizationContext.DebugDrawSettings = EditorSettings;
+		VisualizationContext.DepthPriority = EditorSettings->DepthPriority;
+		VisualizationContext.Thickness = EditorSettings->BaseThickness;
+		VisualizationContext.bShowDebugText = EditorSettings->bShowDebugText;
+	}
+
+	if (!VisualizationContext.IsVisualizationFlagEnabled(EChaosVDGenericDebugDrawVisualizationFlags::EnableDraw))
+	{
+		return;
+	}
+
+	DrawBoxes(PDI, View, VisualizationContext);
+
+	DrawLines(PDI, View, VisualizationContext);
+
+	DrawSpheres(PDI, View, VisualizationContext);
+
+	DrawImplicitObjects(PDI, View, VisualizationContext);
+}
+
+void FChaosVDGenericDebugDrawDataComponentVisualizer::DrawBoxes(FPrimitiveDrawInterface* PDI, const FSceneView* View, const FChaosVDGenericDebugDrawDataVisualizationSettings& InVisualizationContext)
+{
+	if (InVisualizationContext.IsVisualizationFlagEnabled(EChaosVDGenericDebugDrawVisualizationFlags::DrawBoxes))
+	{
+		TConstArrayView<TSharedPtr<FChaosVDDebugDrawBoxDataWrapper>> DebugDrawBoxesView = InVisualizationContext.DataComponent->GetDebugDrawBoxesDataView();
+		for (const TSharedPtr<FChaosVDDebugDrawBoxDataWrapper>& DebugDrawBox : DebugDrawBoxesView)
+		{
+			if (DebugDrawBox)
+			{
+				FVector BoxCenter = DebugDrawBox->Box.GetCenter();
+				FVector BoxExtent = DebugDrawBox->Box.GetExtent();
+
+				if (!View->ViewFrustum.IntersectBox(BoxCenter, BoxExtent))
+				{
+					continue;
+				}
+				
+				const FText& DebugText = InVisualizationContext.bShowDebugText ? FText::FromName(DebugDrawBox->Tag) : FText::GetEmpty();
+	
+				FTransform BoxTransform(BoxCenter);
+				FChaosVDDebugDrawUtils::DrawBox(PDI, BoxExtent, DebugDrawBox->Color, BoxTransform, DebugText, InVisualizationContext.DepthPriority, InVisualizationContext.Thickness);			
+			}
+		}
+	}
+}
+
+void FChaosVDGenericDebugDrawDataComponentVisualizer::DrawLines(FPrimitiveDrawInterface* PDI, const FSceneView* View, const FChaosVDGenericDebugDrawDataVisualizationSettings& InVisualizationContext)
+{
+	if (InVisualizationContext.IsVisualizationFlagEnabled(EChaosVDGenericDebugDrawVisualizationFlags::DrawLines))
+	{
+		TConstArrayView<TSharedPtr<FChaosVDDebugDrawLineDataWrapper>> DebugDrawLinesView = InVisualizationContext.DataComponent->GetDebugDrawLinesDataView();
+		for (const TSharedPtr<FChaosVDDebugDrawLineDataWrapper>& DebugDrawLine : DebugDrawLinesView)
+		{
+			if (DebugDrawLine)
+			{
+				if (!View->ViewFrustum.IntersectLineSegment(DebugDrawLine->StartLocation, DebugDrawLine->EndLocation))
+				{
+					continue;
+				}
+
+				const FText& DebugText = InVisualizationContext.bShowDebugText ? FText::FromName(DebugDrawLine->Tag) : FText::GetEmpty();
+
+				if (DebugDrawLine->bIsArrow)
+				{
+					FChaosVDDebugDrawUtils::DrawArrowVector(PDI, DebugDrawLine->StartLocation, DebugDrawLine->EndLocation, DebugText, DebugDrawLine->Color, InVisualizationContext.DepthPriority, InVisualizationContext.Thickness);
+				}
+				else
+				{
+					FChaosVDDebugDrawUtils::DrawLine(PDI, DebugDrawLine->StartLocation, DebugDrawLine->EndLocation, DebugDrawLine->Color, DebugText, InVisualizationContext.DepthPriority, InVisualizationContext.Thickness);
+				}			
+			}
+		}
+	}
+}
+
+void FChaosVDGenericDebugDrawDataComponentVisualizer::DrawSpheres(FPrimitiveDrawInterface* PDI, const FSceneView* View, const FChaosVDGenericDebugDrawDataVisualizationSettings& InVisualizationContext)
+{
+	if (InVisualizationContext.IsVisualizationFlagEnabled(EChaosVDGenericDebugDrawVisualizationFlags::DrawSpheres))
+	{
+		TConstArrayView<TSharedPtr<FChaosVDDebugDrawSphereDataWrapper>> DebugDrawSpheres = InVisualizationContext.DataComponent->GetDebugDrawSpheresDataView();
+		for (const TSharedPtr<FChaosVDDebugDrawSphereDataWrapper>& DebugDrawSphere : DebugDrawSpheres)
+		{
+			if (DebugDrawSphere)
+			{
+				if (!View->ViewFrustum.IntersectSphere(DebugDrawSphere->Origin,  DebugDrawSphere->Radius))
+				{
+					continue;
+				}
+
+				const FText& DebugText = InVisualizationContext.bShowDebugText ? FText::FromName(DebugDrawSphere->Tag) : FText::GetEmpty();
+				constexpr int32 Segments = 12;
+
+				FChaosVDDebugDrawUtils::DrawSphere(PDI, DebugDrawSphere->Origin, DebugDrawSphere->Radius, Segments, DebugDrawSphere->Color, DebugText, InVisualizationContext.DepthPriority, InVisualizationContext.Thickness);			
+			}
+		}
+	}
+}
+
+void FChaosVDGenericDebugDrawDataComponentVisualizer::DrawImplicitObjects(FPrimitiveDrawInterface* PDI, const FSceneView* View, const FChaosVDGenericDebugDrawDataVisualizationSettings& InVisualizationContext)
+{
+	using namespace Chaos;
+
+	TSharedPtr<FChaosVDScene> CVDScene = InVisualizationContext.CVDScene.Pin();
+
+	if (!CVDScene)
+	{
+		return;
+	}
+	
+	if (InVisualizationContext.IsVisualizationFlagEnabled(EChaosVDGenericDebugDrawVisualizationFlags::DrawImplicitObjects))
+	{
+		TConstArrayView<TSharedPtr<FChaosVDDebugDrawImplicitObjectDataWrapper>> DebugDrawImplicitObjects = InVisualizationContext.DataComponent->GetDebugDrawImplicitObjectsDataView();
+		for (const TSharedPtr<FChaosVDDebugDrawImplicitObjectDataWrapper>& DebugDrawImplicitObjectData : DebugDrawImplicitObjects)
+		{
+			if (DebugDrawImplicitObjectData)
+			{
+				if (FConstImplicitObjectPtr LoadedGeometry = CVDScene->GetUpdatedGeometry(DebugDrawImplicitObjectData->ImplicitObjectHash))
+				{
+					FAABB3 Bounds = LoadedGeometry->CalculateTransformedBounds(DebugDrawImplicitObjectData->ParentTransform);
+					if (!View->ViewFrustum.IntersectBox(Bounds.Center(),  Bounds.Extents()))
+					{
+						continue;
+					}
+
+					if (InVisualizationContext.bShowDebugText)
+					{
+						FChaosVDDebugDrawUtils::DrawText(FText::FromName(DebugDrawImplicitObjectData->Tag), DebugDrawImplicitObjectData->ParentTransform.GetLocation(), DebugDrawImplicitObjectData->Color, EChaosVDDebugDrawTextLocationMode::World);
+					}				
+				
+					FChaosVDDebugDrawUtils::DrawImplicitObject(PDI, CVDScene->GetGeometryGenerator().Pin(), LoadedGeometry, DebugDrawImplicitObjectData->ParentTransform, DebugDrawImplicitObjectData->Color, FText::GetEmpty(), InVisualizationContext.DepthPriority, InVisualizationContext.Thickness);
+					
+				}
+			}
+		}
+	}
+}
+
+#undef LOCTEXT_NAMESPACE
diff --git a/Engine/Plugins/ChaosVD/Source/ChaosVD/Private/Visualizers/ChaosVDGenericDebugDrawDataComponentVisualizer.h b/Engine/Plugins/ChaosVD/Source/ChaosVD/Private/Visualizers/ChaosVDGenericDebugDrawDataComponentVisualizer.h
new file mode 100644
--- /dev/null
+++ b/Engine/Plugins/ChaosVD/Source/ChaosVD/Private/Visualizers/ChaosVDGenericDebugDrawDataComponentVisualizer.h
@@ -0,0 +1,48 @@
+// Copyright Epic Games, Inc. All Rights Reserved.
+
+#pragma once
+#include "ChaosVDComponentVisualizerBase.h"
+#include "Settings/ChaosVDGenericDebugDrawSettings.h"
+
+struct FChaosVDDebugDrawImplicitObjectDataWrapper;
+class UChaosVDGenericDebugDrawDataComponent;
+class UChaosVDGenericDebugDrawSettings;
+
+/** Visualization context structure specific for acceleration structure visualizations */
+struct FChaosVDGenericDebugDrawDataVisualizationSettings : public FChaosVDVisualizationContext
+{
+	ESceneDepthPriorityGroup DepthPriority = SDPG_Foreground;
+
+	float Thickness = 0.0f;
+
+	bool bShowDebugText = false;
+
+	const UChaosVDGenericDebugDrawDataComponent* DataComponent = nullptr;
+
+	bool IsVisualizationFlagEnabled(EChaosVDGenericDebugDrawVisualizationFlags Flag) const
+	{
+		const EChaosVDGenericDebugDrawVisualizationFlags FlagsAsGenericDebugDrawVisualizationVisFlags = static_cast<EChaosVDGenericDebugDrawVisualizationFlags>(VisualizationFlags);
+		return EnumHasAnyFlags(FlagsAsGenericDebugDrawVisualizationVisFlags, Flag);
+	}	
+};
+
+class FChaosVDGenericDebugDrawDataComponentVisualizer final : public FChaosVDComponentVisualizerBase
+{
+public:
+	
+	FChaosVDGenericDebugDrawDataComponentVisualizer();
+
+	virtual void RegisterVisualizerMenus() override;
+
+	virtual bool CanHandleClick(const HChaosVDComponentVisProxy& VisProxy) override { return false;}
+
+	virtual void DrawVisualization(const UActorComponent* Component, const FSceneView* View, FPrimitiveDrawInterface* PDI) override;
+
+	void DrawBoxes(FPrimitiveDrawInterface* PDI, const FSceneView* View, const FChaosVDGenericDebugDrawDataVisualizationSettings& InVisualizationContext);
+    void DrawLines(FPrimitiveDrawInterface* PDI, const FSceneView* View, const FChaosVDGenericDebugDrawDataVisualizationSettings& InVisualizationContext);
+    void DrawSpheres(FPrimitiveDrawInterface* PDI, const FSceneView* View, const FChaosVDGenericDebugDrawDataVisualizationSettings& InVisualizationContext);
+    void DrawImplicitObjects(FPrimitiveDrawInterface* PDI, const FSceneView* View, const FChaosVDGenericDebugDrawDataVisualizationSettings& InVisualizationContext);
+    void DrawImplicitObject(FPrimitiveDrawInterface* PDI, const FSceneView* View, const TSharedRef<FChaosVDScene>& InCVDSceneRef, const Chaos::FImplicitObject* InImplicitObject, const FTransform& BaseTransform, FColor Color, const FChaosVDGenericDebugDrawDataVisualizationSettings& InVisualizationContext);
+
+	
+};
diff --git a/Engine/Plugins/ChaosVD/Source/ChaosVD/Private/Visualizers/ChaosVDSceneQueryDataComponentVisualizer.cpp b/Engine/Plugins/ChaosVD/Source/ChaosVD/Private/Visualizers/ChaosVDSceneQueryDataComponentVisualizer.cpp
--- a/Engine/Plugins/ChaosVD/Source/ChaosVD/Private/Visualizers/ChaosVDSceneQueryDataComponentVisualizer.cpp
+++ b/Engine/Plugins/ChaosVD/Source/ChaosVD/Private/Visualizers/ChaosVDSceneQueryDataComponentVisualizer.cpp
@@ -2,7 +2,7 @@
 
 #include "Visualizers/ChaosVDSceneQueryDataComponentVisualizer.h"
 
-#include "Actors/ChaosVDSceneQueryDataContainer.h"
+#include "Actors/ChaosVDGameFrameInfoActor.h"
 #include "Components/ChaosVDSceneQueryDataComponent.h"
 #include "ChaosVDGeometryBuilder.h"
 #include "ChaosVDScene.h"
@@ -58,13 +58,18 @@ void FChaosVDSceneQueryDataComponentVisualizer::DrawVisualization(const UActorCo
 		return;
 	}
 
-	const AChaosVDSceneQueryDataContainer* SceneQueryDataContainer = Cast<AChaosVDSceneQueryDataContainer>(SceneQueryDataComponent->GetOwner());
-	if (!SceneQueryDataContainer)
+	const AChaosVDSolverInfoActor* SolverInfoActor = Cast<AChaosVDSolverInfoActor>(SceneQueryDataComponent->GetOwner());
+	if (!SolverInfoActor)
+	{
+		return;
+	}
+
+	if (!SolverInfoActor->IsVisible())
 	{
 		return;
 	}
 
-	const TSharedPtr<FChaosVDScene> CVDScene = SceneQueryDataContainer->GetScene().Pin();
+	const TSharedPtr<FChaosVDScene> CVDScene = SolverInfoActor->GetScene().Pin();
 	if (!CVDScene)
 	{
 		return;
@@ -89,7 +94,7 @@ void FChaosVDSceneQueryDataComponentVisualizer::DrawVisualization(const UActorCo
 	}
 
 	FChaosVDSceneQueryVisualizationDataContext VisualizationContext;
-	VisualizationContext.CVDScene = SceneQueryDataContainer->GetScene();
+	VisualizationContext.CVDScene = SolverInfoActor->GetScene();
 	VisualizationContext.SpaceTransform = FTransform::Identity;
 	VisualizationContext.GeometryGenerator = GeometryGenerator;
 	VisualizationContext.SolverDataSelectionObject = SolverDataSelectionObject;
@@ -315,17 +320,6 @@ void FChaosVDSceneQueryDataComponentVisualizer::DrawSceneQuery(const UActorCompo
 		return;
 	}
 
-	AChaosVDSolverInfoActor* SolverInfoActor = CVDScene->GetSolverInfoActor(Query->WorldSolverID);
-	if (!SolverInfoActor)
-	{
-		return;
-	}
-	
-	if (!SolverInfoActor->IsVisible())
-	{
-		return;
-	}
-
 	const bool bHideEmptyQueries = EnumHasAnyFlags(EChaosVDSceneQueryVisualizationFlags::HideEmptyQueries, static_cast<EChaosVDSceneQueryVisualizationFlags>(VisualizationContext.VisualizationFlags));
 	if (bHideEmptyQueries && Query->SQVisitData.IsEmpty())
 	{
diff --git a/Engine/Plugins/ChaosVD/Source/ChaosVD/Private/Visualizers/ChaosVDSolverCollisionDataComponentVisualizer.cpp b/Engine/Plugins/ChaosVD/Source/ChaosVD/Private/Visualizers/ChaosVDSolverCollisionDataComponentVisualizer.cpp
--- a/Engine/Plugins/ChaosVD/Source/ChaosVD/Private/Visualizers/ChaosVDSolverCollisionDataComponentVisualizer.cpp
+++ b/Engine/Plugins/ChaosVD/Source/ChaosVD/Private/Visualizers/ChaosVDSolverCollisionDataComponentVisualizer.cpp
@@ -39,7 +39,7 @@ void FChaosVDSolverCollisionDataComponentVisualizer::RegisterVisualizerMenus()
 	FText MenuSectionLabel = LOCTEXT("CollisionDataVisualizationShowMenuLabel", "Collision Data Visualization");
 	FText FlagsMenuLabel = LOCTEXT("CollisionVisualizationFlagsMenuLabel", "Collision Data Flags");
 	FText FlagsMenuTooltip = LOCTEXT("CollisionVisualizationFlagsMenuToolTip", "Set of flags to enable/disable visibility of specific types of collisionData");
-	FSlateIcon FlagsMenuIcon = FSlateIcon(FSlateIcon(FAppStyle::Get().GetStyleSetName(), TEXT("EditorViewport.CollisionVisibility")));
+	FSlateIcon FlagsMenuIcon = FSlateIcon(FSlateIcon(FAppStyle::Get().GetStyleSetName(), TEXT("ShowFlagsMenu.Collision")));
 
 	FText SettingsMenuLabel = LOCTEXT("CollisionVisualizationSettingsMenuLabel", "Collision Visualization Settings");
 	FText SettingsMenuTooltip = LOCTEXT("CollisionVisualizationSettingsMenuToolTip", "Options to change how the recorded collision Data is debug drawn");
diff --git a/Engine/Plugins/ChaosVD/Source/ChaosVD/Private/Widgets/SChaosVDMainTab.cpp b/Engine/Plugins/ChaosVD/Source/ChaosVD/Private/Widgets/SChaosVDMainTab.cpp
--- a/Engine/Plugins/ChaosVD/Source/ChaosVD/Private/Widgets/SChaosVDMainTab.cpp
+++ b/Engine/Plugins/ChaosVD/Source/ChaosVD/Private/Widgets/SChaosVDMainTab.cpp
@@ -29,6 +29,7 @@
 #include "StatusBarSubsystem.h"
 #include "ToolMenu.h"
 #include "ToolMenus.h"
+#include "Components/ChaosVDGenericDebugDrawDataComponent.h"
 #include "Components/ChaosVDGTAccelerationStructuresDataComponent.h"
 #include "Components/ChaosVDParticleDataComponent.h"
 #include "Components/ChaosVDSolverCharacterGroundConstraintDataComponent.h"
@@ -38,6 +39,7 @@
 #include "Styling/ToolBarStyle.h"
 #include "Trace/ChaosVDTraceManager.h"
 #include "Visualizers/ChaosVDCharacterGroundConstraintsDataComponentVisualizer.h"
+#include "Visualizers/ChaosVDGenericDebugDrawDataComponentVisualizer.h"
 #include "Visualizers/ChaosVDGTAccelerationStructureDataComponentVisualizer.h"
 #include "Visualizers/ChaosVDJointConstraintsDataComponentVisualizer.h"
 #include "Visualizers/ChaosVDParticleDataComponentVisualizer.h"
@@ -67,6 +69,7 @@ void SChaosVDMainTab::Construct(const FArguments& InArgs, TSharedPtr<FChaosVDEng
 	RegisterComponentVisualizer(UChaosVDSolverJointConstraintDataComponent::StaticClass()->GetFName(), MakeShared<FChaosVDJointConstraintsDataComponentVisualizer>());
 	RegisterComponentVisualizer(UChaosVDSolverCharacterGroundConstraintDataComponent::StaticClass()->GetFName(), MakeShared<FChaosVDCharacterGroundConstraintDataComponentVisualizer>());
 	RegisterComponentVisualizer(UChaosVDGTAccelerationStructuresDataComponent::StaticClass()->GetFName(), MakeShared<FChaosVDGTAccelerationStructureDataComponentVisualizer>());
+	RegisterComponentVisualizer(UChaosVDGenericDebugDrawDataComponent::StaticClass()->GetFName(), MakeShared<FChaosVDGenericDebugDrawDataComponentVisualizer>());
 
 	TabManager = FGlobalTabmanager::Get()->NewTabManager(InArgs._OwnerTab.ToSharedRef()).ToSharedPtr();
 
diff --git a/Engine/Plugins/ChaosVD/Source/ChaosVD/Private/Widgets/SChaosVDSceneQueryDataInspector.cpp b/Engine/Plugins/ChaosVD/Source/ChaosVD/Private/Widgets/SChaosVDSceneQueryDataInspector.cpp
--- a/Engine/Plugins/ChaosVD/Source/ChaosVD/Private/Widgets/SChaosVDSceneQueryDataInspector.cpp
+++ b/Engine/Plugins/ChaosVD/Source/ChaosVD/Private/Widgets/SChaosVDSceneQueryDataInspector.cpp
@@ -318,10 +318,10 @@ void SChaosVDSceneQueryDataInspector::SetQueryDataToInspect(const TSharedPtr<FCh
 		{
 			TArray<TSharedPtr<FName>> NewSubQueryNameList;
 			NewSubQueryNameList.Reserve(QueryDataToInspect->SubQueriesIDs.Num());
-			Algo::Transform(QueryDataToInspect->SubQueriesIDs, NewSubQueryNameList, [this](int32 QueryID)
+			Algo::Transform(QueryDataToInspect->SubQueriesIDs, NewSubQueryNameList, [this, &QueryDataToInspect](int32 QueryID)
 			{
 				TSharedPtr<FName> NewName = MakeShared<FName>(FString::Format(TEXT("Query ID {0}"), {QueryID}));
-				CurrentSubQueriesByName.Add(NewName, QueryID);
+				CurrentSubQueriesByName.Add(NewName, {QueryID, QueryDataToInspect->WorldSolverID});
 				return NewName;
 			});
 
@@ -370,7 +370,8 @@ void SChaosVDSceneQueryDataInspector::HandleQueryStepSelectionUpdated(int32 NewS
 		return;
 	}
 
-	UChaosVDSceneQueryDataComponent* SQDataComponent = ScenePtr->GetSceneQueryDataContainerComponent();
+	AChaosVDSolverInfoActor* SolverInfoActor = ScenePtr->GetSolverInfoActor(QueryDataBeingInspected->WorldSolverID);
+	UChaosVDSceneQueryDataComponent* SQDataComponent = SolverInfoActor ? SolverInfoActor->GetSceneQueryDataComponent() : nullptr;
 	if (!SQDataComponent)
 	{
 		ClearInspector();
@@ -467,7 +468,7 @@ FReply SChaosVDSceneQueryDataInspector::SelectParticleForCurrentQueryData() cons
 	return FReply::Handled();
 }
 
-FReply SChaosVDSceneQueryDataInspector::SelectQueryToInspectByID(int32 QueryID)
+FReply SChaosVDSceneQueryDataInspector::SelectQueryToInspectByID(int32 QueryID, int32 SolverID)
 {
 	const TSharedPtr<FChaosVDScene> ScenePtr = SceneWeakPtr.Pin();
 	if (!ScenePtr)
@@ -481,7 +482,8 @@ FReply SChaosVDSceneQueryDataInspector::SelectQueryToInspectByID(int32 QueryID)
 		return FReply::Handled();
 	}
 
-	if (UChaosVDSceneQueryDataComponent* SQDataComponent = ScenePtr->GetSceneQueryDataContainerComponent())
+	AChaosVDSolverInfoActor* SolverInfoActor = ScenePtr->GetSolverInfoActor(SolverID);
+	if (UChaosVDSceneQueryDataComponent* SQDataComponent = SolverInfoActor ? SolverInfoActor->GetSceneQueryDataComponent() : nullptr)
 	{						
 		SelectionObject->SelectData(SelectionObject->MakeSelectionHandle(SQDataComponent->GetQueryByID(QueryID)));
 	}
@@ -496,10 +498,10 @@ FReply SChaosVDSceneQueryDataInspector::SelectQueryToInspectByID(int32 QueryID)
 FReply SChaosVDSceneQueryDataInspector::SelectParentQuery()
 {
 	const TSharedPtr<FChaosVDScene> ScenePtr = SceneWeakPtr.Pin();
-	const UChaosVDSceneQueryDataComponent* SQDataComponent = ScenePtr ? ScenePtr->GetSceneQueryDataContainerComponent() : nullptr;
-	if (const TSharedPtr<FChaosVDQueryDataWrapper> SelectedQuery = SQDataComponent ? CurrentSceneQueryBeingInspectedHandle->GetDataAsShared<FChaosVDQueryDataWrapper>() : nullptr)
-	{
-		SelectQueryToInspectByID(SelectedQuery->ParentQueryID);
+
+	if (const TSharedPtr<FChaosVDQueryDataWrapper> SelectedQuery = CurrentSceneQueryBeingInspectedHandle->GetDataAsShared<FChaosVDQueryDataWrapper>())
+	{		
+		SelectQueryToInspectByID(SelectedQuery->ParentQueryID, SelectedQuery->WorldSolverID);
 	}
 	else
 	{
@@ -543,9 +545,9 @@ void SChaosVDSceneQueryDataInspector::HandleSubQueryNameSelected(TSharedPtr<FNam
 		return;
 	}
 
-	if (const int32* SelectedQueryID = CurrentSubQueriesByName.Find(Name))
+	if (const FChaosVDSQSubQueryID* SelectedQueryID = CurrentSubQueriesByName.Find(Name))
 	{
-		SelectQueryToInspectByID(*SelectedQueryID);
+		SelectQueryToInspectByID(SelectedQueryID->QueryID, SelectedQueryID->SolverID);
 		return;
 	}
 
diff --git a/Engine/Plugins/ChaosVD/Source/ChaosVD/Private/Widgets/SChaosVDSceneQueryDataInspector.h b/Engine/Plugins/ChaosVD/Source/ChaosVD/Private/Widgets/SChaosVDSceneQueryDataInspector.h
--- a/Engine/Plugins/ChaosVD/Source/ChaosVD/Private/Widgets/SChaosVDSceneQueryDataInspector.h
+++ b/Engine/Plugins/ChaosVD/Source/ChaosVD/Private/Widgets/SChaosVDSceneQueryDataInspector.h
@@ -16,6 +16,12 @@ struct FChaosVDQueryDataWrapper;
 class IStructureDetailsView;
 class FChaosVDScene;
 
+struct FChaosVDSQSubQueryID
+{
+	int32 QueryID = INDEX_NONE;
+	int32 SolverID = INDEX_NONE;
+};
+
 /**
  * Widget for the Chaos Visual Debugger Scene Queries data inspector
  */
@@ -53,7 +59,7 @@ protected:
 	FText GetSQVisitsStepsText() const;
 	
 	FReply SelectParticleForCurrentQueryData() const;
-	FReply SelectQueryToInspectByID(int32 QueryID);
+	FReply SelectQueryToInspectByID(int32 QueryID, int32 SolverID);
 	FReply SelectParentQuery();
 
 	static TSharedPtr<IStructureDetailsView> CreateDataDetailsView();
@@ -88,7 +94,7 @@ protected:
 
 	TWeakPtr<FEditorModeTools> EditorModeToolsWeakPtr;
 
-	TMap<TSharedPtr<FName>, int32> CurrentSubQueriesByName;
+	TMap<TSharedPtr<FName>, FChaosVDSQSubQueryID> CurrentSubQueriesByName;
 	
 	TSharedRef<FChaosVDSolverDataSelectionHandle> CurrentSceneQueryBeingInspectedHandle;
 
diff --git a/Engine/Plugins/ChaosVD/Source/ChaosVDBlueprint/ChaosVDBlueprint.build.cs b/Engine/Plugins/ChaosVD/Source/ChaosVDBlueprint/ChaosVDBlueprint.build.cs
new file mode 100644
--- /dev/null
+++ b/Engine/Plugins/ChaosVD/Source/ChaosVDBlueprint/ChaosVDBlueprint.build.cs
@@ -0,0 +1,22 @@
+// Copyright Epic Games, Inc. All Rights Reserved.
+
+using System.Linq;
+using System.Collections.Generic;
+
+namespace UnrealBuildTool.Rules
+{
+	public class ChaosVDBlueprint : ModuleRules
+	{
+		public ChaosVDBlueprint(ReadOnlyTargetRules Target) : base(Target)
+		{
+			PublicDependencyModuleNames.AddRange(
+				new string[] {
+					"Core",
+					"CoreUObject",
+					"Chaos",
+					"Engine"
+				}
+			);
+		}
+	}
+}
\ No newline at end of file
diff --git a/Engine/Plugins/ChaosVD/Source/ChaosVDBlueprint/Private/ChaosVDRuntimeBlueprintLibrary.cpp b/Engine/Plugins/ChaosVD/Source/ChaosVDBlueprint/Private/ChaosVDRuntimeBlueprintLibrary.cpp
new file mode 100644
--- /dev/null
+++ b/Engine/Plugins/ChaosVD/Source/ChaosVDBlueprint/Private/ChaosVDRuntimeBlueprintLibrary.cpp
@@ -0,0 +1,30 @@
+// Copyright Epic Games, Inc. All Rights Reserved.
+
+#include "ChaosVDRuntimeBlueprintLibrary.h"
+#include "Engine/World.h"
+
+#include "Physics/Experimental/PhysScene_Chaos.h"
+#include "Chaos/Framework/PhysicsSolverBase.h"
+#include "PBDRigidsSolver.h"
+
+IMPLEMENT_MODULE(FDefaultModuleImpl, ChaosVDBlueprint);
+
+void UChaosVDRuntimeBlueprintLibrary::RecordDebugDrawBox(const UObject* WorldContext, const FBox& InBox, FName Tag, FLinearColor Color)
+{
+	CVD_TRACE_DEBUG_DRAW_BOX(InBox, Tag, Color.ToFColorSRGB(), WorldContext ? CVD_TRACE_GET_SOLVER_ID_FROM_WORLD(WorldContext->GetWorld()) : INDEX_NONE);
+}
+
+void UChaosVDRuntimeBlueprintLibrary::RecordDebugDrawLine(const UObject* WorldContext, const FVector& InStartLocation, const FVector& InEndLocation, FName Tag, FLinearColor Color)
+{
+	CVD_TRACE_DEBUG_DRAW_LINE(InStartLocation, InEndLocation, Tag, Color.ToFColorSRGB(), WorldContext ? CVD_TRACE_GET_SOLVER_ID_FROM_WORLD(WorldContext->GetWorld()) : INDEX_NONE);
+}
+
+void UChaosVDRuntimeBlueprintLibrary::RecordDebugDrawVector(const UObject* WorldContext, const FVector& InStartLocation, const FVector& InVector, FName Tag, FLinearColor Color)
+{
+	CVD_TRACE_DEBUG_DRAW_VECTOR(InStartLocation, InVector, Tag, Color.ToFColorSRGB(), WorldContext ? CVD_TRACE_GET_SOLVER_ID_FROM_WORLD(WorldContext->GetWorld()) : INDEX_NONE);
+}
+
+void UChaosVDRuntimeBlueprintLibrary::RecordDebugDrawSphere(const UObject* WorldContext, const FVector& InCenter, float Radius, FName Tag, FLinearColor Color)
+{
+	CVD_TRACE_DEBUG_DRAW_SPHERE(InCenter, Radius, Tag, Color.ToFColorSRGB(), WorldContext ? CVD_TRACE_GET_SOLVER_ID_FROM_WORLD(WorldContext->GetWorld()) : INDEX_NONE);
+}
diff --git a/Engine/Plugins/ChaosVD/Source/ChaosVDBlueprint/Public/ChaosVDRuntimeBlueprintLibrary.h b/Engine/Plugins/ChaosVD/Source/ChaosVDBlueprint/Public/ChaosVDRuntimeBlueprintLibrary.h
new file mode 100644
--- /dev/null
+++ b/Engine/Plugins/ChaosVD/Source/ChaosVDBlueprint/Public/ChaosVDRuntimeBlueprintLibrary.h
@@ -0,0 +1,27 @@
+// Copyright Epic Games, Inc. All Rights Reserved.
+
+#pragma once
+#include "Kismet/BlueprintFunctionLibrary.h"
+
+#include "ChaosVDRuntimeBlueprintLibrary.generated.h"
+
+/**
+ * Library function to record debug draw shapes that will be played back when a CVD recording is loaded
+ */
+UCLASS()
+class CHAOSVDBLUEPRINT_API UChaosVDRuntimeBlueprintLibrary : public UBlueprintFunctionLibrary
+{
+	GENERATED_BODY()
+
+	UFUNCTION(BlueprintCallable, Category = "Chaos Visual Debugger", DisplayName="CVD Record Debug Draw Box", meta = (WorldContext = "WorldContext"))
+	static void RecordDebugDrawBox(const UObject* WorldContext, const FBox& InBox, FName Tag = NAME_None, FLinearColor Color = FLinearColor::Blue);
+
+	UFUNCTION(BlueprintCallable, Category = "Chaos Visual Debugger",  DisplayName="CVD Record Debug Draw Line", meta = (WorldContext = "WorldContext"))
+	static void RecordDebugDrawLine(const UObject* WorldContext, const FVector& InStartLocation, const FVector& InEndLocation, FName Tag = NAME_None, FLinearColor Color = FLinearColor::Blue);
+
+	UFUNCTION(BlueprintCallable, Category = "Chaos Visual Debugger",  DisplayName="CVD Record Debug Draw Vector", meta = (WorldContext = "WorldContext"))
+	static void RecordDebugDrawVector(const UObject* WorldContext, const FVector& InStartLocation, const FVector& InVector, FName Tag = NAME_None, FLinearColor Color = FLinearColor::Blue);
+
+	UFUNCTION(BlueprintCallable, Category = "Chaos Visual Debugger", DisplayName="CVD Record Debug Draw Sphere", meta = (WorldContext = "WorldContext"))
+	static void RecordDebugDrawSphere(const UObject* WorldContext, const FVector& InCenter, float Radius, FName Tag = NAME_None, FLinearColor Color = FLinearColor::Blue);
+};
diff --git a/Engine/Source/Runtime/Experimental/Chaos/Private/ChaosVisualDebugger/ChaosVDDataWrapperUtils.cpp b/Engine/Source/Runtime/Experimental/Chaos/Private/ChaosVisualDebugger/ChaosVDDataWrapperUtils.cpp
--- a/Engine/Source/Runtime/Experimental/Chaos/Private/ChaosVisualDebugger/ChaosVDDataWrapperUtils.cpp
+++ b/Engine/Source/Runtime/Experimental/Chaos/Private/ChaosVisualDebugger/ChaosVDDataWrapperUtils.cpp
@@ -12,13 +12,41 @@
 
 #include "Chaos/Character/CharacterGroundConstraintContainer.h"
 #include "Chaos/Collision/ParticlePairMidPhase.h"
-#include "Chaos/Collision/CollisionConstraintAllocator.h"
 #include "Chaos/ParticleHandle.h"
 #include "ChaosVisualDebugger/ChaosVDSerializedNameTable.h"
 #include "DataWrappers/ChaosVDCollisionDataWrappers.h"
 #include "DataWrappers/ChaosVDParticleDataWrapper.h"
 #include "Math/UnitConversion.h"
 
+namespace Chaos::VisualDebugger::Utils
+{
+	FGeometryParticle* GetPayloadForExternalThread(const Chaos::FAccelerationStructureHandle& Payload)
+	{
+		return Payload.GetExternalGeometryParticle_ExternalThread();	
+	}
+
+	/** Calculates and returns the current Game Thread bounds for the provided particle */
+	FBox GetGeometricGTParticleBounds(const FGeometryParticle* GeometryParticle)
+	{
+		using namespace Chaos;
+		const FShapesArray& Shapes = GeometryParticle->ShapesArray();
+
+		FBox Bounds(ForceInitToZero);
+
+		for (const TUniquePtr<FPerShapeData>& Shape : Shapes)
+		{
+			Bounds += FBox(Shape->GetWorldSpaceShapeBounds().Min(), Shape->GetWorldSpaceShapeBounds().Max());
+		}
+
+		return Bounds;
+	}
+
+	FBox GetPayloadBounds(const FAccelerationStructureHandle& Payload)
+	{
+		return GetGeometricGTParticleBounds(GetPayloadForExternalThread(Payload));
+	}
+}
+
 void FChaosVDDataWrapperUtils::CopyManifoldPointsToDataWrapper(const Chaos::FManifoldPoint& InCopyFrom, FChaosVDManifoldPoint& OutCopyTo)
 {
 	OutCopyTo.bDisabled = InCopyFrom.Flags.bDisabled;
@@ -493,6 +521,7 @@ void FChaosVDDataWrapperUtils::AddTreeLeaves(const TConstArrayView<Chaos::TAABBT
 			FChaosVDAABBTreePayloadBoundsElement CVELeafElement;
 			CVELeafElement.ParticleIndex = GetUniqueIdx(Elem).Idx;
 			CVELeafElement.Bounds = ConvertToFBox(Elem.Bounds);
+			CVELeafElement.ActualBounds = VisualDebugger::Utils::GetPayloadBounds(Elem.Payload);
 
 			CVELeafElement.MarkAsValid();
 	
diff --git a/Engine/Source/Runtime/Experimental/Chaos/Private/ChaosVisualDebugger/ChaosVisualDebuggerTrace.cpp b/Engine/Source/Runtime/Experimental/Chaos/Private/ChaosVisualDebugger/ChaosVisualDebuggerTrace.cpp
--- a/Engine/Source/Runtime/Experimental/Chaos/Private/ChaosVisualDebugger/ChaosVisualDebuggerTrace.cpp
+++ b/Engine/Source/Runtime/Experimental/Chaos/Private/ChaosVisualDebugger/ChaosVisualDebuggerTrace.cpp
@@ -18,6 +18,7 @@
 #include "Compression/OodleDataCompressionUtil.h"
 #include "DataWrappers/ChaosVDCharacterGroundConstraintDataWrappers.h"
 #include "DataWrappers/ChaosVDCollisionDataWrappers.h"
+#include "DataWrappers/ChaosVDDebugShapeDataWrapper.h"
 #include "DataWrappers/ChaosVDImplicitObjectDataWrapper.h"
 #include "DataWrappers/ChaosVDJointDataWrappers.h"
 #include "DataWrappers/ChaosVDParticleDataWrapper.h"
@@ -848,6 +849,120 @@ void FChaosVisualDebuggerTrace::TraceNetworkTickOffset(int32 TickOffset, int32 S
 		<< ChaosVDNetworkTickOffset.SolverID(SolverID);
 }
 
+void FChaosVisualDebuggerTrace::TraceDebugDrawBox(const FBox& InBox, FName Tag, FColor Color, int32 SolverID)
+{
+	if (!IsTracing())
+	{
+		return;
+	}
+
+	FChaosVDDebugDrawBoxDataWrapper DataWrapper;
+	DataWrapper.SolverID = SolverID;
+	DataWrapper.Tag = Tag;
+	DataWrapper.Color = Color;
+	DataWrapper.Box = InBox;
+
+	DataWrapper.MarkAsValid();
+
+	FChaosVDScopedTLSBufferAccessor TLSDataBuffer;
+	Chaos::VisualDebugger::WriteDataToBuffer(TLSDataBuffer.BufferRef, DataWrapper);
+
+	TraceBinaryData(TLSDataBuffer.BufferRef, FChaosVDDebugDrawBoxDataWrapper::WrapperTypeName);
+}
+
+void FChaosVisualDebuggerTrace::TraceDebugDrawLine(const FVector& InStartLocation, const FVector& InEndLocation, FName Tag, FColor Color, int32 SolverID)
+{
+	if (!IsTracing())
+	{
+		return;
+	}
+
+	FChaosVDDebugDrawLineDataWrapper DataWrapper;
+	DataWrapper.SolverID = SolverID;
+	DataWrapper.Tag = Tag;
+	DataWrapper.Color = Color;
+	DataWrapper.StartLocation = InStartLocation;
+	DataWrapper.EndLocation = InEndLocation;
+
+	DataWrapper.MarkAsValid();
+
+	FChaosVDScopedTLSBufferAccessor TLSDataBuffer;
+	Chaos::VisualDebugger::WriteDataToBuffer(TLSDataBuffer.BufferRef, DataWrapper);
+
+	TraceBinaryData(TLSDataBuffer.BufferRef, FChaosVDDebugDrawLineDataWrapper::WrapperTypeName);
+}
+
+void FChaosVisualDebuggerTrace::TraceDebugDrawVector(const FVector& InStartLocation, const FVector& InVector, FName Tag, FColor Color, int32 SolverID)
+{
+	if (!IsTracing())
+	{
+		return;
+	}
+
+	FChaosVDDebugDrawLineDataWrapper DataWrapper;
+	DataWrapper.SolverID = SolverID;
+	DataWrapper.Tag = Tag;
+	DataWrapper.Color = Color;
+	DataWrapper.StartLocation = InStartLocation;
+	DataWrapper.EndLocation = InStartLocation + InVector;
+	DataWrapper.bIsArrow = true;
+
+	DataWrapper.MarkAsValid();
+
+	FChaosVDScopedTLSBufferAccessor TLSDataBuffer;
+	Chaos::VisualDebugger::WriteDataToBuffer(TLSDataBuffer.BufferRef, DataWrapper);
+
+	TraceBinaryData(TLSDataBuffer.BufferRef, FChaosVDDebugDrawLineDataWrapper::WrapperTypeName);
+}
+
+void FChaosVisualDebuggerTrace::TraceDebugDrawSphere(const FVector& Center, float Radius, FName Tag, FColor Color, int32 SolverID)
+{
+	if (!IsTracing())
+	{
+		return;
+	}
+
+	FChaosVDDebugDrawSphereDataWrapper DataWrapper;
+	DataWrapper.SolverID = SolverID;
+	DataWrapper.Tag = Tag;
+	DataWrapper.Color = Color;
+	DataWrapper.Origin = Center;
+	DataWrapper.Radius = Radius;
+
+	DataWrapper.MarkAsValid();
+
+	FChaosVDScopedTLSBufferAccessor TLSDataBuffer;
+	Chaos::VisualDebugger::WriteDataToBuffer(TLSDataBuffer.BufferRef, DataWrapper);
+
+	TraceBinaryData(TLSDataBuffer.BufferRef, FChaosVDDebugDrawSphereDataWrapper::WrapperTypeName);
+}
+
+void FChaosVisualDebuggerTrace::TraceDebugDrawImplicitObject(const Chaos::FImplicitObject* Implicit, const FTransform& InParentTransform, FName Tag, FColor Color, int32 SolverID)
+{
+	if (!IsTracing())
+	{
+		return;
+	}
+
+	FChaosVDDebugDrawImplicitObjectDataWrapper DataWrapper;
+	DataWrapper.SolverID = SolverID;
+	DataWrapper.Tag = Tag;
+	DataWrapper.Color = Color;
+	DataWrapper.ParentTransform = InParentTransform;
+
+	const uint32 GeometryHash = GeometryTracerObject.GetGeometryHashForImplicit(Implicit);
+	TraceImplicitObject({ GeometryHash, const_cast<Chaos::FImplicitObject*>(Implicit) });
+
+	DataWrapper.ImplicitObjectHash = GeometryHash;
+
+	DataWrapper.MarkAsValid();
+
+	FChaosVDScopedTLSBufferAccessor TLSDataBuffer;
+	Chaos::VisualDebugger::WriteDataToBuffer(TLSDataBuffer.BufferRef, DataWrapper);
+
+	TraceBinaryData(TLSDataBuffer.BufferRef, FChaosVDDebugDrawImplicitObjectDataWrapper::WrapperTypeName);
+}
+
 bool FChaosVisualDebuggerTrace::IsTracing()
 {
 	return bIsTracing;
diff --git a/Engine/Source/Runtime/Experimental/Chaos/Public/ChaosVisualDebugger/ChaosVDTraceMacros.h b/Engine/Source/Runtime/Experimental/Chaos/Public/ChaosVisualDebugger/ChaosVDTraceMacros.h
--- a/Engine/Source/Runtime/Experimental/Chaos/Public/ChaosVisualDebugger/ChaosVDTraceMacros.h
+++ b/Engine/Source/Runtime/Experimental/Chaos/Public/ChaosVisualDebugger/ChaosVDTraceMacros.h
@@ -177,17 +177,42 @@
 			}
 	#endif
 
-#ifndef CVD_TRACE_NETWORK_TICK_OFFSET
-	#define CVD_TRACE_NETWORK_TICK_OFFSET(TickOffset, SolverID) \
-	{ \
-		FChaosVisualDebuggerTrace::TraceNetworkTickOffset(TickOffset, SolverID); \
-	}
-#endif
-
-#ifndef CVD_TRACE_GET_SOLVER_ID_FROM_WORLD
-	#define CVD_TRACE_GET_SOLVER_ID_FROM_WORLD(World) \
-			FChaosVisualDebuggerTrace::GetSolverIDFromWorld(World)
-#endif
+	#ifndef CVD_TRACE_NETWORK_TICK_OFFSET
+		#define CVD_TRACE_NETWORK_TICK_OFFSET(TickOffset, SolverID) \
+		{ \
+			FChaosVisualDebuggerTrace::TraceNetworkTickOffset(TickOffset, SolverID); \
+		}
+	#endif
+
+	#ifndef CVD_TRACE_GET_SOLVER_ID_FROM_WORLD
+		#define CVD_TRACE_GET_SOLVER_ID_FROM_WORLD(World) \
+				FChaosVisualDebuggerTrace::GetSolverIDFromWorld(World)
+	#endif
+
+	#ifndef CVD_TRACE_DEBUG_DRAW_BOX
+      		#define CVD_TRACE_DEBUG_DRAW_BOX(Box, ...) \
+      		FChaosVisualDebuggerTrace::TraceDebugDrawBox(Box, ##__VA_ARGS__)
+	#endif
+
+	#ifndef CVD_TRACE_DEBUG_DRAW_LINE
+		#define CVD_TRACE_DEBUG_DRAW_LINE(StartLocation, EndLocation, ...) \
+		FChaosVisualDebuggerTrace::TraceDebugDrawLine(StartLocation, EndLocation, ##__VA_ARGS__)
+	#endif
+
+	#ifndef CVD_TRACE_DEBUG_DRAW_VECTOR
+		#define CVD_TRACE_DEBUG_DRAW_VECTOR(StartLocation, Vector, ...) \
+		FChaosVisualDebuggerTrace::TraceDebugDrawVector(StartLocation, Vector, ##__VA_ARGS__)
+	#endif
+
+	#ifndef CVD_TRACE_DEBUG_DRAW_SPHERE
+		#define CVD_TRACE_DEBUG_DRAW_SPHERE(Center, Radius, ...) \
+		FChaosVisualDebuggerTrace::TraceDebugDrawSphere(Center, Radius, ##__VA_ARGS__)
+	#endif
+
+	#ifndef CVD_TRACE_DEBUG_DRAW_IMPLICIT_OBJECT
+		#define CVD_TRACE_DEBUG_DRAW_IMPLICIT_OBJECT(ImplicitObject, ParentTransform, ...) \
+		FChaosVisualDebuggerTrace::TraceDebugDrawImplicitObject(ImplicitObject, ParentTransform,  ##__VA_ARGS__)
+	#endif
 
 #else // WITH_CHAOS_VISUAL_DEBUGGER
 
@@ -293,8 +318,28 @@
 		#define CVD_TRACE_NETWORK_TICK_OFFSET(TickOffset, SolverID)
 	#endif
 
-#ifndef CVD_TRACE_GET_SOLVER_ID_FROM_WORLD
-	#define CVD_TRACE_GET_SOLVER_ID_FROM_WORLD(World)
-#endif
+	#ifndef CVD_TRACE_GET_SOLVER_ID_FROM_WORLD
+		#define CVD_TRACE_GET_SOLVER_ID_FROM_WORLD(World)
+	#endif
+
+	#ifndef CVD_TRACE_DEBUG_DRAW_BOX
+			  #define CVD_TRACE_DEBUG_DRAW_BOX(Box, ...)
+	#endif
+
+	#ifndef CVD_TRACE_DEBUG_DRAW_LINE
+		#define CVD_TRACE_DEBUG_DRAW_LINE(StartLocation, EndLocation, ...)
+	#endif
+
+	#ifndef CVD_TRACE_DEBUG_DRAW_VECTOR
+		#define CVD_TRACE_DEBUG_DRAW_VECTOR(StartLocation, Vector, ...)
+	#endif
+
+	#ifndef CVD_TRACE_DEBUG_DRAW_SPHERE
+		#define CVD_TRACE_DEBUG_DRAW_SPHERE(...)
+	#endif
+
+	#ifndef CVD_TRACE_DEBUG_DRAW_IMPLICIT_OBJECT
+		#define CVD_TRACE_DEBUG_DRAW_IMPLICIT_OBJECT(ImplicitObject, ParentTransform, ...)
+	#endif
 
 #endif // WITH_CHAOS_VISUAL_DEBUGGER
\ No newline at end of file
diff --git a/Engine/Source/Runtime/Experimental/Chaos/Public/ChaosVisualDebugger/ChaosVisualDebuggerTrace.h b/Engine/Source/Runtime/Experimental/Chaos/Public/ChaosVisualDebugger/ChaosVisualDebuggerTrace.h
--- a/Engine/Source/Runtime/Experimental/Chaos/Public/ChaosVisualDebugger/ChaosVisualDebuggerTrace.h
+++ b/Engine/Source/Runtime/Experimental/Chaos/Public/ChaosVisualDebugger/ChaosVisualDebuggerTrace.h
@@ -260,16 +260,89 @@ public:
 	 */
 	static CHAOS_API void InvalidateGeometryFromCache(const Chaos::FImplicitObject* CachedGeometryToInvalidate);
 
+	/** Records a visit step of a scene query. This needs to me called within the scope of an inflight scene query
+	 * @param InputGeometry : Geometry used as input to perform the query we are recording
+	 * @param GeometryOrientation : Orientation of the input geometry
+	 * @param Start : Start location of the Query.
+	 * @param End : End location of the Query
+	 * @param TraceChannel : Trace channel used for the query we are recording.
+	 * @param Params : Collision query params used for the query we are recording.
+	 * @param ResponseParams :  Collision response params used for the query we are recording.
+	 * @param ObjectParams :  Collision object query params used for the query we are recording.
+	 * @param SolverID : ID of the solver this where query is being performed.
+	 * @param bIsRetry : Set to true if the query que are recording it is from a retry attempt.
+	 * */
 	static CHAOS_API void TraceSceneQueryStart(const Chaos::FImplicitObject* InputGeometry, const FQuat& GeometryOrientation, const FVector& Start, const FVector& End, ECollisionChannel TraceChannel, FChaosVDCollisionQueryParams&& Params, FChaosVDCollisionResponseParams&& ResponseParams, FChaosVDCollisionObjectQueryParams&& ObjectParams, EChaosVDSceneQueryType QueryType, EChaosVDSceneQueryMode QueryMode, int32 SolverID, bool bIsRetry);
+
+	/** Records a visit step of a scene query. This needs to me called within the scope of an inflight scene query
+	 * @param InQueryVisitData : Processed Scene Query step data.
+	 * */
 	static CHAOS_API void TraceSceneQueryVisit(FChaosVDQueryVisitStep&& InQueryVisitData);
 
+	/** Records all the suppported accelerations structures contained by the provided acceleration structure collection
+	 * @param InAccelerationCollection : Ptr to the collection we want to evaluate and trace.
+	 * */
 	static CHAOS_API void TraceSceneAccelerationStructures(const Chaos::ISpatialAccelerationCollection<Chaos::FAccelerationStructureHandle, Chaos::FReal, 3>* InAccelerationCollection);
 
+	/** Records the current tick offset for any given solver. This is used in the CVD editor to sync solver tracks based on their network tick
+	 * @param TickOffset : Offset relative to the server's solver current tick number.
+	 * @param SolverID : Solver ID of the solver that has this offset.
+	 * */
 	static CHAOS_API void TraceNetworkTickOffset(int32 TickOffset, int32 SolverID);
 
+	/** Records the provided box and the rest of the arguments, so then it can be visualized via debug draw during playback in the CVD editor
+	 * @param InBox : Shape to record.
+	 * @param Tag : FName that will be used as a tag for filtering & search, and debug draw as a text tag in CVD's viewport.
+	 * @param Color : Color to apply to this shape when it is debug drawn in CVD.
+	 * @param SolverID : ID of the solver this shape should be associated with. if no ID is provided, this shape will be added as part of the current game frame data bucket.
+	 */
+	static CHAOS_API void TraceDebugDrawBox(const FBox& InBox, FName Tag = NAME_None, FColor Color = FColor::Blue, int32 SolverID = INDEX_NONE);
+	
+	/** Records the provided Line and the rest of the arguments, so then it can be visualized via debug draw during playback in the CVD editor
+	 * @param InStartLocation : Start point of the line.
+	 * @param InEndLocation : End point of the line.
+	 * @param Tag : FName that will be used as a tag for filtering & search, and debug draw as a text tag in CVD's viewport.
+	 * @param Color : Color to apply to this shape when it is debug drawn in CVD.
+	 * @param SolverID : ID of the solver this shape should be associated with. if no ID is provided, this shape will be added as part of the current game frame data bucket.
+	 */
+	static CHAOS_API void TraceDebugDrawLine(const FVector& InStartLocation, const FVector& InEndLocation, FName Tag = NAME_None, FColor Color = FColor::Blue, int32 SolverID = INDEX_NONE);
+	
+	/** Records the provided Vector and the rest of the arguments, so then it can be visualized via debug draw during playback in the CVD editor
+	 * @param InStartLocation : Start point of the line.
+	 * @param InVector : Vector we want to record.
+	 * @param Tag : FName that will be used as a tag for filtering & search, and debug draw as a text tag in CVD's viewport.
+	 * @param Color : Color to apply to this shape when it is debug drawn in CVD.
+	 * @param SolverID : ID of the solver this shape should be associated with. if no ID is provided, this shape will be added as part of the current game frame data bucket.
+	 */
+	static CHAOS_API void TraceDebugDrawVector(const FVector& InStartLocation, const FVector& InVector, FName Tag = NAME_None, FColor Color = FColor::Blue, int32 SolverID = INDEX_NONE);
+
+	/** Records the provided Sphere and the rest of the arguments, so then it can be visualized via debug draw during playback in the CVD editor
+	 * @param Center : Origin point of the Sphere.
+	 * @param Radius : Radius of the Sphere.
+	 * @param Tag : FName that will be used as a tag for filtering & search, and debug draw as a text tag in CVD's viewport.
+	 * @param Color : Color to apply to this shape when it is debug drawn in CVD.
+	 * @param SolverID : ID of the solver this shape should be associated with. if no ID is provided, this shape will be added as part of the current game frame data bucket.
+	 */
+	static CHAOS_API void TraceDebugDrawSphere(const FVector& Center, float Radius, FName Tag = NAME_None, FColor Color = FColor::Blue, int32 SolverID = INDEX_NONE);
+
+	/** Records the provided Implicit Object and the rest of the arguments, so then it can be visualized via debug draw during playback in the CVD editor
+	 * @param Implicit : Ptr to the implicit object to record.
+	 * @param InParentTransform : Root transform of the object that owns this geometry.
+	 * @param Tag : FName that will be used as a tag for filtering & search, and debug draw as a text tag in CVD's viewport.
+	 * @param Color : Color to apply to this shape when it is debug drawn in CVD.
+	 * @param SolverID : ID of the solver this shape should be associated with. if no ID is provided, this shape will be added as part of the current game frame data bucket.
+	 */
+	static CHAOS_API void TraceDebugDrawImplicitObject(const Chaos::FImplicitObject* Implicit, const FTransform& InParentTransform, FName Tag = NAME_None, FColor Color = FColor::Blue, int32 SolverID = INDEX_NONE);
+
+	/** Returns the ID of the main solver of the provided world
+	 * @param World Ptr to the world we want to get the solverID from
+	 */
 	template<typename WorldType>
 	static int32 GetSolverIDFromWorld(WorldType* World);
 
+	/** Returns the CVD solver ID of the provided solver
+	 * @param Solver Reference to the solver instance we want to get the id from
+	 */
 	static int32 CHAOS_API GetSolverID(Chaos::FPhysicsSolverBase& Solver);
 
 	/** Returns true if the provided solver ID needs a Full Capture */
diff --git a/Engine/Source/Runtime/Experimental/ChaosVDData/Public/ChaosVDRecording.h b/Engine/Source/Runtime/Experimental/ChaosVDData/Public/ChaosVDRecording.h
--- a/Engine/Source/Runtime/Experimental/ChaosVDData/Public/ChaosVDRecording.h
+++ b/Engine/Source/Runtime/Experimental/ChaosVDData/Public/ChaosVDRecording.h
@@ -17,6 +17,7 @@
 #include <atomic>
 
 #include "DataWrappers/ChaosVDAccelerationStructureDataWrappers.h"
+#include "DataWrappers/ChaosVDDebugShapeDataWrapper.h"
 
 namespace Chaos::VisualDebugger
 {
@@ -129,6 +130,15 @@ PRAGMA_ENABLE_DEPRECATION_WARNINGS
 
 struct FChaosVDGameFrameData
 {
+
+PRAGMA_DISABLE_DEPRECATION_WARNINGS
+	FChaosVDGameFrameData() = default;
+	FChaosVDGameFrameData(const FChaosVDGameFrameData& Other) = default;
+	FChaosVDGameFrameData(FChaosVDGameFrameData&& Other) noexcept = default;
+	FChaosVDGameFrameData& operator=(const FChaosVDGameFrameData& Other) = default;
+	FChaosVDGameFrameData& operator=(FChaosVDGameFrameData&& Other) noexcept = default;
+PRAGMA_ENABLE_DEPRECATION_WARNINGS
+
 	uint64 FirstCycle;
 	uint64 LastCycle;
 	double StartTime = -1.0;
@@ -149,8 +159,18 @@ struct FChaosVDGameFrameData
 
 	TMap<FName, FChaosVDTrackedLocation> RecordedNonSolverLocationsByID;
 	TMap<FName, FChaosVDTrackedTransform> RecordedNonSolverTransformsByID;
+	TMap<int32, TMap<int32, TSharedPtr<FChaosVDQueryDataWrapper>>> RecordedSceneQueriesBySolverID;
+
+	UE_DEPRECATED(5.5, "RecordedSceneQueries is deprecated and will be removed in a future release, use RecordedSceneQueriesByQueryID instead.")
 	TMap<int32, TSharedPtr<FChaosVDQueryDataWrapper>> RecordedSceneQueries;
+
+	TMap<int32, TSharedPtr<FChaosVDQueryDataWrapper>> RecordedSceneQueriesByQueryID;
 	TMap<int32, TArray<TSharedPtr<FChaosVDAABBTreeDataWrapper>>> RecordedAABBTreesBySolverID;
+
+	TMap<int32,TArray<TSharedPtr<FChaosVDDebugDrawBoxDataWrapper>>> RecordedDebugDrawBoxesBySolverID;
+	TMap<int32,TArray<TSharedPtr<FChaosVDDebugDrawLineDataWrapper>>> RecordedDebugDrawLinesBySolverID;
+	TMap<int32,TArray<TSharedPtr<FChaosVDDebugDrawSphereDataWrapper>>> RecordedDebugDrawSpheresBySolverID;
+	TMap<int32,TArray<TSharedPtr<FChaosVDDebugDrawImplicitObjectDataWrapper>>> RecordedDebugDrawImplicitObjectsBySolverID;
 };
 
 /**
diff --git a/Engine/Source/Runtime/Experimental/ChaosVisualDebugger/Private/DataWrappers/ChaosVDAccelerationStructureDataWrappers.cpp b/Engine/Source/Runtime/Experimental/ChaosVisualDebugger/Private/DataWrappers/ChaosVDAccelerationStructureDataWrappers.cpp
--- a/Engine/Source/Runtime/Experimental/ChaosVisualDebugger/Private/DataWrappers/ChaosVDAccelerationStructureDataWrappers.cpp
+++ b/Engine/Source/Runtime/Experimental/ChaosVisualDebugger/Private/DataWrappers/ChaosVDAccelerationStructureDataWrappers.cpp
@@ -86,10 +86,12 @@ bool FChaosVDAABBTreePayloadBoundsElement::Serialize(FArchive& Ar)
 	Ar << ParticleIndex;
 	Ar << Bounds;
 
+	// TODO: This will be uncommented in the nex CL with the required object version bump
+	//Ar << ActualBounds;
+
 	return !Ar.IsError();
 }
 
-
 bool FChaosVDAABBTreeLeafDataWrapper::Serialize(FArchive& Ar)
 {
 	Ar << bHasValidData;
diff --git a/Engine/Source/Runtime/Experimental/ChaosVisualDebugger/Private/DataWrappers/ChaosVDDebugShapeDataWrapper.cpp b/Engine/Source/Runtime/Experimental/ChaosVisualDebugger/Private/DataWrappers/ChaosVDDebugShapeDataWrapper.cpp
new file mode 100644
--- /dev/null
+++ b/Engine/Source/Runtime/Experimental/ChaosVisualDebugger/Private/DataWrappers/ChaosVDDebugShapeDataWrapper.cpp
@@ -0,0 +1,78 @@
+// Copyright Epic Games, Inc. All Rights Reserved.
+
+#include "DataWrappers/ChaosVDDebugShapeDataWrapper.h"
+
+void FChaosVDDebugDrawShapeBase::SerializeBase_Internal(FArchive& Ar)
+{
+	Ar << SolverID;
+	Ar << Tag;
+	Ar << Color;
+}
+
+bool FChaosVDDebugDrawBoxDataWrapper::Serialize(FArchive& Ar)
+{
+	Ar << bHasValidData;
+
+	if (!bHasValidData)
+	{
+		return !Ar.IsError();
+	}
+
+	SerializeBase_Internal(Ar);
+
+	Ar << Box;
+	
+	return !Ar.IsError();
+}
+
+bool FChaosVDDebugDrawSphereDataWrapper::Serialize(FArchive& Ar)
+{
+	Ar << bHasValidData;
+
+	if (!bHasValidData)
+	{
+		return !Ar.IsError();
+	}
+
+	SerializeBase_Internal(Ar);
+
+	Ar << Origin;
+	Ar << Radius;
+	
+	return !Ar.IsError();
+}
+
+bool FChaosVDDebugDrawLineDataWrapper::Serialize(FArchive& Ar)
+{
+	Ar << bHasValidData;
+
+	if (!bHasValidData)
+	{
+		return !Ar.IsError();
+	}
+
+	SerializeBase_Internal(Ar);
+
+	Ar << StartLocation;
+	Ar << EndLocation;
+	Ar << bIsArrow;
+	
+	return !Ar.IsError();
+}
+
+bool FChaosVDDebugDrawImplicitObjectDataWrapper::Serialize(FArchive& Ar)
+{
+	Ar << bHasValidData;
+
+	if (!bHasValidData)
+	{
+		return !Ar.IsError();
+	}
+
+	SerializeBase_Internal(Ar);
+
+	Ar << ImplicitObjectHash;
+	Ar << ParentTransform;
+	
+	return !Ar.IsError();
+}
diff --git a/Engine/Source/Runtime/Experimental/ChaosVisualDebugger/Public/DataWrappers/ChaosVDAccelerationStructureDataWrappers.h b/Engine/Source/Runtime/Experimental/ChaosVisualDebugger/Public/DataWrappers/ChaosVDAccelerationStructureDataWrappers.h
--- a/Engine/Source/Runtime/Experimental/ChaosVisualDebugger/Public/DataWrappers/ChaosVDAccelerationStructureDataWrappers.h
+++ b/Engine/Source/Runtime/Experimental/ChaosVisualDebugger/Public/DataWrappers/ChaosVDAccelerationStructureDataWrappers.h
@@ -102,8 +102,13 @@ struct FChaosVDAABBTreePayloadBoundsElement : public FChaosVDWrapperDataBase
 	UPROPERTY(VisibleAnywhere, Category="Leaf")
 	int32 ParticleIndex = INDEX_NONE;
 
+	/** Bounds used to add this element into the AABBTree */
 	UPROPERTY(VisibleAnywhere, Category="Leaf")
 	FBox Bounds = FBox(ForceInitToZero);
+
+	/** Real Bounds of the element at the time we recorded the AABB Tree */
+	//UPROPERTY(VisibleAnywhere, Category="Leaf")
+	FBox ActualBounds = FBox(ForceInitToZero);
 	
 	CHAOSVDRUNTIME_API bool Serialize(FArchive& Ar);
 };
diff --git a/Engine/Source/Runtime/Experimental/ChaosVisualDebugger/Public/DataWrappers/ChaosVDDebugShapeDataWrapper.h b/Engine/Source/Runtime/Experimental/ChaosVisualDebugger/Public/DataWrappers/ChaosVDDebugShapeDataWrapper.h
new file mode 100644
--- /dev/null
+++ b/Engine/Source/Runtime/Experimental/ChaosVisualDebugger/Public/DataWrappers/ChaosVDDebugShapeDataWrapper.h
@@ -0,0 +1,93 @@
+// Copyright Epic Games, Inc. All Rights Reserved.
+
+#pragma once
+#include "ChaosVDDataSerializationMacros.h"
+#include "ChaosVDParticleDataWrapper.h"
+
+#include "ChaosVDDebugShapeDataWrapper.generated.h"
+
+USTRUCT()
+struct CHAOSVDRUNTIME_API FChaosVDDebugDrawShapeBase : public FChaosVDWrapperDataBase
+{
+	GENERATED_BODY()
+
+	UPROPERTY()
+	int32 SolverID = INDEX_NONE;
+	
+	UPROPERTY()
+	FName Tag = NAME_None;
+
+	UPROPERTY()
+	FColor Color = FColor::Blue;
+
+protected:
+	void SerializeBase_Internal(FArchive& Ar);
+};
+
+USTRUCT()
+struct CHAOSVDRUNTIME_API FChaosVDDebugDrawBoxDataWrapper : public FChaosVDDebugDrawShapeBase
+{
+	GENERATED_BODY()
+
+	inline static FStringView WrapperTypeName = TEXT("FChaosVDDDebugDrawBoxDataWrapper");
+
+	UPROPERTY()
+	FBox Box = FBox(ForceInitToZero);
+
+	bool Serialize(FArchive& Ar);
+};
+
+CVD_IMPLEMENT_SERIALIZER(FChaosVDDebugDrawBoxDataWrapper)
+
+USTRUCT()
+struct CHAOSVDRUNTIME_API FChaosVDDebugDrawSphereDataWrapper : public FChaosVDDebugDrawShapeBase
+{
+	GENERATED_BODY()
+	
+	inline static FStringView WrapperTypeName = TEXT("FChaosVDDebugDrawSphereDataWrapper");
+
+	UPROPERTY()
+	FVector Origin = FVector::ZeroVector;
+	
+	UPROPERTY()
+	float Radius = 0.0f;
+
+	bool Serialize(FArchive& Ar);
+};
+
+CVD_IMPLEMENT_SERIALIZER(FChaosVDDebugDrawSphereDataWrapper)
+
+USTRUCT()
+struct CHAOSVDRUNTIME_API FChaosVDDebugDrawLineDataWrapper : public FChaosVDDebugDrawShapeBase
+{
+	GENERATED_BODY()
+	
+	inline static FStringView WrapperTypeName = TEXT("FChaosVDDebugDrawLineDataWrapper");
+
+	UPROPERTY()
+	FVector StartLocation = FVector::ZeroVector;
+
+	UPROPERTY()
+	FVector EndLocation = FVector::ZeroVector;
+
+	UPROPERTY()
+	bool bIsArrow = false;
+
+	bool Serialize(FArchive& Ar);
+};
+
+CVD_IMPLEMENT_SERIALIZER(FChaosVDDebugDrawLineDataWrapper)
+
+USTRUCT()
+struct CHAOSVDRUNTIME_API FChaosVDDebugDrawImplicitObjectDataWrapper : public FChaosVDDebugDrawShapeBase
+{
+	GENERATED_BODY()
+	
+	inline static FStringView WrapperTypeName = TEXT("FChaosVDDebugDrawImplicitObjectDataWrapper");
+
+	uint32 ImplicitObjectHash = 0;
+
+	FTransform ParentTransform = FTransform();
+
+	bool Serialize(FArchive& Ar);
+};
