Fixes for TObjectPtr issues identified by LLVM during buildfarm cooks for Linux.

diff --git a/Engine/Plugins/ChaosVD/Source/ChaosVD/Private/Actors/ChaosVDSolverInfoActor.cpp b/Engine/Plugins/ChaosVD/Source/ChaosVD/Private/Actors/ChaosVDSolverInfoActor.cpp
--- a/Engine/Plugins/ChaosVD/Source/ChaosVD/Private/Actors/ChaosVDSolverInfoActor.cpp
+++ b/Engine/Plugins/ChaosVD/Source/ChaosVD/Private/Actors/ChaosVDSolverInfoActor.cpp
@@ -144,7 +144,7 @@ bool AChaosVDSolverInfoActor::SelectParticleByID(int32 ParticleIDToSelect)
 
 void AChaosVDSolverInfoActor::HandleVisibilitySettingsUpdated(UObject* SettingsObject)
 {
-	for (const TPair<int32, AChaosVDParticleActor*>& ParticleWithIDPair : SolverParticlesByID)
+	for (const TPair<int32, TObjectPtr<AChaosVDParticleActor>>& ParticleWithIDPair : SolverParticlesByID)
 	{
 		if (AChaosVDParticleActor* ParticleActor = ParticleWithIDPair.Value)
 		{
@@ -155,7 +155,7 @@ void AChaosVDSolverInfoActor::HandleVisibilitySettingsUpdated(UObject* SettingsO
 
 void AChaosVDSolverInfoActor::HandleColorsSettingsUpdated(UObject* SettingsObject)
 {
-	for (const TPair<int32, AChaosVDParticleActor*>& ParticleWithIDPair : SolverParticlesByID)
+	for (const TPair<int32, TObjectPtr<AChaosVDParticleActor>>& ParticleWithIDPair : SolverParticlesByID)
 	{
 		if (AChaosVDParticleActor* ParticleActor = ParticleWithIDPair.Value)
 		{
@@ -228,7 +228,7 @@ void AChaosVDSolverInfoActor::SetIsTemporarilyHiddenInEditor(bool bIsHidden)
 		FScopedSlowTask VisibilityUpdateProgress(AmountOfWork, LOCTEXT("UpdatingParticlesVisisibility", "Updating Particles Visibility ..."));
 		VisibilityUpdateProgress.MakeDialog();
 
-		for (const TPair<int32, AChaosVDParticleActor*>& ParticleVDInstanceWithID : SolverParticlesByID)
+		for (const TPair<int32, TObjectPtr<AChaosVDParticleActor>>& ParticleVDInstanceWithID : SolverParticlesByID)
 		{
 			ApplySolverVisibilityToParticle(ParticleVDInstanceWithID.Value, bIsHidden);
 
@@ -269,7 +269,7 @@ void AChaosVDSolverInfoActor::Destroyed()
 		FScopedSlowTask CleaningParticleDataSlowTask(AmountOfWork, LOCTEXT("CleaningParticleDataMessage", "Cleaning Up Particle Data ..."));
 		CleaningParticleDataSlowTask.MakeDialog();
 
-		for (const TPair<int32, AChaosVDParticleActor*>& ParticleVDInstanceWithID : SolverParticlesByID)
+		for (const TPair<int32, TObjectPtr<AChaosVDParticleActor>>& ParticleVDInstanceWithID : SolverParticlesByID)
 		{
 			World->DestroyActor(ParticleVDInstanceWithID.Value);
 
diff --git a/Engine/Plugins/ChaosVD/Source/ChaosVD/Private/Actors/ChaosVDSolverInfoActor.h b/Engine/Plugins/ChaosVD/Source/ChaosVD/Private/Actors/ChaosVDSolverInfoActor.h
--- a/Engine/Plugins/ChaosVD/Source/ChaosVD/Private/Actors/ChaosVDSolverInfoActor.h
+++ b/Engine/Plugins/ChaosVD/Source/ChaosVD/Private/Actors/ChaosVDSolverInfoActor.h
@@ -154,7 +154,7 @@ void AChaosVDSolverInfoActor::VisitSelectedParticleData(TCallback VisitCallback)
 template <typename TCallback>
 void AChaosVDSolverInfoActor::VisitAllParticleData(TCallback VisitCallback)
 {
-	for (const TPair<int32, AChaosVDParticleActor*>& ParticleWithIDPair : SolverParticlesByID)
+	for (const TPair<int32, TObjectPtr<AChaosVDParticleActor>>& ParticleWithIDPair : SolverParticlesByID)
 	{
 		AChaosVDParticleActor* ParticleActor = ParticleWithIDPair.Value;
 		TSharedPtr<const FChaosVDParticleDataWrapper> ParticleDataViewer = ParticleActor ? ParticleActor->GetParticleData() : nullptr;
