[ChaosVD] Adding support to record data outside a solver stage.

Initially, CVD only supported recording data within a recorded solver stage, everything else was discarded and an ensure was triggered.

But currently, we also record events that might not fall withing a predefined stage (like disabling a particle or creating a particle). For these cases we added them to the last available stage on load, and discard them if there was none available.

This CL add the possibility to properly create a new solver stage entry during load, if we reach one of these scenarios. This allow us to ensure the data is applied in the correct order, and clearly see where the change of state (if any came from) as it will be shown as an stage named "In Between Stage Data".

diff --git a/Engine/Plugins/ChaosVD/Source/ChaosVD/Private/ChaosVDPlaybackController.cpp b/Engine/Plugins/ChaosVD/Source/ChaosVD/Private/ChaosVDPlaybackController.cpp
--- a/Engine/Plugins/ChaosVD/Source/ChaosVD/Private/ChaosVDPlaybackController.cpp
+++ b/Engine/Plugins/ChaosVD/Source/ChaosVD/Private/ChaosVDPlaybackController.cpp
@@ -558,6 +558,12 @@ int32 FChaosVDPlaybackController::GetTrackCurrentStep(EChaosVDTrackType TrackTyp
 	return INDEX_NONE;
 }
 
+int32 FChaosVDPlaybackController::GetTrackLastStageAtFrame(EChaosVDTrackType TrackType, int32 InTrackID, int32 InFrameNumber) const
+{
+	FReadScopeLock ReadLock(LoadedRecording->GetRecordingDataLock());
+	return GetTrackLastStageAtFrame_AssumesLocked(TrackType, InTrackID, InFrameNumber);
+}
+
 int32 FChaosVDPlaybackController::GetTrackLastStageAtFrame_AssumesLocked(EChaosVDTrackType TrackType, int32 InTrackID, int32 InFrameNumber) const
 {
 	switch (TrackType)
@@ -570,7 +576,6 @@ int32 FChaosVDPlaybackController::GetTrackLastStageAtFrame_AssumesLocked(EChaosV
 		}
 		case EChaosVDTrackType::Game:
 		default:
-			ensureMsgf(false, TEXT("Unsuported Track Type"));
 			return INDEX_NONE;
 			break;
 	}
@@ -935,15 +940,15 @@ void FChaosVDPlaybackController::HandleFramePlaybackControlInput(EChaosVDPlaybac
 		}
 	case EChaosVDPlaybackButtonsID::Next:
 		{
-			const int32 LastSolverStage = InTrackInfoRef->CurrentStageNames.Num() - 1;
 			const int32 NextFrame = InTrackInfoRef->CurrentFrame +1;
+			const int32 LastSolverStage = GetTrackLastStageAtFrame(InTrackInfoRef->TrackType, InTrackInfoRef->TrackID, NextFrame);
 			GoToTrackFrameAndSync(Instigator, InTrackInfoRef->TrackType, InTrackInfoRef->TrackID, NextFrame, LastSolverStage);
 		}
 		break;
 	case EChaosVDPlaybackButtonsID::Prev:
 		{
-			const int32 LastSolverStage = InTrackInfoRef->CurrentStageNames.Num() - 1;
 			const int32 PrevFrame = InTrackInfoRef->CurrentFrame -1;
+			const int32 LastSolverStage = GetTrackLastStageAtFrame(InTrackInfoRef->TrackType, InTrackInfoRef->TrackID, PrevFrame);
 			GoToTrackFrameAndSync(Instigator, InTrackInfoRef->TrackType, InTrackInfoRef->TrackID, PrevFrame, LastSolverStage);
 		}
 		break;
diff --git a/Engine/Plugins/ChaosVD/Source/ChaosVD/Private/ChaosVDPlaybackController.h b/Engine/Plugins/ChaosVD/Source/ChaosVD/Private/ChaosVDPlaybackController.h
--- a/Engine/Plugins/ChaosVD/Source/ChaosVD/Private/ChaosVDPlaybackController.h
+++ b/Engine/Plugins/ChaosVD/Source/ChaosVD/Private/ChaosVDPlaybackController.h
@@ -177,6 +177,7 @@ public:
 	 * @param InTrackID ID of the track to evaluate
 	 * @return Number of the last step
 	 */
+	int32 GetTrackLastStageAtFrame(EChaosVDTrackType TrackType, int32 InTrackID, int32 InFrameNumber) const;
 	int32 GetTrackLastStageAtFrame_AssumesLocked(EChaosVDTrackType TrackType, int32 InTrackID, int32 InFrameNumber) const;
 
 	/** Converts the current frame number of a track, to a frame number in other tracks space time
diff --git a/Engine/Plugins/ChaosVD/Source/ChaosVD/Private/Trace/ChaosVDTraceAnalyzer.cpp b/Engine/Plugins/ChaosVD/Source/ChaosVD/Private/Trace/ChaosVDTraceAnalyzer.cpp
--- a/Engine/Plugins/ChaosVD/Source/ChaosVD/Private/Trace/ChaosVDTraceAnalyzer.cpp
+++ b/Engine/Plugins/ChaosVD/Source/ChaosVD/Private/Trace/ChaosVDTraceAnalyzer.cpp
@@ -122,18 +122,38 @@ bool FChaosVDTraceAnalyzer::OnEvent(uint16 RouteId, EStyle Style, const FOnEvent
 			// This can be null if the recording started Mid-Frame. In this case we just discard the data for now
 			if (FChaosVDSolverFrameData* FrameData = ChaosVDTraceProvider->GetCurrentSolverFrame(SolverID))
 			{
+				if (FrameData->SolverSteps.Num() > 0)
+				{
+					FChaosVDStepData& LastSolverStage = FrameData->SolverSteps.Last();
+					if (EnumHasAnyFlags(LastSolverStage.StageFlags, EChaosVDSolverStageFlags::Open) && ensure(!EnumHasAnyFlags(LastSolverStage.StageFlags, EChaosVDSolverStageFlags::ExplicitStage)))
+					{
+						// If the current Solver stage was implicitly generated, we need to close it before starting a new one.
+						// This should not happen with an explicitly recorded stage
+						EnumRemoveFlags(LastSolverStage.StageFlags, EChaosVDSolverStageFlags::Open);
+					}
+				}
+
 				// Add an empty step. It will be filled out by the particle (and later on other objects/elements) events
-				FChaosVDStepData& StepData = FrameData->SolverSteps.AddDefaulted_GetRef();
+				FChaosVDStepData& NewSolverStageData = FrameData->SolverSteps.AddDefaulted_GetRef();
 
 				FWideStringView DebugNameView;
 				EventData.GetString("StepName", DebugNameView);
-				StepData.StepName = DebugNameView;
+				NewSolverStageData.StepName = DebugNameView;
+				EnumAddFlags(NewSolverStageData.StageFlags, EChaosVDSolverStageFlags::Open);
 			}
 	
 			break;
 		}
 	case RouteId_ChaosVDSolverStepEnd:
 		{
+			const int32 SolverID = EventData.GetValue<int32>("SolverID");
+			if (FChaosVDSolverFrameData* FrameData = ChaosVDTraceProvider->GetCurrentSolverFrame(SolverID))
+			{
+				if (FrameData->SolverSteps.Num() > 0)
+				{
+					EnumRemoveFlags(FrameData->SolverSteps.Last().StageFlags, EChaosVDSolverStageFlags::Open);
+				}
+			}
 			break;
 		}
 	case RouteId_ChaosVDParticleDestroyed:
diff --git a/Engine/Plugins/ChaosVD/Source/ChaosVD/Private/Trace/ChaosVDTraceProvider.cpp b/Engine/Plugins/ChaosVD/Source/ChaosVD/Private/Trace/ChaosVDTraceProvider.cpp
--- a/Engine/Plugins/ChaosVD/Source/ChaosVD/Private/Trace/ChaosVDTraceProvider.cpp
+++ b/Engine/Plugins/ChaosVD/Source/ChaosVD/Private/Trace/ChaosVDTraceProvider.cpp
@@ -232,6 +232,43 @@ void FChaosVDTraceProvider::HandleAnalysisComplete()
 	UE_LOG(LogChaosVDEditor, Log, TEXT("Total size of loaded data => [%s]"), *FText::AsMemory(TotalBytes, &SizeFormattingOptions,nullptr, EMemoryUnitStandard::IEC).ToString());
 }
 
+FChaosVDStepData* FChaosVDTraceProvider::GetCurrentSolverStageDataForCurrentFrame(int32 SolverID, EChaosVDSolverStageAccessorFlags Flags)
+{
+	auto CreateInBetweenSolverStage = [](FChaosVDSolverFrameData& InFrameData)
+	{
+		// Add an empty step. It will be filled out by the particle (and later on other objects/elements) events
+		FChaosVDStepData& SolverStageData = InFrameData.SolverSteps.AddDefaulted_GetRef();
+		SolverStageData.StepName = TEXT("Between Stage Data");
+		EnumAddFlags(SolverStageData.StageFlags, EChaosVDSolverStageFlags::Open);
+
+		return &SolverStageData;
+	};
+
+	if (FChaosVDSolverFrameData* FrameData = GetCurrentSolverFrame(SolverID))
+	{
+		if (FrameData->SolverSteps.Num() == 0)
+		{
+			if (EnumHasAnyFlags(Flags, EChaosVDSolverStageAccessorFlags::CreateNewIfEmpty))
+			{
+				return CreateInBetweenSolverStage(*FrameData);
+			}
+		}
+
+		FChaosVDStepData& CurrentSolverStage = FrameData->SolverSteps.Last();
+		if (EnumHasAnyFlags(CurrentSolverStage.StageFlags, EChaosVDSolverStageFlags::Open))
+		{
+			return &CurrentSolverStage;
+		}
+
+		if (EnumHasAnyFlags(Flags, EChaosVDSolverStageAccessorFlags::CreateNewIfClosed))
+		{
+			return CreateInBetweenSolverStage(*FrameData);
+		}
+	}
+
+	return nullptr;
+}
+
 void FChaosVDTraceProvider::RegisterDefaultDataProcessorsIfNeeded()
 {
 	if (bDefaultDataProcessorsRegistered)
diff --git a/Engine/Plugins/ChaosVD/Source/ChaosVD/Private/Trace/ChaosVDTraceProvider.h b/Engine/Plugins/ChaosVD/Source/ChaosVD/Private/Trace/ChaosVDTraceProvider.h
--- a/Engine/Plugins/ChaosVD/Source/ChaosVD/Private/Trace/ChaosVDTraceProvider.h
+++ b/Engine/Plugins/ChaosVD/Source/ChaosVD/Private/Trace/ChaosVDTraceProvider.h
@@ -47,6 +47,16 @@ struct FChaosVDTraceSessionData
 	TMap<int32, TSharedPtr<FChaosVDBinaryDataContainer>> UnprocessedDataByID;
 };
 
+enum class EChaosVDSolverStageAccessorFlags : uint8
+{
+	None = 0,
+	/* If the solver frame has valid stage data but the last stage is closed, create a new stage which will be labeled as non-staged data */
+	CreateNewIfClosed = 1 << 0,
+	/* If the solver frame does not have any solver stage data, create a new stage which will be labeled as non-staged data */
+	CreateNewIfEmpty = 1 << 1 
+};
+ENUM_CLASS_FLAGS(EChaosVDSolverStageAccessorFlags);
+
 /** Provider class for Chaos VD trace recordings.
  * It stores and handles rebuilt recorded frame data from Trace events
  * dispatched by the Chaos VD Trace analyzer
@@ -79,7 +89,10 @@ public:
 
 	void HandleAnalysisComplete();
 
-	TMap<int32,int32>& GetCurrentTickOffsetsBySolverID() {return CurrentNetworkTickOffsets; };
+	TMap<int32,int32>& GetCurrentTickOffsetsBySolverID() { return CurrentNetworkTickOffsets; };
+
+
+	FChaosVDStepData* GetCurrentSolverStageDataForCurrentFrame(int32 SolverID, EChaosVDSolverStageAccessorFlags Flags);
 
 private:
 
diff --git a/Engine/Plugins/ChaosVD/Source/ChaosVD/Private/Trace/DataProcessors/ChaosVDConstraintDataProcessor.cpp b/Engine/Plugins/ChaosVD/Source/ChaosVD/Private/Trace/DataProcessors/ChaosVDConstraintDataProcessor.cpp
--- a/Engine/Plugins/ChaosVD/Source/ChaosVD/Private/Trace/DataProcessors/ChaosVDConstraintDataProcessor.cpp
+++ b/Engine/Plugins/ChaosVD/Source/ChaosVD/Private/Trace/DataProcessors/ChaosVDConstraintDataProcessor.cpp
@@ -26,28 +26,25 @@ bool FChaosVDConstraintDataProcessor::ProcessRawData(const TArray<uint8>& InData
 
 	if (bSuccess)
 	{
-		// This can be null if the recording started Mid-Frame. In this case we just discard the data for now
-		if (FChaosVDSolverFrameData* FrameData = ProviderSharedPtr->GetCurrentSolverFrame(RecordedConstraint.SolverID))
+		FChaosVDStepData* CurrentSolverStage = ProviderSharedPtr->GetCurrentSolverStageDataForCurrentFrame(RecordedConstraint.SolverID, EChaosVDSolverStageAccessorFlags::None);
+		if (ensureMsgf(CurrentSolverStage, TEXT("A MidPhase was traced without a valid step scope")))
 		{
-			if (ensureMsgf(FrameData->SolverSteps.Num() > 0, TEXT("A MidPhase was traced without a valid step scope")))
-			{
-				AddConstraintToParticleIDMap(RecordedConstraint, RecordedConstraint.Particle0Index, *FrameData);
-				AddConstraintToParticleIDMap(RecordedConstraint, RecordedConstraint.Particle1Index, *FrameData);
-			}
+			AddConstraintToParticleIDMap(RecordedConstraint, RecordedConstraint.Particle0Index, *CurrentSolverStage);
+			AddConstraintToParticleIDMap(RecordedConstraint, RecordedConstraint.Particle1Index, *CurrentSolverStage);
 		}
 	}
 
 	return bSuccess;
 }
 
-void FChaosVDConstraintDataProcessor::AddConstraintToParticleIDMap(const FChaosVDConstraint& InConstraintData, int32 ParticleID, FChaosVDSolverFrameData& InFrameData)
+void FChaosVDConstraintDataProcessor::AddConstraintToParticleIDMap(const FChaosVDConstraint& InConstraintData, int32 ParticleID, FChaosVDStepData& InSolverStageData)
 {
-	if (TArray<FChaosVDConstraint>* ParticleConstraints = InFrameData.SolverSteps.Last().RecordedConstraintsByParticleID.Find(ParticleID))
+	if (TArray<FChaosVDConstraint>* ParticleConstraints = InSolverStageData.RecordedConstraintsByParticleID.Find(ParticleID))
 	{
 		ParticleConstraints->Add(InConstraintData);
 	}
 	else
 	{
-		InFrameData.SolverSteps.Last().RecordedConstraintsByParticleID.Add(ParticleID, { InConstraintData });
+		InSolverStageData.RecordedConstraintsByParticleID.Add(ParticleID, { InConstraintData });
 	}
 }
diff --git a/Engine/Plugins/ChaosVD/Source/ChaosVD/Private/Trace/DataProcessors/ChaosVDConstraintDataProcessor.h b/Engine/Plugins/ChaosVD/Source/ChaosVD/Private/Trace/DataProcessors/ChaosVDConstraintDataProcessor.h
--- a/Engine/Plugins/ChaosVD/Source/ChaosVD/Private/Trace/DataProcessors/ChaosVDConstraintDataProcessor.h
+++ b/Engine/Plugins/ChaosVD/Source/ChaosVD/Private/Trace/DataProcessors/ChaosVDConstraintDataProcessor.h
@@ -4,6 +4,7 @@
 #include "ChaosVDDataProcessorBase.h"
 
 struct FChaosVDSolverFrameData;
+struct FChaosVDStepData;
 struct FChaosVDConstraint;
 
 /**
@@ -16,5 +17,5 @@ public:
 	
 	virtual bool ProcessRawData(const TArray<uint8>& InData) override;
 
-	void AddConstraintToParticleIDMap(const FChaosVDConstraint& InConstraintData, int32 ParticleID, FChaosVDSolverFrameData& InFrameData);
+	void AddConstraintToParticleIDMap(const FChaosVDConstraint& InConstraintData, int32 ParticleID, FChaosVDStepData& InSolverStageData);
 };
diff --git a/Engine/Plugins/ChaosVD/Source/ChaosVD/Private/Trace/DataProcessors/ChaosVDJointConstraintDataProcessor.cpp b/Engine/Plugins/ChaosVD/Source/ChaosVD/Private/Trace/DataProcessors/ChaosVDJointConstraintDataProcessor.cpp
--- a/Engine/Plugins/ChaosVD/Source/ChaosVD/Private/Trace/DataProcessors/ChaosVDJointConstraintDataProcessor.cpp
+++ b/Engine/Plugins/ChaosVD/Source/ChaosVD/Private/Trace/DataProcessors/ChaosVDJointConstraintDataProcessor.cpp
@@ -26,13 +26,10 @@ bool FChaosVDJointConstraintDataProcessor::ProcessRawData(const TArray<uint8>& I
 
 	if (bSuccess)
 	{
-		// This can be null if the recording started Mid-Frame. In this case we just discard the data for now
-		if (FChaosVDSolverFrameData* FrameData = ProviderSharedPtr->GetCurrentSolverFrame(JointConstraint->SolverID))
+		FChaosVDStepData* CurrentSolverStage = ProviderSharedPtr->GetCurrentSolverStageDataForCurrentFrame(JointConstraint->SolverID, EChaosVDSolverStageAccessorFlags::None);
+		if (ensureMsgf(CurrentSolverStage, TEXT("A Joint Constraint was traced without a valid step scope")))
 		{
-			if (ensureMsgf(FrameData->SolverSteps.Num() > 0, TEXT("A Joint Constraint was traced without a valid step scope")))
-			{
-				FrameData->SolverSteps.Last().RecordedJointConstraints.Add(JointConstraint);
-			}
+			CurrentSolverStage->RecordedJointConstraints.Add(JointConstraint);
 		}
 	}
 
diff --git a/Engine/Plugins/ChaosVD/Source/ChaosVD/Private/Trace/DataProcessors/ChaosVDMidPhaseDataProcessor.cpp b/Engine/Plugins/ChaosVD/Source/ChaosVD/Private/Trace/DataProcessors/ChaosVDMidPhaseDataProcessor.cpp
--- a/Engine/Plugins/ChaosVD/Source/ChaosVD/Private/Trace/DataProcessors/ChaosVDMidPhaseDataProcessor.cpp
+++ b/Engine/Plugins/ChaosVD/Source/ChaosVD/Private/Trace/DataProcessors/ChaosVDMidPhaseDataProcessor.cpp
@@ -26,30 +26,28 @@ bool FChaosVDMidPhaseDataProcessor::ProcessRawData(const TArray<uint8>& InData)
 
 	if (bSuccess)
 	{
-		// This can be null if the recording started Mid-Frame. In this case we just discard the data for now
-		if (FChaosVDSolverFrameData* FrameData = ProviderSharedPtr->GetCurrentSolverFrame(MidPhase->SolverID))
+		FChaosVDStepData* CurrentSolverStage = ProviderSharedPtr->GetCurrentSolverStageDataForCurrentFrame(MidPhase->SolverID, EChaosVDSolverStageAccessorFlags::None);
+
+		if (ensureMsgf(CurrentSolverStage, TEXT("A MidPhase was traced without a valid step scope")))
 		{
-			if (ensureMsgf(FrameData->SolverSteps.Num() > 0, TEXT("A MidPhase was traced without a valid step scope")))
-			{
-				FrameData->SolverSteps.Last().RecordedMidPhases.Add(MidPhase);
+			CurrentSolverStage->RecordedMidPhases.Add(MidPhase);
 
-				AddMidPhaseToParticleIDMap(MidPhase, MidPhase->Particle0Idx, *FrameData);
-				AddMidPhaseToParticleIDMap(MidPhase, MidPhase->Particle1Idx, *FrameData);
-			}
+			AddMidPhaseToParticleIDMap(MidPhase, MidPhase->Particle0Idx, *CurrentSolverStage);
+			AddMidPhaseToParticleIDMap(MidPhase, MidPhase->Particle1Idx, *CurrentSolverStage);
 		}
 	}
 
 	return bSuccess;
 }
 
-void FChaosVDMidPhaseDataProcessor::AddMidPhaseToParticleIDMap(const TSharedPtr<FChaosVDParticlePairMidPhase>& MidPhaseData, int32 ParticleID, FChaosVDSolverFrameData& InFrameData)
+void FChaosVDMidPhaseDataProcessor::AddMidPhaseToParticleIDMap(const TSharedPtr<FChaosVDParticlePairMidPhase>& MidPhaseData, int32 ParticleID, FChaosVDStepData& InSolverStageData)
 {
-	if (TArray<TSharedPtr<FChaosVDParticlePairMidPhase>>* ParticleMidPhases = InFrameData.SolverSteps.Last().RecordedMidPhasesByParticleID.Find(ParticleID))
+	if (TArray<TSharedPtr<FChaosVDParticlePairMidPhase>>* ParticleMidPhases = InSolverStageData.RecordedMidPhasesByParticleID.Find(ParticleID))
 	{
 		ParticleMidPhases->Add(MidPhaseData);
 	}
 	else
 	{
-		InFrameData.SolverSteps.Last().RecordedMidPhasesByParticleID.Add(ParticleID, { MidPhaseData });
+		InSolverStageData.RecordedMidPhasesByParticleID.Add(ParticleID, { MidPhaseData });
 	}
 }
diff --git a/Engine/Plugins/ChaosVD/Source/ChaosVD/Private/Trace/DataProcessors/ChaosVDMidPhaseDataProcessor.h b/Engine/Plugins/ChaosVD/Source/ChaosVD/Private/Trace/DataProcessors/ChaosVDMidPhaseDataProcessor.h
--- a/Engine/Plugins/ChaosVD/Source/ChaosVD/Private/Trace/DataProcessors/ChaosVDMidPhaseDataProcessor.h
+++ b/Engine/Plugins/ChaosVD/Source/ChaosVD/Private/Trace/DataProcessors/ChaosVDMidPhaseDataProcessor.h
@@ -5,6 +5,7 @@
 #include "Templates/SharedPointer.h"
 
 struct FChaosVDSolverFrameData;
+struct FChaosVDStepData;
 struct FChaosVDParticlePairMidPhase;
 
 /**
@@ -17,5 +18,5 @@ public:
 
 	virtual bool ProcessRawData(const TArray<uint8>& InData) override;
 
-	void AddMidPhaseToParticleIDMap(const TSharedPtr<FChaosVDParticlePairMidPhase>& MidPhaseData, int32 ParticleID, FChaosVDSolverFrameData& InFrameData);
+	void AddMidPhaseToParticleIDMap(const TSharedPtr<FChaosVDParticlePairMidPhase>& MidPhaseData, int32 ParticleID, FChaosVDStepData& InSolverStageData);
 };
diff --git a/Engine/Plugins/ChaosVD/Source/ChaosVD/Private/Trace/DataProcessors/ChaosVDTraceParticleDataProcessor.cpp b/Engine/Plugins/ChaosVD/Source/ChaosVD/Private/Trace/DataProcessors/ChaosVDTraceParticleDataProcessor.cpp
--- a/Engine/Plugins/ChaosVD/Source/ChaosVD/Private/Trace/DataProcessors/ChaosVDTraceParticleDataProcessor.cpp
+++ b/Engine/Plugins/ChaosVD/Source/ChaosVD/Private/Trace/DataProcessors/ChaosVDTraceParticleDataProcessor.cpp
@@ -2,6 +2,7 @@
 
 #include "Trace/DataProcessors/ChaosVDTraceParticleDataProcessor.h"
 
+#include "ChaosVDModule.h"
 #include "ChaosVDRecording.h"
 #include "ChaosVisualDebugger/ChaosVDMemWriterReader.h"
 #include "ChaosVisualDebugger/ChaosVisualDebuggerTrace.h"
@@ -27,13 +28,10 @@ bool FChaosVDTraceParticleDataProcessor::ProcessRawData(const TArray<uint8>& InD
 
 	if (bSuccess)
 	{
-		// This can be null if the recording started Mid-Frame. In this case we just discard the data for now
-		if (FChaosVDSolverFrameData* FrameData = ProviderSharedPtr->GetCurrentSolverFrame(ParticleData->SolverID))
+		EChaosVDSolverStageAccessorFlags StageAccessorFlags = EChaosVDSolverStageAccessorFlags::CreateNewIfEmpty | EChaosVDSolverStageAccessorFlags::CreateNewIfClosed; 
+		if (FChaosVDStepData* CurrentSolverStage = ProviderSharedPtr->GetCurrentSolverStageDataForCurrentFrame(ParticleData->SolverID, StageAccessorFlags))
 		{
-			if (ensureMsgf(FrameData->SolverSteps.Num() > 0, TEXT("A particle was traced without a valid step scope")))
-			{
-				FrameData->SolverSteps.Last().RecordedParticlesData.Add(ParticleData);
-			}
+			CurrentSolverStage->RecordedParticlesData.Add(ParticleData);
 		}
 	}
 
diff --git a/Engine/Source/Runtime/Experimental/ChaosVDData/Private/ChaosVDRecording.cpp b/Engine/Source/Runtime/Experimental/ChaosVDData/Private/ChaosVDRecording.cpp
--- a/Engine/Source/Runtime/Experimental/ChaosVDData/Private/ChaosVDRecording.cpp
+++ b/Engine/Source/Runtime/Experimental/ChaosVDData/Private/ChaosVDRecording.cpp
@@ -373,9 +373,7 @@ void FChaosVDRecording::CollapseSolverFramesRange_AssumesLocked(int32 SolverID,
 		{
 			OutCollapsedFrameData.ParticlesDestroyedIDs.Append(SolverFrameData->ParticlesDestroyedIDs);
 
-			// Only evaluate the last step as it contains all the particles that changed
-			const int32 LastStepNumber = SolverFrameData->SolverSteps.Num() - 1;
-			if (SolverFrameData->SolverSteps.IsValidIndex(LastStepNumber))
+			if (SolverFrameData->SolverSteps.Num() > 0)
 			{
 				for (const FChaosVDStepData& StepData : SolverFrameData->SolverSteps)
 				{
@@ -410,7 +408,7 @@ void FChaosVDRecording::CollapseSolverFramesRange_AssumesLocked(int32 SolverID,
 				OutCollapsedFrameData.DebugFName = SolverFrameData->DebugFName;
 
 				FChaosVDStepData CollapsedStepData;
-				CollapsedStepData.StepName = TEXT("GeneratedStep");
+				CollapsedStepData.StepName = TEXT("Auto Generated Stage");
 
 				ParticlesOnCurrentGeneratedKeyframe.GenerateValueArray(CollapsedStepData.RecordedParticlesData);
 
diff --git a/Engine/Source/Runtime/Experimental/ChaosVDData/Public/ChaosVDRecording.h b/Engine/Source/Runtime/Experimental/ChaosVDData/Public/ChaosVDRecording.h
--- a/Engine/Source/Runtime/Experimental/ChaosVDData/Public/ChaosVDRecording.h
+++ b/Engine/Source/Runtime/Experimental/ChaosVDData/Public/ChaosVDRecording.h
@@ -26,6 +26,18 @@ namespace Chaos::VisualDebugger
 
 DECLARE_MULTICAST_DELEGATE_TwoParams(FChaosVDGeometryDataLoaded, const Chaos::FConstImplicitObjectPtr&, const uint32 GeometryID)
 
+/** Set of flags used to define characteristics of a loaded solver stage */
+enum class EChaosVDSolverStageFlags : uint8
+{
+	None = 0,
+	/** Set if the solver stage is open and can take new data */
+	Open = 1 << 0,
+	/** Set if the solver stage was explicitly recorded - If not set, this stage was created on the fly during load */
+	ExplicitStage = 1 << 1,
+};
+
+ENUM_CLASS_FLAGS(EChaosVDSolverStageFlags)
+
 struct FChaosVDStepData
 {
 	FString StepName;
@@ -36,6 +48,8 @@ struct FChaosVDStepData
 	TMap<int32, TArray<FChaosVDConstraint>> RecordedConstraintsByParticleID;
 	TMap<int32, TArray<TSharedPtr<FChaosVDParticlePairMidPhase>>> RecordedMidPhasesByParticleID;
 	TSet<int32> ParticlesDestroyedIDs;
+
+	EChaosVDSolverStageFlags StageFlags = EChaosVDSolverStageFlags::None;
 };
 
 struct FChaosVDTrackedLocation
