[ChaosVD] First iteration of a Scene Query Browser

- Implemented a new window withing CVD that shows all available scene queries in a Scene Outliner kind of way.

This window allows you to select any query, see if its visibility state, search by its trace tag or owner tag, and see any subquery as a tree hierarchy (most queries that hit a CU and trace into its subcomponents).

The browser view disables itself during playback to not impact performances (sames as we do with the Scene Outliner).

- Implemented a new visualization mode where the instead of drawing all the queries for the current frame at ones, you can instead draw them one at the time in the order they were recorded. You can step into these by using a playback control at the bottom of the scene query browser window.

- Fixed an issue where some queries might dispensary depending the orientation on the camera. the issue was caused because we didn't take into account the shape bounds when we tested to see if we should cull the query (this was a known issue).

- Fixed an issue where if you have the flag "Only Draw Selected Query", the visualization will show out of date information if you change of frame.

diff --git a/Engine/Plugins/ChaosVD/Source/ChaosVD/ChaosVD.Build.cs b/Engine/Plugins/ChaosVD/Source/ChaosVD/ChaosVD.Build.cs
--- a/Engine/Plugins/ChaosVD/Source/ChaosVD/ChaosVD.Build.cs
+++ b/Engine/Plugins/ChaosVD/Source/ChaosVD/ChaosVD.Build.cs
@@ -15,6 +15,7 @@ public class ChaosVD : ModuleRules
 		PrivateDependencyModuleNames.AddRange(
 			new string[]
 			{
+				"ApplicationCore",
 				"CoreUObject",
 				"Engine",
 				"RenderCore",
diff --git a/Engine/Plugins/ChaosVD/Source/ChaosVD/Private/ChaosVDObjectDetailsTab.cpp b/Engine/Plugins/ChaosVD/Source/ChaosVD/Private/ChaosVDObjectDetailsTab.cpp
--- a/Engine/Plugins/ChaosVD/Source/ChaosVD/Private/ChaosVDObjectDetailsTab.cpp
+++ b/Engine/Plugins/ChaosVD/Source/ChaosVD/Private/ChaosVDObjectDetailsTab.cpp
@@ -126,7 +126,7 @@ void FChaosVDObjectDetailsTab::HandlePostSelectionChange(const UTypedElementSele
 
 void FChaosVDObjectDetailsTab::HandleSolverDataSelectionChange(const TSharedPtr<FChaosVDSolverDataSelectionHandle>& SelectionHandle)
 {
-	TSharedPtr<FStructOnScope> StructOnScope = SelectionHandle->GetDataAsStructScope();
+	TSharedPtr<FStructOnScope> StructOnScope = SelectionHandle ? SelectionHandle->GetDataAsStructScope() : nullptr;
 	if (!StructOnScope || !IsSupportedStruct(StructOnScope->GetStructPtr()))
 	{
 		DetailsPanelView->SetSelectedStruct(nullptr);
diff --git a/Engine/Plugins/ChaosVD/Source/ChaosVD/Private/ChaosVDPlaybackViewportClient.cpp b/Engine/Plugins/ChaosVD/Source/ChaosVD/Private/ChaosVDPlaybackViewportClient.cpp
--- a/Engine/Plugins/ChaosVD/Source/ChaosVD/Private/ChaosVDPlaybackViewportClient.cpp
+++ b/Engine/Plugins/ChaosVD/Source/ChaosVD/Private/ChaosVDPlaybackViewportClient.cpp
@@ -36,11 +36,11 @@ FChaosVDPlaybackViewportClient::FChaosVDPlaybackViewportClient(const TSharedPtr<
 
 FChaosVDPlaybackViewportClient::~FChaosVDPlaybackViewportClient()
 {
-	if (ObjectFocusedDelegateHandle.IsValid())
+	if (FocusRequestDelegateHandle.IsValid())
 	{
 		if (TSharedPtr<FChaosVDScene> ScenePtr = CVDScene.Pin())
 		{
-			ScenePtr->OnObjectFocused().Remove(ObjectFocusedDelegateHandle);
+			ScenePtr->OnFocusRequest().Remove(FocusRequestDelegateHandle);
 		}
 	}
 
@@ -141,16 +141,13 @@ void FChaosVDPlaybackViewportClient::SetScene(TWeakPtr<FChaosVDScene> InScene)
 		CVDWorld = ScenePtr->GetUnderlyingWorld();
 		CVDScene = InScene;
 
-		ObjectFocusedDelegateHandle = ScenePtr->OnObjectFocused().AddRaw(this, &FChaosVDPlaybackViewportClient::HandleObjectFocused);
+		FocusRequestDelegateHandle = ScenePtr->OnFocusRequest().AddRaw(this, &FChaosVDPlaybackViewportClient::HandleFocusRequest);
 	}
 }
 
-void FChaosVDPlaybackViewportClient::HandleObjectFocused(UObject* FocusedObject)
+void FChaosVDPlaybackViewportClient::HandleFocusRequest(FBox BoxToFocusOn)
 {
-	if (AActor* FocusedActor = Cast<AActor>(FocusedObject))
-	{
-		FocusViewportOnBox(FocusedActor->GetComponentsBoundingBox(false));
-	}
+	FocusViewportOnBox(BoxToFocusOn);
 }
 
 void FChaosVDPlaybackViewportClient::HandleActorMoving(AActor* MovedActor) const
diff --git a/Engine/Plugins/ChaosVD/Source/ChaosVD/Private/ChaosVDPlaybackViewportClient.h b/Engine/Plugins/ChaosVD/Source/ChaosVD/Private/ChaosVDPlaybackViewportClient.h
--- a/Engine/Plugins/ChaosVD/Source/ChaosVD/Private/ChaosVDPlaybackViewportClient.h
+++ b/Engine/Plugins/ChaosVD/Source/ChaosVD/Private/ChaosVDPlaybackViewportClient.h
@@ -42,10 +42,10 @@ public:
 
 private:
 
-	void HandleObjectFocused(UObject* FocusedObject);
+	void HandleFocusRequest(FBox BoxToFocusOn);
 	void HandleActorMoving(AActor* MovedActor) const;
 
-	FDelegateHandle ObjectFocusedDelegateHandle;
+	FDelegateHandle FocusRequestDelegateHandle;
 	UWorld* CVDWorld;
 	TWeakPtr<FChaosVDScene> CVDScene;
 
diff --git a/Engine/Plugins/ChaosVD/Source/ChaosVD/Private/ChaosVDScene.h b/Engine/Plugins/ChaosVD/Source/ChaosVD/Private/ChaosVDScene.h
--- a/Engine/Plugins/ChaosVD/Source/ChaosVD/Private/ChaosVDScene.h
+++ b/Engine/Plugins/ChaosVD/Source/ChaosVD/Private/ChaosVDScene.h
@@ -41,6 +41,7 @@ typedef TMap<int32, AChaosVDSolverInfoActor*> FChaosVDSolverInfoByIDMap;
 DECLARE_MULTICAST_DELEGATE(FChaosVDSceneUpdatedDelegate)
 DECLARE_MULTICAST_DELEGATE_OneParam(FChaosVDActorUpdatedDelegate, AChaosVDParticleActor*)
 DECLARE_MULTICAST_DELEGATE_OneParam(FChaosVDOnObjectSelectedDelegate, UObject*)
+DECLARE_MULTICAST_DELEGATE_OneParam(FChaosVDFocusRequestDelegate, FBox)
 DECLARE_MULTICAST_DELEGATE_OneParam(FChaosVDSolverInfoActorCreatedDelegate, AChaosVDSolverInfoActor*)
 
 DECLARE_MULTICAST_DELEGATE_TwoParams(FChaosVDSolverVisibilityChangedDelegate, int32 SolverID, bool bNewVisibility)
@@ -115,7 +116,7 @@ public:
 	USelection* GetObjectsSelectionObject() const { return ObjectSelection; }
 	
 	/** Event triggered when an object is focused in the scene (double click in the scene outliner)*/
-	FChaosVDOnObjectSelectedDelegate& OnObjectFocused() { return ObjectFocusedDelegate; }
+	FChaosVDFocusRequestDelegate& OnFocusRequest() { return FocusRequestDelegate; }
 
 	/** Returns a ptr to the particle actor representing the provided Particle ID
 	 * @param SolverID ID of the solver owning the Particle
@@ -197,7 +198,7 @@ private:
 
 	FChaosVDGeometryDataLoaded NewGeometryAvailableDelegate;
 
-	FChaosVDOnObjectSelectedDelegate ObjectFocusedDelegate;
+	FChaosVDFocusRequestDelegate FocusRequestDelegate;
 
 	/** Selection set object holding the current selection state */
 	TObjectPtr<UTypedElementSelectionSet> SelectionSet;
diff --git a/Engine/Plugins/ChaosVD/Source/ChaosVD/Private/ChaosVDTabsIDs.cpp b/Engine/Plugins/ChaosVD/Source/ChaosVD/Private/ChaosVDTabsIDs.cpp
--- a/Engine/Plugins/ChaosVD/Source/ChaosVD/Private/ChaosVDTabsIDs.cpp
+++ b/Engine/Plugins/ChaosVD/Source/ChaosVD/Private/ChaosVDTabsIDs.cpp
@@ -12,4 +12,5 @@ const FName FChaosVDTabID::SolversTrack = TEXT("SolversTrack");
 const FName FChaosVDTabID::StatusBar = TEXT("StatusBar");
 const FName FChaosVDTabID::CollisionDataDetails = TEXT("CollisionDataDetails");
 const FName FChaosVDTabID::SceneQueryDataDetails = TEXT("SceneQueryDataDetails");
-const FName FChaosVDTabID::ConstraintsInspector = TEXT("ConstraintsDetails");
\ No newline at end of file
+const FName FChaosVDTabID::ConstraintsInspector = TEXT("ConstraintsDetails");
+const FName FChaosVDTabID::SceneQueryBrowser = TEXT("SceneQueryBrowser");
\ No newline at end of file
diff --git a/Engine/Plugins/ChaosVD/Source/ChaosVD/Private/ChaosVDTabsIDs.h b/Engine/Plugins/ChaosVD/Source/ChaosVD/Private/ChaosVDTabsIDs.h
--- a/Engine/Plugins/ChaosVD/Source/ChaosVD/Private/ChaosVDTabsIDs.h
+++ b/Engine/Plugins/ChaosVD/Source/ChaosVD/Private/ChaosVDTabsIDs.h
@@ -18,4 +18,5 @@ public:
 	static const FName CollisionDataDetails;
 	static const FName SceneQueryDataDetails;
 	static const FName ConstraintsInspector;
+	static const FName SceneQueryBrowser;
 };
diff --git a/Engine/Plugins/ChaosVD/Source/ChaosVD/Private/ChaosVDWorldOutlinerMode.cpp b/Engine/Plugins/ChaosVD/Source/ChaosVD/Private/ChaosVDWorldOutlinerMode.cpp
--- a/Engine/Plugins/ChaosVD/Source/ChaosVD/Private/ChaosVDWorldOutlinerMode.cpp
+++ b/Engine/Plugins/ChaosVD/Source/ChaosVD/Private/ChaosVDWorldOutlinerMode.cpp
@@ -186,7 +186,7 @@ void FChaosVDWorldOutlinerMode::OnItemDoubleClick(FSceneOutlinerTreeItemPtr Item
 	{
 		if (AActor* Actor = ActorItem->Actor.Get())
 		{
-			ScenePtr->OnObjectFocused().Broadcast(Actor);
+			ScenePtr->OnFocusRequest().Broadcast(Actor->GetComponentsBoundingBox(false));
 		}
 	}
 }
diff --git a/Engine/Plugins/ChaosVD/Source/ChaosVD/Private/Components/ChaosVDSceneQueryDataComponent.cpp b/Engine/Plugins/ChaosVD/Source/ChaosVD/Private/Components/ChaosVDSceneQueryDataComponent.cpp
--- a/Engine/Plugins/ChaosVD/Source/ChaosVD/Private/Components/ChaosVDSceneQueryDataComponent.cpp
+++ b/Engine/Plugins/ChaosVD/Source/ChaosVD/Private/Components/ChaosVDSceneQueryDataComponent.cpp
@@ -3,9 +3,9 @@
 #include "Components/ChaosVDSceneQueryDataComponent.h"
 
 #include "ChaosVDRecording.h"
+#include "ChaosVDScene.h"
 #include "Actors/ChaosVDSolverInfoActor.h"
 
-
 UChaosVDSceneQueryDataComponent::UChaosVDSceneQueryDataComponent()
 {
 	PrimaryComponentTick.bCanEverTick = false;
@@ -14,6 +14,21 @@ UChaosVDSceneQueryDataComponent::UChaosVDSceneQueryDataComponent()
 	bNavigationRelevant = false;
 }
 
+void UChaosVDSceneQueryDataComponent::ClearSceneQuerySelection()
+{
+	const AChaosVDSolverInfoActor* SolverInfo = Cast<AChaosVDSolverInfoActor>(GetOwner());
+	const TSharedPtr<FChaosVDScene> CVDScene = SolverInfo ? SolverInfo->GetScene().Pin() : nullptr;
+
+	if (TSharedPtr<FChaosVDSolverDataSelection> SolverDataSelectionObject = CVDScene ? CVDScene->GetSolverDataSelectionObject().Pin() : nullptr)
+	{
+		TSharedPtr<FChaosVDSolverDataSelectionHandle> SelectionHandle = SolverDataSelectionObject->GetCurrentSelectionHandle();
+		if (SelectionHandle && SelectionHandle->IsA<FChaosVDQueryDataWrapper>())
+		{
+			SolverDataSelectionObject->SelectData(nullptr);
+		}
+	}
+}
+
 void UChaosVDSceneQueryDataComponent::UpdateFromNewGameFrameData(const FChaosVDGameFrameData& InGameFrameData)
 {
 
@@ -41,6 +56,8 @@ void UChaosVDSceneQueryDataComponent::UpdateFromNewGameFrameData(const FChaosVDG
 			RecordedQueries.Add(QueryData);
 		}
 	}
+	
+	ClearSceneQuerySelection();
 }
 
 TConstArrayView<TSharedPtr<FChaosVDQueryDataWrapper>> UChaosVDSceneQueryDataComponent::GetQueriesByType(EChaosVDSceneQueryType Type) const
diff --git a/Engine/Plugins/ChaosVD/Source/ChaosVD/Private/Components/ChaosVDSceneQueryDataComponent.h b/Engine/Plugins/ChaosVD/Source/ChaosVD/Private/Components/ChaosVDSceneQueryDataComponent.h
--- a/Engine/Plugins/ChaosVD/Source/ChaosVD/Private/Components/ChaosVDSceneQueryDataComponent.h
+++ b/Engine/Plugins/ChaosVD/Source/ChaosVD/Private/Components/ChaosVDSceneQueryDataComponent.h
@@ -18,6 +18,7 @@ class UChaosVDSceneQueryDataComponent : public UChaosVDSolverDataComponent
 
 public:
 	UChaosVDSceneQueryDataComponent();
+	void ClearSceneQuerySelection();
 
 	virtual void UpdateFromNewGameFrameData(const FChaosVDGameFrameData& InGameFrameData) override;
 
diff --git a/Engine/Plugins/ChaosVD/Source/ChaosVD/Private/Settings/ChaosVDSceneQueryVisualizationSettings.h b/Engine/Plugins/ChaosVD/Source/ChaosVD/Private/Settings/ChaosVDSceneQueryVisualizationSettings.h
--- a/Engine/Plugins/ChaosVD/Source/ChaosVD/Private/Settings/ChaosVDSceneQueryVisualizationSettings.h
+++ b/Engine/Plugins/ChaosVD/Source/ChaosVD/Private/Settings/ChaosVDSceneQueryVisualizationSettings.h
@@ -21,6 +21,17 @@ enum class EChaosVDSceneQueryVisualizationFlags: uint32
 };
 ENUM_CLASS_FLAGS(EChaosVDSceneQueryVisualizationFlags);
 
+/** Available scene query Visualization Modes */
+UENUM()
+enum class EChaosVDSQFrameVisualizationMode : uint8
+{
+	/** All the recorded scene queries for the current frame, that passes the visualization flag filter
+	 * will be shown */
+	AllEnabledQueries,
+	/** Scene queries will be shown one at the time, in the order they were recorded, controlled by the Scene Query browser timeline */
+	RecordingOrder
+};
+
 UCLASS(config=ChaosVD, PerObjectConfig)
 class UChaosVDSceneQueriesVisualizationSettings : public UChaosVDVisualizationSettingsObjectBase
 {
@@ -35,6 +46,9 @@ public:
 	UPROPERTY(EditAnywhere, Category=DebugDraw)
 	bool bShowText = false;
 
+	UPROPERTY()
+	EChaosVDSQFrameVisualizationMode CurrentVisualizationMode = EChaosVDSQFrameVisualizationMode::AllEnabledQueries;
+
 private:
 	/** Set of flags to enable/disable visualization of specific scene queries data as debug draw */
 	UPROPERTY(config, meta = (Bitmask, BitmaskEnum = "/Script/ChaosVD.EChaosVDSceneQueryVisualizationFlags"))
diff --git a/Engine/Plugins/ChaosVD/Source/ChaosVD/Private/TabSpawers/ChaosVDSceneQueryBrowserTab.cpp b/Engine/Plugins/ChaosVD/Source/ChaosVD/Private/TabSpawers/ChaosVDSceneQueryBrowserTab.cpp
new file mode 100644
--- /dev/null
+++ b/Engine/Plugins/ChaosVD/Source/ChaosVD/Private/TabSpawers/ChaosVDSceneQueryBrowserTab.cpp
@@ -0,0 +1,48 @@
+// Copyright Epic Games, Inc. All Rights Reserved.
+
+#include "ChaosVDSceneQueryBrowserTab.h"
+
+#include "ChaosVDStyle.h"
+#include "EditorModeManager.h"
+#include "Widgets/SChaosVDMainTab.h"
+#include "Widgets/SChaosVDSceneQueryBrowser.h"
+
+#define LOCTEXT_NAMESPACE "ChaosVisualDebugger"
+
+FChaosVDSceneQueryBrowserTab::~FChaosVDSceneQueryBrowserTab()
+{
+}
+
+TSharedRef<SDockTab> FChaosVDSceneQueryBrowserTab::HandleTabSpawnRequest(const FSpawnTabArgs& Args)
+{
+	TSharedRef<SDockTab> SceneQueryBrowserTab =
+	SNew(SDockTab)
+	.TabRole(ETabRole::MajorTab)
+	.Label(LOCTEXT("SceneQueryBrowserTab", "Scene Query Browser"))
+	.ToolTipText(LOCTEXT("SceneQueryBrowserTabTabTip", "Shows all recorded scene queries for the current frame, and allows you to select them or play then in order"));
+
+	if (const TSharedPtr<SChaosVDMainTab> MainTabPtr = OwningTabWidget.Pin())
+	{
+		SceneQueryBrowserTab->SetContent
+		(
+			SAssignNew(SceneQueryBrowser, SChaosVDSceneQueryBrowser, GetChaosVDScene(), MainTabPtr->GetEditorModeManager().AsWeak())
+		);
+	}
+	else
+	{
+		SceneQueryBrowserTab->SetContent(GenerateErrorWidget());
+	}
+
+	SceneQueryBrowserTab->SetTabIcon(FChaosVDStyle::Get().GetBrush("SceneQueriesInspectorIcon"));
+
+	HandleTabSpawned(SceneQueryBrowserTab);
+
+	return SceneQueryBrowserTab;
+}
+
+void FChaosVDSceneQueryBrowserTab::HandleTabClosed(TSharedRef<SDockTab> InTabClosed)
+{
+	FChaosVDTabSpawnerBase::HandleTabClosed(InTabClosed);
+}
+
+#undef LOCTEXT_NAMESPACE
diff --git a/Engine/Plugins/ChaosVD/Source/ChaosVD/Private/TabSpawers/ChaosVDSceneQueryBrowserTab.h b/Engine/Plugins/ChaosVD/Source/ChaosVD/Private/TabSpawers/ChaosVDSceneQueryBrowserTab.h
new file mode 100644
--- /dev/null
+++ b/Engine/Plugins/ChaosVD/Source/ChaosVD/Private/TabSpawers/ChaosVDSceneQueryBrowserTab.h
@@ -0,0 +1,27 @@
+// Copyright Epic Games, Inc. All Rights Reserved.
+
+#pragma once
+
+#include "ChaosVDTabSpawnerBase.h"
+#include "Templates/SharedPointer.h"
+
+class SChaosVDSceneQueryBrowser;
+
+/** Spawns and handles and instance for the visual debugger Scene Query browser tab */
+class FChaosVDSceneQueryBrowserTab : public FChaosVDTabSpawnerBase
+{
+public:
+	FChaosVDSceneQueryBrowserTab(const FName& InTabID, const TSharedPtr<FTabManager>& InTabManager, const TWeakPtr<SChaosVDMainTab>& InOwningTabWidget)
+		: FChaosVDTabSpawnerBase(InTabID, InTabManager, InOwningTabWidget)
+	{
+	}
+
+	virtual ~FChaosVDSceneQueryBrowserTab() override;
+	virtual TSharedRef<SDockTab> HandleTabSpawnRequest(const FSpawnTabArgs& Args) override;
+	virtual void HandleTabClosed(TSharedRef<SDockTab> InTabClosed) override;
+
+	TWeakPtr<SChaosVDSceneQueryBrowser> GetSceneQueryDataInspectorInstance() const { return SceneQueryBrowser; }
+
+protected:
+	TSharedPtr<SChaosVDSceneQueryBrowser> SceneQueryBrowser;
+};
diff --git a/Engine/Plugins/ChaosVD/Source/ChaosVD/Private/Utils/ChaosVDUserInterfaceUtils.h b/Engine/Plugins/ChaosVD/Source/ChaosVD/Private/Utils/ChaosVDUserInterfaceUtils.h
--- a/Engine/Plugins/ChaosVD/Source/ChaosVD/Private/Utils/ChaosVDUserInterfaceUtils.h
+++ b/Engine/Plugins/ChaosVD/Source/ChaosVD/Private/Utils/ChaosVDUserInterfaceUtils.h
@@ -4,6 +4,9 @@
 #include "ChaosVDSettingsManager.h"
 #include "IStructureDetailsView.h"
 #include "SEnumCombo.h"
+#include "ToolMenu.h"
+#include "ToolMenuSection.h"
+#include "Widgets/SChaosVDEnumFlagsMenu.h"
 
 class UToolMenu;
 
@@ -65,4 +68,25 @@ namespace Chaos::VisualDebugger::Utils
 
 		InDetailsView->SetStructureData(StructDataView);
 	}
+
+	template <typename ObjectSettingsType, typename VisualizationFlagsType>
+	void CreateVisualizationOptionsMenuSections(UToolMenu* Menu, FName SectionName, const FText& InSectionLabel, const FText& InFlagsMenuLabel,  const FText& InFlagsMenuTooltip, FSlateIcon FlagsMenuIcon,  const FText& InSettingsMenuLabel, const FText& InSettingsMenuTooltip)
+	{
+		FToolMenuSection& Section = Menu->AddSection(SectionName, InSectionLabel);
+		
+		Section.AddSubMenu(FName(InSectionLabel.ToString()), InFlagsMenuLabel, InFlagsMenuTooltip, FNewToolMenuDelegate::CreateLambda([](UToolMenu* Menu)
+						   {
+							   TSharedRef<SWidget> VisualizationFlagsWidget = SNew(SChaosVDEnumFlagsMenu<VisualizationFlagsType>)
+								   .CurrentValue_Static(&ObjectSettingsType::GetDataVisualizationFlags)
+								   .OnEnumSelectionChanged_Static(&ObjectSettingsType::SetDataVisualizationFlags);
+			
+							   FToolMenuEntry FlagsMenuEntry = FToolMenuEntry::InitWidget("VisualizationFlags", VisualizationFlagsWidget,FText::GetEmpty());
+							   Menu->AddMenuEntry(NAME_None, FlagsMenuEntry);
+						   }),
+						   false, FlagsMenuIcon);
+
+		using namespace Chaos::VisualDebugger::Utils;
+		Section.AddSubMenu(FName(InSettingsMenuLabel.ToString()), InSettingsMenuLabel, InSettingsMenuTooltip, FNewToolMenuDelegate::CreateStatic(&CreateMenuEntryForSettingsObject<ObjectSettingsType>, EChaosVDSaveSettingsOptions::ShowResetButton),
+						   false, FSlateIcon(FAppStyle::Get().GetStyleSetName(), TEXT("Icons.Toolbar.Settings")));
+	}
 }
diff --git a/Engine/Plugins/ChaosVD/Source/ChaosVD/Private/Visualizers/ChaosVDComponentVisualizerBase.h b/Engine/Plugins/ChaosVD/Source/ChaosVD/Private/Visualizers/ChaosVDComponentVisualizerBase.h
--- a/Engine/Plugins/ChaosVD/Source/ChaosVD/Private/Visualizers/ChaosVDComponentVisualizerBase.h
+++ b/Engine/Plugins/ChaosVD/Source/ChaosVD/Private/Visualizers/ChaosVDComponentVisualizerBase.h
@@ -81,21 +81,7 @@ void FChaosVDComponentVisualizerBase::CreateGenericVisualizerMenu(FName MenuToEx
 
 	if (UToolMenu* Menu = ToolMenus->ExtendMenu(MenuToExtend))
 	{
-		FToolMenuSection& Section = Menu->AddSection(SectionName, InSectionLabel);
-		
-		Section.AddSubMenu(FName(InSectionLabel.ToString()), InFlagsMenuLabel, InFlagsMenuTooltip, FNewToolMenuDelegate::CreateLambda([](UToolMenu* Menu)
-						   {
-							   TSharedRef<SWidget> VisualizationFlagsWidget = SNew(SChaosVDEnumFlagsMenu<VisualizationFlagsType>)
-								   .CurrentValue_Static(&ObjectSettingsType::GetDataVisualizationFlags)
-								   .OnEnumSelectionChanged_Static(&ObjectSettingsType::SetDataVisualizationFlags);
-			
-							   FToolMenuEntry FlagsMenuEntry = FToolMenuEntry::InitWidget("VisualizationFlags", VisualizationFlagsWidget,FText::GetEmpty());
-							   Menu->AddMenuEntry(NAME_None, FlagsMenuEntry);
-						   }),
-						   false, FlagsMenuIcon);
-
 		using namespace Chaos::VisualDebugger::Utils;
-		Section.AddSubMenu(FName(InSettingsMenuLabel.ToString()), InSettingsMenuLabel, InSettingsMenuTooltip, FNewToolMenuDelegate::CreateStatic(&CreateMenuEntryForSettingsObject<ObjectSettingsType>, EChaosVDSaveSettingsOptions::ShowResetButton),
-						   false, FSlateIcon(FAppStyle::Get().GetStyleSetName(), TEXT("Icons.Toolbar.Settings")));
+		CreateVisualizationOptionsMenuSections<ObjectSettingsType, VisualizationFlagsType>(Menu, SectionName, InSectionLabel, InFlagsMenuLabel, InFlagsMenuTooltip, FlagsMenuIcon, InSettingsMenuLabel, InSettingsMenuTooltip);
 	}
 }
diff --git a/Engine/Plugins/ChaosVD/Source/ChaosVD/Private/Visualizers/ChaosVDDebugDrawUtils.cpp b/Engine/Plugins/ChaosVD/Source/ChaosVD/Private/Visualizers/ChaosVDDebugDrawUtils.cpp
--- a/Engine/Plugins/ChaosVD/Source/ChaosVD/Private/Visualizers/ChaosVDDebugDrawUtils.cpp
+++ b/Engine/Plugins/ChaosVD/Source/ChaosVD/Private/Visualizers/ChaosVDDebugDrawUtils.cpp
@@ -427,4 +427,27 @@ FString Chaos::VisualDebugger::Utils::GenerateDebugTextForVector(const FVector&
 	return FString::Format(TEXT("{5} : {0} {4} \n	|-- X : {1} {4} \n	|-- Y : {2} {4} \n	|-- Z : {3} {4}"), {InVector.Size(), InVector.X, InVector.Y, InVector.Z, InVectorUnits, VectorName });
 }
 
+FBox Chaos::VisualDebugger::Utils::CalculateSceneQueryShapeBounds(const TSharedRef<FChaosVDQueryDataWrapper>& InSceneQueryData, const TSharedRef<FChaosVDRecording> InRecordedData)
+{
+	FBoxSphereBounds::Builder BoundsBuilder;
+	const FConstImplicitObjectPtr* InputShapePtrPtr = InRecordedData->GetGeometryMap().Find(InSceneQueryData->InputGeometryKey);
+	const FConstImplicitObjectPtr InputShapePtr = InputShapePtrPtr ? *InputShapePtrPtr : nullptr;
+	
+	if (InputShapePtr && InputShapePtr->HasBoundingBox())
+	{
+		FAABB3 StartBounds = InputShapePtr->CalculateTransformedBounds(FRigidTransform3(InSceneQueryData->StartLocation, InSceneQueryData->GeometryOrientation));
+		FAABB3 EndBounds = InputShapePtr->CalculateTransformedBounds(FRigidTransform3(InSceneQueryData->EndLocation, InSceneQueryData->GeometryOrientation));
+
+		BoundsBuilder += FBox(StartBounds.Min(), StartBounds.Max());
+		BoundsBuilder += FBox(EndBounds.Min(), EndBounds.Max());
+	}
+	else
+	{
+		BoundsBuilder+= InSceneQueryData->EndLocation;
+		BoundsBuilder+= InSceneQueryData->StartLocation;
+	}
+	
+	return FBoxSphereBounds(BoundsBuilder).GetBox();
+}
+
 #undef LOCTEXT_NAMESPACE
diff --git a/Engine/Plugins/ChaosVD/Source/ChaosVD/Private/Visualizers/ChaosVDDebugDrawUtils.h b/Engine/Plugins/ChaosVD/Source/ChaosVD/Private/Visualizers/ChaosVDDebugDrawUtils.h
--- a/Engine/Plugins/ChaosVD/Source/ChaosVD/Private/Visualizers/ChaosVDDebugDrawUtils.h
+++ b/Engine/Plugins/ChaosVD/Source/ChaosVD/Private/Visualizers/ChaosVDDebugDrawUtils.h
@@ -6,8 +6,11 @@
 #include "Containers/Queue.h"
 #include "SceneManagement.h"
 
-class FChaosVDPlaybackViewportClient;
 class FChaosVDGeometryBuilder;
+class FChaosVDPlaybackViewportClient;
+
+struct FChaosVDQueryDataWrapper;
+struct FChaosVDRecording;
 
 enum class EChaosVDDebugDrawTextLocationMode
 {
@@ -63,4 +66,6 @@ private:
 namespace Chaos::VisualDebugger::Utils
 {
 	FString GenerateDebugTextForVector(const FVector& InVector, const FString& VectorName, const FString& InVectorUnits);
+
+	FBox CalculateSceneQueryShapeBounds(const TSharedRef<FChaosVDQueryDataWrapper>& InSceneQueryData, const TSharedRef<FChaosVDRecording> InRecordedData);
 }
diff --git a/Engine/Plugins/ChaosVD/Source/ChaosVD/Private/Visualizers/ChaosVDSceneQueryDataComponentVisualizer.cpp b/Engine/Plugins/ChaosVD/Source/ChaosVD/Private/Visualizers/ChaosVDSceneQueryDataComponentVisualizer.cpp
--- a/Engine/Plugins/ChaosVD/Source/ChaosVD/Private/Visualizers/ChaosVDSceneQueryDataComponentVisualizer.cpp
+++ b/Engine/Plugins/ChaosVD/Source/ChaosVD/Private/Visualizers/ChaosVDSceneQueryDataComponentVisualizer.cpp
@@ -93,14 +93,19 @@ void FChaosVDSceneQueryDataComponentVisualizer::DrawVisualization(const UActorCo
 		return;
 	}
 
+	UChaosVDSceneQueriesVisualizationSettings* Settings = FChaosVDSettingsManager::Get().GetSettingsObject<UChaosVDSceneQueriesVisualizationSettings>();
+	if (!Settings)
+	{
+		return;
+	}
+
 	FChaosVDSceneQueryVisualizationDataContext VisualizationContext;
 	VisualizationContext.CVDScene = SolverInfoActor->GetScene();
 	VisualizationContext.SpaceTransform = FTransform::Identity;
 	VisualizationContext.GeometryGenerator = GeometryGenerator;
 	VisualizationContext.SolverDataSelectionObject = SolverDataSelectionObject;
-
 	VisualizationContext.VisualizationFlags = static_cast<uint32>(UChaosVDSceneQueriesVisualizationSettings::GetDataVisualizationFlags());
-	VisualizationContext.DebugDrawSettings = FChaosVDSettingsManager::Get().GetSettingsObject<UChaosVDSceneQueriesVisualizationSettings>();
+	VisualizationContext.DebugDrawSettings = Settings;
 
 	if (EnumHasAnyFlags(EChaosVDSceneQueryVisualizationFlags::EnableDraw, static_cast<EChaosVDSceneQueryVisualizationFlags>(VisualizationContext.VisualizationFlags)))
 	{
@@ -108,8 +113,10 @@ void FChaosVDSceneQueryDataComponentVisualizer::DrawVisualization(const UActorCo
 
 		TSharedPtr<FChaosVDSolverDataSelectionHandle> SelectionHandle = SolverDataSelectionObject->GetCurrentSelectionHandle();
 		const bool bHasSelectedQuery = SelectionHandle && SelectionHandle->IsA<FChaosVDQueryDataWrapper>();
+		const bool bOnlyDrawSelected = Settings->CurrentVisualizationMode == EChaosVDSQFrameVisualizationMode::RecordingOrder ||
+										EnumHasAnyFlags(EChaosVDSceneQueryVisualizationFlags::OnlyDrawSelectedQuery, static_cast<EChaosVDSceneQueryVisualizationFlags>(VisualizationContext.VisualizationFlags));
 
-		if (bHasSelectedQuery && EnumHasAnyFlags(EChaosVDSceneQueryVisualizationFlags::OnlyDrawSelectedQuery, static_cast<EChaosVDSceneQueryVisualizationFlags>(VisualizationContext.VisualizationFlags)))
+		if (bHasSelectedQuery && bOnlyDrawSelected)
 		{
 			DrawSceneQuery(Component, View, PDI, CVDScene, CVDRecording, VisualizationContext, SelectionHandle->GetDataAsShared<FChaosVDQueryDataWrapper>());
 		}
@@ -332,11 +339,8 @@ void FChaosVDSceneQueryDataComponentVisualizer::DrawSceneQuery(const UActorCompo
 		return;
 	}
 
-	//TODO: Should we try to calculate actual bounds?
-	constexpr float MinVisibleRadius = 100.0f;
-	const float QueryHalfDistance = (Query->EndLocation - Query->StartLocation).Size() * 0.5;
-	const float VisibleRadius = FMath::Max(QueryHalfDistance, MinVisibleRadius);
-	if (!View->ViewFrustum.IntersectSphere(Query->StartLocation, VisibleRadius))
+	FBox QueryBounds = Chaos::VisualDebugger::Utils::CalculateSceneQueryShapeBounds(Query.ToSharedRef(), CVDRecording.ToSharedRef());
+	if (!View->ViewFrustum.IntersectBox(QueryBounds.GetCenter(), QueryBounds.GetExtent()))
 	{
 		// If this query location is not even visible, just ignore it.
 		return;
diff --git a/Engine/Plugins/ChaosVD/Source/ChaosVD/Private/Widgets/SChaosVDMainTab.cpp b/Engine/Plugins/ChaosVD/Source/ChaosVD/Private/Widgets/SChaosVDMainTab.cpp
--- a/Engine/Plugins/ChaosVD/Source/ChaosVD/Private/Widgets/SChaosVDMainTab.cpp
+++ b/Engine/Plugins/ChaosVD/Source/ChaosVD/Private/Widgets/SChaosVDMainTab.cpp
@@ -24,6 +24,7 @@
 #include "Framework/Application/SWindowTitleBar.h"
 #include "Framework/Commands/UICommandList.h"
 #include "Framework/MultiBox/MultiBoxBuilder.h"
+#include "HAL/PlatformApplicationMisc.h"
 #include "IDesktopPlatform.h"
 #include "Misc/MessageDialog.h"
 #include "StatusBarSubsystem.h"
@@ -37,6 +38,7 @@
 #include "Components/ChaosVDSolverJointConstraintDataComponent.h"
 #include "Styling/StyleColors.h"
 #include "Styling/ToolBarStyle.h"
+#include "TabSpawers/ChaosVDSceneQueryBrowserTab.h"
 #include "Trace/ChaosVDTraceManager.h"
 #include "Visualizers/ChaosVDCharacterGroundConstraintsDataComponentVisualizer.h"
 #include "Visualizers/ChaosVDGenericDebugDrawDataComponentVisualizer.h"
@@ -81,6 +83,7 @@ void SChaosVDMainTab::Construct(const FArguments& InArgs, TSharedPtr<FChaosVDEng
 	RegisterTabSpawner<FChaosVDCollisionDataDetailsTab>(FChaosVDTabID::CollisionDataDetails);
 	RegisterTabSpawner<FChaosVDSceneQueryDataInspectorTab>(FChaosVDTabID::SceneQueryDataDetails);
 	RegisterTabSpawner<FChaosVDConstraintDataInspectorTab>(FChaosVDTabID::ConstraintsInspector);
+	RegisterTabSpawner<FChaosVDSceneQueryBrowserTab>(FChaosVDTabID::SceneQueryBrowser);
 
 	StatusBarID = FName(FChaosVDTabID::StatusBar.ToString() + InChaosVDEngine->GetInstanceGuid().ToString());
 	
@@ -244,6 +247,8 @@ void SChaosVDMainTab::HandleTabDestroyed(TSharedRef<SDockTab> Tab, FName TabID)
 
 TSharedRef<FTabManager::FLayout> SChaosVDMainTab::GenerateMainLayout()
 {
+	const float DPIScaleFactor = FPlatformApplicationMisc::GetDPIScaleFactorAtPoint(10.0f, 10.0f);
+
 	return FTabManager::NewLayout("ChaosVisualDebugger_Layout")
 		->AddArea
 		(
@@ -291,7 +296,16 @@ TSharedRef<FTabManager::FLayout> SChaosVDMainTab::GenerateMainLayout()
 					->AddTab(FChaosVDTabID::ConstraintsInspector, ETabState::ClosedTab)
 				)
 			)
-		);
+		)
+	->AddArea(
+	FTabManager::NewArea(800.0f * DPIScaleFactor, 600.0f * DPIScaleFactor)
+	->SetOrientation(Orient_Vertical)
+	->Split
+	(
+		FTabManager::NewStack()
+		->SetSizeCoefficient(1.0f)
+		->AddTab(FChaosVDTabID::SceneQueryBrowser, ETabState::ClosedTab)
+	));
 }
 
 void SChaosVDMainTab::GenerateMainWindowMenu()
@@ -581,6 +595,47 @@ void SChaosVDMainTab::RegisterMainTabMenu()
 	}));
 	
 	Section.AddSeparator(NAME_None);
+
+	//TODO : This button should not be added to the toolbar here. Ideally it should be added from the SceneQueryComponent Visualizer, but we have two issues :
+	// 1- The recording control buttons are still implemented as a widget we instantiate alongside the tool bar, that needs to be moved to be a properly
+	// registered menu entry that is part of the toolbar.
+	// 2- We need to ensure the main toolbar is created and ready to use before we allow other system to register into it.
+	// Jira for tracking UE-221454
+
+	Section.AddDynamicEntry("DataBrowsers", FNewToolMenuSectionDelegate::CreateLambda([](FToolMenuSection& InSection)
+	{
+		const UChaosVDMainToolbarMenuContext* Context = InSection.FindContext<UChaosVDMainToolbarMenuContext>();
+		TSharedPtr<SChaosVDMainTab> MainTabPtr = Context->MainTab.Pin();
+		if (!MainTabPtr)
+		{
+			return;
+		}
+
+		FOnClicked OnClickedDelegate = FOnClicked::CreateLambda([WeakTab = StaticCastWeakPtr<SChaosVDMainTab>(MainTabPtr->AsWeak())]()
+		{
+			if (TSharedPtr<SChaosVDMainTab> TabPtr = WeakTab.Pin())
+			{
+				TabPtr->TabManager->TryInvokeTab(FChaosVDTabID::SceneQueryBrowser);
+			}
+
+			return FReply::Handled();
+		});
+
+		TSharedRef<SButton> ConnectToSessionButton = MainTabPtr->CreateSimpleButton(
+														[](){ return LOCTEXT("SceneQueryBrowserButton", "Scene Query Browser"); },
+														[](){ return LOCTEXT("SceneQueryBrowserButtonTooltip", "Opens the Scene Query Browser window, which shows all the available scene queries in the current frame."); },
+														FChaosVDStyle::Get().GetBrush("SceneQueriesInspectorIcon"),
+														Context, MoveTemp(OnClickedDelegate));
+
+		InSection.AddEntry(
+			FToolMenuEntry::InitWidget(
+				"SceneQueryBrowser",
+				ConnectToSessionButton,
+				FText::GetEmpty(),
+				true,
+				false
+			));
+	}));
 }
 
 void SChaosVDMainTab::BrowseLiveSessionsFromTraceStore() const
diff --git a/Engine/Plugins/ChaosVD/Source/ChaosVD/Private/Widgets/SChaosVDSceneQueryBrowser.cpp b/Engine/Plugins/ChaosVD/Source/ChaosVD/Private/Widgets/SChaosVDSceneQueryBrowser.cpp
new file mode 100644
--- /dev/null
+++ b/Engine/Plugins/ChaosVD/Source/ChaosVD/Private/Widgets/SChaosVDSceneQueryBrowser.cpp
@@ -0,0 +1,834 @@
+// Copyright Epic Games, Inc. All Rights Reserved.
+
+#include "SChaosVDSceneQueryBrowser.h"
+
+#include "ChaosVDEngine.h"
+#include "Actors/ChaosVDSolverInfoActor.h"
+#include "ChaosVDScene.h"
+#include "ChaosVDStyle.h"
+#include "ChaosVDTabsIDs.h"
+#include "Components/ChaosVDSceneQueryDataComponent.h"
+#include "EditorModeManager.h"
+#include "EditorViewportClient.h"
+#include "SChaosVDGameFramesPlaybackControls.h"
+#include "SChaosVDMainTab.h"
+#include "SChaosVDSceneQueryTree.h"
+#include "SChaosVDWarningMessageBox.h"
+#include "Settings/ChaosVDSceneQueryVisualizationSettings.h"
+#include "ToolMenu.h"
+#include "ToolMenus.h"
+#include "ToolMenuEntry.h"
+#include "ToolMenuSection.h"
+#include "Utils/ChaosVDUserInterfaceUtils.h"
+#include "Visualizers/ChaosVDDebugDrawUtils.h"
+#include "Widgets/Input/SSearchBox.h"
+
+#define LOCTEXT_NAMESPACE "ChaosVisualDebugger"
+
+SChaosVDSceneQueryBrowser::~SChaosVDSceneQueryBrowser()
+{
+	UnregisterSceneEvents();
+
+	if (UChaosVDSceneQueriesVisualizationSettings* Settings = FChaosVDSettingsManager::Get().GetSettingsObject<UChaosVDSceneQueriesVisualizationSettings>())
+	{
+		Settings->OnSettingsChanged().RemoveAll(this);
+	}
+}
+
+void SChaosVDSceneQueryBrowser::Construct(const FArguments& InArgs, TWeakPtr<FChaosVDScene> Scene,  TWeakPtr<FEditorModeTools> EditorModeTools)
+{
+	SceneWeakPtr = Scene;
+	EditorModeToolsWeakPtr = EditorModeTools;
+
+	FilteredCachedTreeItems = MakeShared<TArray<TSharedPtr<FChaosVDSceneQueryTreeItem>>>();
+
+	//TODO : At this point we should extract these values toa  style file other widgets can access so we have a coherent style between widgets
+	constexpr float NoPadding = 0.0f;
+	constexpr float MainContentBoxHorizontalPadding = 2.0f;
+	constexpr float MainContentBoxVerticalPadding = 5.0f;
+	constexpr float StatusBarSlotVerticalPadding = 1.0f;
+	constexpr float StatusBarInnerVerticalPadding = 9.0f;
+	constexpr float StatusBarInnerHorizontalPadding = 14.0f;
+
+	RegisterMainToolbarMenu();
+	
+	RegisterSceneEvents();
+	
+	ChildSlot
+	[
+		SNew(SVerticalBox)
+		+SVerticalBox::Slot()
+		.AutoHeight()
+		.Padding(NoPadding)
+		[
+			GenerateMainToolbarWidget()
+		]
+		+SVerticalBox::Slot()
+		.Padding(MainContentBoxHorizontalPadding, MainContentBoxVerticalPadding, MainContentBoxHorizontalPadding, NoPadding)
+		.AutoHeight()
+		[
+			SNew(SChaosVDWarningMessageBox)
+			.Visibility(this, &SChaosVDSceneQueryBrowser::GetUpdatesPausedMessageVisibility)
+			.WarningText(LOCTEXT("SceneQueryBrowserDataNoAvailableMessage", "Browser data updates disabled during playback..."))
+		]
+		+SVerticalBox::Slot()
+		.Padding(MainContentBoxHorizontalPadding, MainContentBoxVerticalPadding, MainContentBoxHorizontalPadding, NoPadding)
+		.FillHeight(1.0f)
+		[
+			SAssignNew(SceneQueryTreeWidget, SChaosVDSceneQueryTree)
+			.IsEnabled(this, &SChaosVDSceneQueryBrowser::GetQueryTreeWidgetEnabled)
+			.OnItemSelected(this, &SChaosVDSceneQueryBrowser::HandleTreeItemSelected)
+			.OnItemFocused(this, &SChaosVDSceneQueryBrowser::HandleTreeItemFocused)
+		]
+		+SVerticalBox::Slot()
+		.Padding(NoPadding, StatusBarSlotVerticalPadding, NoPadding, StatusBarSlotVerticalPadding)
+		.AutoHeight()
+		[
+			SNew(SBorder)
+			.BorderImage(FAppStyle::Get().GetBrush("Brushes.Header") )
+			.VAlign(VAlign_Center)
+			.HAlign(HAlign_Left)
+			.Padding(FMargin(StatusBarInnerHorizontalPadding, StatusBarInnerVerticalPadding))
+			[
+				SNew(STextBlock)
+				.Text(this, &SChaosVDSceneQueryBrowser::GetFilterStatusText)
+				.ColorAndOpacity(this, &SChaosVDSceneQueryBrowser::GetFilterStatusTextColor)
+			]
+		]
+		+SVerticalBox::Slot()
+		.AutoHeight()
+		[
+			GenerateQueriesPlaybackControls()
+		]
+	];
+	
+	if (UChaosVDSceneQueriesVisualizationSettings* Settings = FChaosVDSettingsManager::Get().GetSettingsObject<UChaosVDSceneQueriesVisualizationSettings>())
+	{
+		Settings->OnSettingsChanged().AddSP(this, &SChaosVDSceneQueryBrowser::HandleSettingsChanged);
+
+		HandleSettingsChanged(Settings);
+	}
+}
+
+void SChaosVDSceneQueryBrowser::RegisterSceneEvents()
+{
+	if (const TSharedPtr<FChaosVDScene> ScenePtr = SceneWeakPtr.Pin())
+	{
+		ScenePtr->OnSolverVisibilityUpdated().AddRaw(this, &SChaosVDSceneQueryBrowser::HandleSolverVisibilityChanged);
+
+		ScenePtr->OnSceneUpdated().AddRaw(this, &SChaosVDSceneQueryBrowser::HandleSceneUpdated);
+		if (TSharedPtr<FChaosVDSolverDataSelection> SelectionObject = ScenePtr->GetSolverDataSelectionObject().Pin())
+		{
+			SelectionObject->GetDataSelectionChangedDelegate().AddRaw(this, &SChaosVDSceneQueryBrowser::HandleExternalSelectionEvent);
+		}
+	}
+}
+
+void SChaosVDSceneQueryBrowser::UnregisterSceneEvents()
+{
+	if (const TSharedPtr<FChaosVDScene> ScenePtr = SceneWeakPtr.Pin())
+	{
+		ScenePtr->OnSceneUpdated().RemoveAll(this);
+		ScenePtr->OnSolverVisibilityUpdated().RemoveAll(this);
+
+		if (TSharedPtr<FChaosVDSolverDataSelection> SelectionObject = ScenePtr->GetSolverDataSelectionObject().Pin())
+		{
+			SelectionObject->GetDataSelectionChangedDelegate().RemoveAll(this);
+		}
+	}
+}
+
+void SChaosVDSceneQueryBrowser::Tick(const FGeometry& AllottedGeometry, const double InCurrentTime, const float InDeltaTime)
+{
+	if (!bIsUpToDate && CanUpdate())
+	{
+		UpdateBrowserContents();
+	}
+
+	if (bNeedsToUpdateSettings)
+	{
+		ApplySettingsChange();
+	}
+}
+
+void SChaosVDSceneQueryBrowser::SelectSceneQuery(const TSharedPtr<const FChaosVDQueryDataWrapper>& InQuery, ESelectInfo::Type Type)
+{
+	if (TSharedPtr<FChaosVDSceneQueryTreeItem>* ItemFound = CachedTreeItemsByID.Find(InQuery->ID))
+	{
+		SelectSceneQuery(*ItemFound, Type);
+	}
+}
+
+void SChaosVDSceneQueryBrowser::SelectSceneQuery(const TSharedPtr<FChaosVDSceneQueryTreeItem>& SceneQueryTreeItem, ESelectInfo::Type Type)
+{
+	TSharedPtr<SChaosVDSceneQueryTree> SceneQueryTreeWidgetPtr = SceneQueryTreeWidget.Pin();
+	if (!SceneQueryTreeWidgetPtr)
+	{
+		return;
+	}
+
+	if (SceneQueryTreeItem)
+	{
+		SceneQueryTreeWidgetPtr->SelectItem(SceneQueryTreeItem, Type);
+	}
+}
+
+void SChaosVDSceneQueryBrowser::HandleExternalSelectionEvent(const TSharedPtr<FChaosVDSolverDataSelectionHandle>& InDataSelectionHandle)
+{
+	TSharedPtr<FChaosVDQueryDataWrapper> SelectedQueryDataPtr = InDataSelectionHandle ? InDataSelectionHandle->GetDataAsShared<FChaosVDQueryDataWrapper>() : nullptr;
+	if (!SelectedQueryDataPtr)
+	{
+		return;
+	}
+
+	SelectSceneQuery(SelectedQueryDataPtr, ESelectInfo::Type::Direct);
+	UpdateAllTreeItemsVisibility();
+}
+
+void SChaosVDSceneQueryBrowser::HandleSolverVisibilityChanged(int32 SolverID, bool bNewVisibility)
+{
+	CachedSolverVisibilityByID.FindOrAdd(SolverID) = bNewVisibility;
+
+	for (const TSharedPtr<FChaosVDSceneQueryTreeItem>& TreeItem : UnfilteredCachedTreeItems)
+	{
+		if (TreeItem && TreeItem->OwnerSolverID == SolverID)
+		{
+			UpdateTreeItemVisibility(TreeItem);
+		}
+	}
+}
+
+void SChaosVDSceneQueryBrowser::RegisterMainToolbarMenu()
+{
+	const UToolMenus* ToolMenus = UToolMenus::Get();
+	if (ToolMenus->IsMenuRegistered(ToolBarName))
+	{
+		return;
+	}
+
+	UToolMenu* ToolBar = UToolMenus::Get()->RegisterMenu(ToolBarName, NAME_None, EMultiBoxType::SlimHorizontalToolBar);
+
+	FName MenuSection("SceneQueryBrowser.Toolbar.Settings");
+	FText MenuSectionLabel = LOCTEXT("SceneQueryBrowserMenuLabel", "Scene Query Visualization Settings");
+	FText FlagsMenuLabel = LOCTEXT("SceneQueryBrowserVisFlagsMenuLabel", "Visualization Flags");
+	FText FlagsMenuTooltip = LOCTEXT("SceneQueryBrowserVisFlagsMenuToolTip", "Set of flags to enable/disable visibility of specific types of scene query data");
+	FSlateIcon FlagsMenuIcon = FSlateIcon(FChaosVDStyle::Get().GetStyleSetName(), TEXT("SceneQueriesInspectorIcon"));
+
+	FText SettingsMenuLabel = LOCTEXT("SceneQueryBrowserVisSettingsMenuLabel", "General Settings");
+	FText SettingsMenuTooltip = LOCTEXT("SceneQueryBrowserVisMenuToolTip", "Options to change how the recorded scene query data is debug drawn");
+
+	using namespace Chaos::VisualDebugger::Utils;
+	CreateVisualizationOptionsMenuSections<UChaosVDSceneQueriesVisualizationSettings, EChaosVDSceneQueryVisualizationFlags>(ToolBar, MenuSection, MenuSectionLabel, FlagsMenuLabel, FlagsMenuTooltip, FlagsMenuIcon, SettingsMenuLabel, SettingsMenuTooltip);
+
+	FToolMenuSection& Section = ToolBar->AddSection(TEXT("SceneQueryBrowser.Toolbar"));
+
+	Section.AddSeparator(NAME_None);
+
+	Section.AddDynamicEntry("SearchBar", FNewToolMenuSectionDelegate::CreateLambda([](FToolMenuSection& InSection)
+	{
+		const UChaosVDSceneQueryBrowserToolbarMenuContext* Context = InSection.FindContext<UChaosVDSceneQueryBrowserToolbarMenuContext>();
+		TSharedRef<SChaosVDSceneQueryBrowser> BrowserWidget = Context->BrowserInstance.Pin().ToSharedRef();
+		
+		InSection.AddEntry(
+		FToolMenuEntry::InitWidget(
+			"SearchBar",
+			BrowserWidget->GenerateSearchBarWidget(),
+			FText::GetEmpty(),
+			false,
+			false
+		));
+		
+	}));
+
+	Section.AddDynamicEntry("VisualizationModes", FNewToolMenuSectionDelegate::CreateLambda([](FToolMenuSection& InSection)
+	{
+		const UChaosVDSceneQueryBrowserToolbarMenuContext* Context = InSection.FindContext<UChaosVDSceneQueryBrowserToolbarMenuContext>();
+		TSharedRef<SChaosVDSceneQueryBrowser> BrowserWidget = Context->BrowserInstance.Pin().ToSharedRef();
+		
+		InSection.AddSeparator(NAME_None);
+			
+		InSection.AddEntry(
+		FToolMenuEntry::InitWidget(
+			"VisualizationModes",
+			BrowserWidget->GenerateQueryVisualizationModeWidget(),
+			FText::GetEmpty(),
+			false,
+			false
+		));
+	}));
+}
+
+TSharedRef<SWidget> SChaosVDSceneQueryBrowser::GenerateMainToolbarWidget()
+{
+	RegisterMainToolbarMenu();
+
+	FToolMenuContext MenuContext;
+
+	UChaosVDSceneQueryBrowserToolbarMenuContext* CommonContextObject = NewObject<UChaosVDSceneQueryBrowserToolbarMenuContext>();
+	CommonContextObject->BrowserInstance = SharedThis(this);
+
+	MenuContext.AddObject(CommonContextObject);
+
+	return UToolMenus::Get()->GenerateWidget(ToolBarName, MenuContext);
+}
+
+TSharedRef<SWidget> SChaosVDSceneQueryBrowser::GenerateSearchBarWidget()
+{
+	TSharedPtr<SWidget> SearchBox = SNew(SSearchBox)
+										.IsEnabled(this, &SChaosVDSceneQueryBrowser::GetQueryTreeWidgetEnabled) // If the tree widget is disabled, we can't serach
+										.HintText(FText::FromString(TEXT("Search...")))
+										.OnTextChanged(this, &SChaosVDSceneQueryBrowser::HandleSearchTextChanged);
+
+	return SearchBox.ToSharedRef();
+}
+
+TSharedRef<SWidget> SChaosVDSceneQueryBrowser::GenerateQueryVisualizationModeWidget()
+{
+	TAttribute<int32> GetCurrentMode;
+	GetCurrentMode.BindLambda([]()
+	{
+		if (UChaosVDSceneQueriesVisualizationSettings* Settings = FChaosVDSettingsManager::Get().GetSettingsObject<UChaosVDSceneQueriesVisualizationSettings>())
+		{
+			return static_cast<int32>(Settings->CurrentVisualizationMode);
+		}
+		return 0;
+	});
+
+	SEnumComboBox::FOnEnumSelectionChanged ValueChangedDelegate;
+	ValueChangedDelegate.BindLambda([](int32 NewValue, ESelectInfo::Type)
+	{
+		if (UChaosVDSceneQueriesVisualizationSettings* Settings = FChaosVDSettingsManager::Get().GetSettingsObject<UChaosVDSceneQueriesVisualizationSettings>())
+		{
+			Settings->CurrentVisualizationMode = static_cast<EChaosVDSQFrameVisualizationMode>(NewValue);
+			Settings->OnSettingsChanged().Broadcast(Settings);
+		}
+	});
+
+	return Chaos::VisualDebugger::Utils::MakeEnumMenuEntryWidget<EChaosVDSQFrameVisualizationMode>(LOCTEXT("SQVisualizationModeModeMenuLabel", "Visualization Mode"), MoveTemp(ValueChangedDelegate), MoveTemp(GetCurrentMode));
+}
+
+EVisibility SChaosVDSceneQueryBrowser::GetUpdatesPausedMessageVisibility() const
+{
+	return !bIsUpToDate && !CanUpdate() ? EVisibility::Visible : EVisibility::Collapsed;
+}
+
+void SChaosVDSceneQueryBrowser::HandleSearchTextChanged(const FText& NewText)
+{
+	CurrentTextFilter = NewText;
+	bIsUpToDate = false;
+}
+
+void SChaosVDSceneQueryBrowser::HandleSettingsChanged(UObject* SettingsObject)
+{
+	bNeedsToUpdateSettings = true;
+}
+
+void SChaosVDSceneQueryBrowser::ApplySettingsChange()
+{
+	UChaosVDSceneQueriesVisualizationSettings* Settings = FChaosVDSettingsManager::Get().GetSettingsObject<UChaosVDSceneQueriesVisualizationSettings>();
+	if (!Settings)
+	{
+		return;
+	}
+
+	CurrentVisualizationMode = Settings->CurrentVisualizationMode;
+
+	UpdateBrowserContents();
+
+	switch (CurrentVisualizationMode)
+	{
+		case EChaosVDSQFrameVisualizationMode::AllEnabledQueries:
+			{
+				CurrentPlaybackIndex = 0;
+				break;
+			}
+		case EChaosVDSQFrameVisualizationMode::RecordingOrder:
+			{
+
+				if (FilteredCachedTreeItems->IsValidIndex(CurrentPlaybackIndex))
+				{
+					SelectSceneQuery((*FilteredCachedTreeItems)[CurrentPlaybackIndex], ESelectInfo::Type::OnMouseClick);
+				}
+				break;
+			}
+		default:
+			break;
+	}
+
+	if (TSharedPtr<FEditorModeTools> EditorModeToolsPtr = EditorModeToolsWeakPtr.Pin())
+	{
+		// Make sure the viewport is re-drawn so the selection feedback is shown
+		if (FEditorViewportClient* ViewportClient = EditorModeToolsPtr->GetFocusedViewportClient())
+		{
+			ViewportClient->Invalidate();
+		}
+	}
+
+	bNeedsToUpdateSettings = false;
+}
+
+void SChaosVDSceneQueryBrowser::UpdateBrowserContents()
+{
+	TSharedPtr<FChaosVDScene> ScenePtr = SceneWeakPtr.Pin();
+	if (!ScenePtr)
+	{
+		return;
+	}
+
+	TSharedPtr<SChaosVDSceneQueryTree> SceneQueryTreeWidgetPtr = SceneQueryTreeWidget.Pin();
+	if (!SceneQueryTreeWidgetPtr)
+	{
+		return;
+	}
+
+	UnfilteredCachedTreeItems.Reset();
+	CachedTreeItemsByID.Reset();
+
+	const TMap<int32, AChaosVDSolverInfoActor*>& SolverInfoDataMap = ScenePtr->GetSolverInfoActorsMap();
+
+	for (const TPair<int32, AChaosVDSolverInfoActor*> SolverDataWithID : SolverInfoDataMap)
+	{
+		if (SolverDataWithID.Key)
+		{
+			if (UChaosVDSceneQueryDataComponent* DataComponent = SolverDataWithID.Value->GetSceneQueryDataComponent())
+			{
+				for (const TSharedPtr<FChaosVDQueryDataWrapper>& QueryData : DataComponent->GetAllQueries())
+				{
+					if (QueryData->ParentQueryID == INDEX_NONE)
+					{
+						if (TSharedPtr<FChaosVDSceneQueryTreeItem> NewTreeItem = MakeSceneQueryTreeItem(QueryData, DataComponent))
+						{
+							UnfilteredCachedTreeItems.Add(NewTreeItem);
+						}
+					}
+				}
+			}
+		}
+	}
+
+	UpdateAllTreeItemsVisibility();
+
+	ApplyFilterToData(UnfilteredCachedTreeItems, FilteredCachedTreeItems.ToSharedRef());
+
+	if (CurrentVisualizationMode == EChaosVDSQFrameVisualizationMode::RecordingOrder)
+	{
+		// This is used to visualize queries on their recorded order when that mode is enabled.
+		// This works because the id is only a counter, and the queries are recorded in the same thread.
+		//TODO: This will no longer be the case when we add support to merge 2 CVD files. We will have to find another solution when the time comes
+		Algo::SortBy(*FilteredCachedTreeItems, [](const TSharedPtr<FChaosVDSceneQueryTreeItem>& Item)
+		{
+			return Item->QueryID;
+		});
+	}
+
+	SceneQueryTreeWidgetPtr->SetExternalSourceData(FilteredCachedTreeItems);
+
+	bIsUpToDate = true;
+}
+
+void SChaosVDSceneQueryBrowser::HandleSceneUpdated()
+{
+	bIsUpToDate = false;
+}
+
+bool SChaosVDSceneQueryBrowser::CanUpdate() const
+{
+	// Updating a list/tree widget every frame during playback is too expensive, so we disable it.
+	// TODO: We should modify this to have playback controller access without going through multiple objects 
+	TSharedPtr<FEditorModeTools> EditorModeToolsPtr = EditorModeToolsWeakPtr.Pin();
+	TSharedPtr<SChaosVDMainTab> ToolkitHost = EditorModeToolsPtr ? StaticCastSharedPtr<SChaosVDMainTab>(EditorModeToolsPtr->GetToolkitHost()) : nullptr;
+	TSharedPtr<FChaosVDPlaybackController> PlaybackController = ToolkitHost ? ToolkitHost->GetChaosVDEngineInstance()->GetPlaybackController() : nullptr;
+
+	if (PlaybackController)
+	{	
+		return !PlaybackController->IsPlaying();
+	}
+
+	return true;
+}
+
+bool SChaosVDSceneQueryBrowser::GetQueryTreeWidgetEnabled() const
+{
+	return bIsUpToDate || CanUpdate();
+}
+
+FText SChaosVDSceneQueryBrowser::GetFilterStatusText() const
+{
+	const int32 FilteredItemsNum = FilteredCachedTreeItems->Num();
+	const int32 UnFilteredItemsNum = UnfilteredCachedTreeItems.Num();
+
+	return FText::FormatOrdered(LOCTEXT("SceneQueryBrowserFilterStatusMessage","Showing {0} queries | {1} queries are hidden by search filter."), FilteredItemsNum, UnFilteredItemsNum - FilteredItemsNum);
+}
+
+FSlateColor SChaosVDSceneQueryBrowser::GetFilterStatusTextColor() const
+{
+	const int32 FilteredItemsNum = FilteredCachedTreeItems->Num();
+	const int32 UnFilteredItemsNum = UnfilteredCachedTreeItems.Num();
+	
+	if (CurrentTextFilter.IsEmpty())
+	{
+		return FSlateColor::UseForeground();
+	}
+	else if (FilteredItemsNum == 0 && UnFilteredItemsNum > 0)
+	{
+		return FAppStyle::Get().GetSlateColor("Colors.AccentRed");
+	}
+	else
+	{
+		return FAppStyle::Get().GetSlateColor("Colors.AccentGreen");
+	}
+}
+
+TSharedRef<SWidget> SChaosVDSceneQueryBrowser::GenerateQueriesPlaybackControls()
+{
+	constexpr float NoPadding = 0.0f;
+	constexpr float ContainerHorizontalPadding = 2.0f;
+	constexpr float ContainerTopPadding = 4.0f;
+	constexpr float ContainerBottomPadding = 10.0f;
+
+	constexpr float ControlsLabelVerticalPadding = 2.0f;
+
+	return SNew(SBorder)
+			.BorderImage(FAppStyle::Get().GetBrush("DetailsView.CategoryMiddle"))
+			.BorderBackgroundColor(FAppStyle::Get().GetSlateColor("Colors.Panel"))
+			.Padding(ContainerHorizontalPadding, ContainerTopPadding, ContainerHorizontalPadding, ContainerBottomPadding)
+			[
+				SNew(SVerticalBox)
+				+SVerticalBox::Slot()
+				.AutoHeight()
+				.Padding(NoPadding, ControlsLabelVerticalPadding, NoPadding, ControlsLabelVerticalPadding)
+				[
+					SNew(STextBlock)
+					.Justification(ETextJustify::Center)
+					.Text(this, &SChaosVDSceneQueryBrowser::GetPlaybackQueryControlText)
+				]
+				+SVerticalBox::Slot()
+				[
+					SAssignNew(PlaybackControlsTimelineWidget, SChaosVDTimelineWidget)
+					.ButtonVisibilityFlags(EChaosVDTimelineElementIDFlags::AllManualStepping)
+					.IsEnabled(this, &SChaosVDSceneQueryBrowser::GetPlaybackControlsEnabled)
+					.OnFrameChanged(this, &SChaosVDSceneQueryBrowser::HandlePlaybackQueryIndexUpdated)
+					.OnButtonClicked(this, &SChaosVDSceneQueryBrowser::HandlePlaybackControlInput)
+					.MinFrames(this, &SChaosVDSceneQueryBrowser::GetCurrentMinPlaybackQueryIndex)
+					.MaxFrames(this, &SChaosVDSceneQueryBrowser::GetCurrentMaxPlaybackQueryIndex)
+					.CurrentFrame(this, &SChaosVDSceneQueryBrowser::GetCurrentPlaybackQueryIndex)
+				]
+			];
+}
+
+bool SChaosVDSceneQueryBrowser::GetPlaybackControlsEnabled() const
+{
+	return CurrentVisualizationMode == EChaosVDSQFrameVisualizationMode::RecordingOrder && GetCurrentMaxPlaybackQueryIndex() > 0;
+}
+
+void SChaosVDSceneQueryBrowser::HandlePlaybackControlInput(EChaosVDPlaybackButtonsID InputID)
+{
+	switch (InputID)
+	{
+	case EChaosVDPlaybackButtonsID::Next:
+		{
+			int32 NextQueryIndex = CurrentPlaybackIndex +1;
+			if (FilteredCachedTreeItems->IsValidIndex(NextQueryIndex))
+			{
+				SelectSceneQuery((*FilteredCachedTreeItems)[NextQueryIndex], ESelectInfo::Type::OnMouseClick);
+				CurrentPlaybackIndex = NextQueryIndex;
+			}
+			
+			break;	
+		}
+	case EChaosVDPlaybackButtonsID::Prev:
+		{
+			int32 PrevQueryIndex = CurrentPlaybackIndex -1;
+			if (FilteredCachedTreeItems->IsValidIndex(PrevQueryIndex))
+			{
+				SelectSceneQuery((*FilteredCachedTreeItems)[PrevQueryIndex], ESelectInfo::Type::OnMouseClick);
+				CurrentPlaybackIndex = PrevQueryIndex;
+			}
+			break;	
+		}
+	case EChaosVDPlaybackButtonsID::Play:
+	case EChaosVDPlaybackButtonsID::Pause:
+	case EChaosVDPlaybackButtonsID::Stop:
+	default:
+		break;
+	}
+}
+
+void SChaosVDSceneQueryBrowser::HandlePlaybackQueryIndexUpdated(int32 NewIndex)
+{
+	CurrentPlaybackIndex = NewIndex;
+	if (FilteredCachedTreeItems->IsValidIndex(CurrentPlaybackIndex))
+	{
+		SelectSceneQuery((*FilteredCachedTreeItems)[CurrentPlaybackIndex], ESelectInfo::Type::OnMouseClick);
+	}
+}
+
+FText SChaosVDSceneQueryBrowser::GetPlaybackQueryControlText() const
+{
+	return LOCTEXT("SceneQueryBrowserPlaybackControlsLabel", "Recorded Scene Queries");
+}
+
+int32 SChaosVDSceneQueryBrowser::GetCurrentMinPlaybackQueryIndex() const
+{
+	return 0;
+}
+
+int32 SChaosVDSceneQueryBrowser::GetCurrentMaxPlaybackQueryIndex() const
+{
+	return FilteredCachedTreeItems ? FMath::Clamp(FilteredCachedTreeItems->Num() - 1, 0, TNumericLimits<int32>::Max()) : 0;
+}
+
+int32 SChaosVDSceneQueryBrowser::GetCurrentPlaybackQueryIndex() const
+{
+	return CurrentPlaybackIndex;
+}
+
+TSharedPtr<FChaosVDSceneQueryTreeItem> SChaosVDSceneQueryBrowser::MakeSceneQueryTreeItem(const TSharedPtr<FChaosVDQueryDataWrapper>& InQueryData, const UChaosVDSceneQueryDataComponent* DataComponent)
+{
+	if (!InQueryData || !DataComponent)
+	{
+		return nullptr;
+	}
+
+	if (TSharedPtr<FChaosVDSceneQueryTreeItem>* FoundItemPtr = CachedTreeItemsByID.Find(InQueryData->ID))
+	{
+		return *FoundItemPtr;
+	}
+
+	TSharedPtr<FChaosVDSceneQueryTreeItem> NewTreeItem = MakeShared<FChaosVDSceneQueryTreeItem>();
+	NewTreeItem->ItemWeakPtr = InQueryData;
+
+	for (int32 SubQueryID : InQueryData->SubQueriesIDs)
+	{
+		if (TSharedPtr<FChaosVDSceneQueryTreeItem> NewSubTreeItem = MakeSceneQueryTreeItem(DataComponent->GetQueryByID(SubQueryID), DataComponent))
+		{
+			NewTreeItem->SubItems.Add(NewSubTreeItem);
+		}
+	}
+
+	NewTreeItem->QueryID = InQueryData->ID;
+	NewTreeItem->OwnerSolverID = InQueryData->WorldSolverID;
+
+	if (AChaosVDSolverInfoActor* SolverInfo = Cast<AChaosVDSolverInfoActor>(DataComponent->GetOwner()))
+	{
+		NewTreeItem->OwnerSolverName = SolverInfo->GetSolverName();
+	}
+	
+	CachedTreeItemsByID.Add(InQueryData->ID, NewTreeItem);
+
+	return NewTreeItem;
+}
+
+bool SChaosVDSceneQueryBrowser::IsQueryDataVisible(const TSharedRef<FChaosVDQueryDataWrapper>& InQueryData)
+{
+	bool bIsVisible = false;
+
+	const TSharedPtr<FChaosVDScene> ScenePtr = SceneWeakPtr.Pin();
+	if (!ScenePtr)
+	{
+		return bIsVisible;
+	}
+
+	EChaosVDSceneQueryVisualizationFlags VisualizationFlags = UChaosVDSceneQueriesVisualizationSettings::GetDataVisualizationFlags();
+
+	if (!EnumHasAnyFlags(VisualizationFlags, EChaosVDSceneQueryVisualizationFlags::EnableDraw))
+	{
+		return bIsVisible;
+	}
+
+	switch (InQueryData->Type)
+	{
+		case EChaosVDSceneQueryType::RayCast:
+				bIsVisible = EnumHasAnyFlags(VisualizationFlags, EChaosVDSceneQueryVisualizationFlags::DrawLineTraceQueries);
+				break;
+		case EChaosVDSceneQueryType::Overlap:
+				bIsVisible = EnumHasAnyFlags(VisualizationFlags, EChaosVDSceneQueryVisualizationFlags::DrawOverlapQueries);
+				break;
+		case EChaosVDSceneQueryType::Sweep:
+				bIsVisible = EnumHasAnyFlags(VisualizationFlags, EChaosVDSceneQueryVisualizationFlags::DrawSweepQueries);
+				break;
+		default:
+			break;
+	}
+
+	// Only do the more expensive visibility checks if the simple ones passed
+	if (bIsVisible)
+	{
+		if (EnumHasAnyFlags(VisualizationFlags, EChaosVDSceneQueryVisualizationFlags::OnlyDrawSelectedQuery))
+		{
+			if (TSharedPtr<FChaosVDSolverDataSelection> SelectionObject = ScenePtr->GetSolverDataSelectionObject().Pin())
+			{
+				bIsVisible = SelectionObject->IsDataSelected(InQueryData.ToSharedPtr());
+			}
+		}
+
+		bIsVisible &= GetCachedSolverVisibility(InQueryData->WorldSolverID);
+	}
+
+	return bIsVisible;
+}
+
+void SChaosVDSceneQueryBrowser::UpdateAllTreeItemsVisibility()
+{
+	for (const TSharedPtr<FChaosVDSceneQueryTreeItem>& TreeItem : UnfilteredCachedTreeItems)
+	{
+		UpdateTreeItemVisibility(TreeItem);
+	}
+}
+
+void SChaosVDSceneQueryBrowser::UpdateTreeItemVisibility(const TSharedPtr<FChaosVDSceneQueryTreeItem>& InTreeItem)
+{
+	if (TSharedPtr<FChaosVDQueryDataWrapper> QueryData = InTreeItem ? InTreeItem->ItemWeakPtr.Pin() : nullptr)
+	{
+		InTreeItem->bIsVisible = IsQueryDataVisible(QueryData.ToSharedRef());
+
+		for (const TSharedPtr<FChaosVDSceneQueryTreeItem>& TreeItem : InTreeItem->SubItems)
+		{
+			UpdateTreeItemVisibility(TreeItem);
+		}
+	}
+}
+
+void SChaosVDSceneQueryBrowser::HandleTreeItemSelected(const TSharedPtr<FChaosVDSceneQueryTreeItem>& SelectedTreeItem, ESelectInfo::Type Type)
+{
+	// Only handle this selection event if it came from the UI
+	if (Type == ESelectInfo::Type::Direct)
+	{
+		return;
+	}
+
+	if (!SelectedTreeItem)
+	{
+		return;
+	}
+
+	if (const TSharedPtr<FChaosVDScene> ScenePtr = SceneWeakPtr.Pin())
+	{
+		if (TSharedPtr<FChaosVDSolverDataSelection> SelectionObject = ScenePtr->GetSolverDataSelectionObject().Pin())
+		{
+			TSharedPtr<FChaosVDSolverDataSelectionHandle> SelectionHandle = SelectionObject->MakeSelectionHandle(SelectedTreeItem->ItemWeakPtr.Pin());
+			SelectionObject->SelectData(SelectionHandle);
+
+			if (TSharedPtr<FEditorModeTools> EditorModeToolsPtr = EditorModeToolsWeakPtr.Pin())
+			{
+				// Bring the constraint details tab into focus if available
+				if (TSharedPtr<SChaosVDMainTab> ToolkitHost = StaticCastSharedPtr<SChaosVDMainTab>(EditorModeToolsPtr->GetToolkitHost()))
+				{
+					if (const TSharedPtr<FTabManager> TabManager = ToolkitHost->GetTabManager())
+					{
+						TabManager->TryInvokeTab(FChaosVDTabID::SceneQueryDataDetails);
+					}
+				}
+
+				// Make sure the viewport is re-drawn so the selection feedback is shown
+				//TODO: This likely needs to be done automatically when any selection event is triggered
+				if (FEditorViewportClient* ViewportClient = EditorModeToolsPtr->GetFocusedViewportClient())
+				{
+					ViewportClient->Invalidate();
+				}
+			}
+		}
+	}
+}
+
+void SChaosVDSceneQueryBrowser::HandleTreeItemFocused(const TSharedPtr<FChaosVDSceneQueryTreeItem>& FocusedTreeItem)
+{
+	if (!FocusedTreeItem)
+	{
+		return;
+	}
+	
+	TSharedPtr<FChaosVDQueryDataWrapper> QueryData = FocusedTreeItem->ItemWeakPtr.Pin();
+	
+	if (!QueryData)
+	{
+		return;
+	}
+	
+	const TSharedPtr<FChaosVDScene> ScenePtr = SceneWeakPtr.Pin();
+	if (!ScenePtr)
+	{
+		return;
+	}
+
+	TSharedPtr<FChaosVDRecording> RecordedData = ScenePtr->LoadedRecording;
+	if (!RecordedData)
+	{
+		return;
+	}
+
+	return ScenePtr->OnFocusRequest().Broadcast(Chaos::VisualDebugger::Utils::CalculateSceneQueryShapeBounds(QueryData.ToSharedRef(), RecordedData.ToSharedRef()));
+}
+
+void SChaosVDSceneQueryBrowser::ApplyFilterToData(TConstArrayView<TSharedPtr<FChaosVDSceneQueryTreeItem>> InDataSource, const TSharedRef<TArray<TSharedPtr<FChaosVDSceneQueryTreeItem>>>& OutFilteredData)
+{
+	OutFilteredData->Reset(InDataSource.Num());
+
+	if (CurrentTextFilter.IsEmpty())
+	{
+		*OutFilteredData = InDataSource;
+	}
+	else
+	{
+		FString CurrentFilterAsString = CurrentTextFilter.ToString();
+
+		for (const TSharedPtr<FChaosVDSceneQueryTreeItem>& QueryDataItem : InDataSource)
+		{
+			TSharedPtr<FChaosVDQueryDataWrapper> QueryData = QueryDataItem ? QueryDataItem->ItemWeakPtr.Pin() : nullptr;
+			if (QueryData)
+			{
+				// TODO: Add support for fuzzy search?. We will likely need to make the search an async operation
+				if (GetCachedStringFromName(QueryData->CollisionQueryParams.TraceTag)->Contains(CurrentFilterAsString) ||
+					GetCachedStringFromName(QueryData->CollisionQueryParams.OwnerTag)->Contains(CurrentFilterAsString) ||
+					GetCachedStringFromName(QueryDataItem->OwnerSolverName)->Contains(CurrentFilterAsString))
+				{
+					OutFilteredData->Add(QueryDataItem);
+				}
+			}
+		}
+	}
+}
+
+TSharedRef<FString> SChaosVDSceneQueryBrowser::GetCachedStringFromName(FName Name)
+{
+	if (TSharedPtr<FString>* FoundStringPtr = CachedNameToStringMap.Find(Name))
+	{
+		return FoundStringPtr->ToSharedRef();
+	}
+	else
+	{
+		TSharedPtr<FString> NameAsString = MakeShared<FString>(Name.ToString());
+		CachedNameToStringMap.Add(Name, NameAsString);
+
+		return NameAsString.ToSharedRef();
+	}
+}
+
+bool SChaosVDSceneQueryBrowser::GetCachedSolverVisibility(int32 SolverID)
+{
+	if (const bool* bIsSolverVisiblePtr = CachedSolverVisibilityByID.Find(SolverID))
+	{
+		return *bIsSolverVisiblePtr;
+	}
+
+	const TSharedPtr<FChaosVDScene> ScenePtr = SceneWeakPtr.Pin();
+	if (!ScenePtr)
+	{
+		return false;
+	}
+	
+	if (AChaosVDSolverInfoActor* OwningSolver = ScenePtr->GetSolverInfoActor(SolverID))
+	{
+		bool bIsVisible = OwningSolver->IsVisible();
+
+		CachedSolverVisibilityByID.Add(SolverID, bIsVisible);
+		
+		return bIsVisible;
+	}
+
+	return false;
+}
+
+
+#undef LOCTEXT_NAMESPACE
+
diff --git a/Engine/Plugins/ChaosVD/Source/ChaosVD/Private/Widgets/SChaosVDSceneQueryBrowser.h b/Engine/Plugins/ChaosVD/Source/ChaosVD/Private/Widgets/SChaosVDSceneQueryBrowser.h
new file mode 100644
--- /dev/null
+++ b/Engine/Plugins/ChaosVD/Source/ChaosVD/Private/Widgets/SChaosVDSceneQueryBrowser.h
@@ -0,0 +1,137 @@
+// Copyright Epic Games, Inc. All Rights Reserved.
+
+#pragma once
+
+#include "SChaosVDGameFramesPlaybackControls.h"
+#include "Settings/ChaosVDSceneQueryVisualizationSettings.h"
+#include "Widgets/SCompoundWidget.h"
+
+#include "SChaosVDSceneQueryBrowser.generated.h"
+
+class FEditorModeTools;
+class SChaosVDSceneQueryTree;
+class UChaosVDSceneQueryDataComponent;
+
+struct FChaosVDQueryDataWrapper;
+struct FChaosVDSceneQueryTreeItem;
+struct FChaosVDSolverDataSelectionHandle;
+
+UCLASS()
+class CHAOSVD_API UChaosVDSceneQueryBrowserToolbarMenuContext : public UObject
+{
+	GENERATED_BODY()
+
+public:
+
+	TWeakPtr<class SChaosVDSceneQueryBrowser> BrowserInstance;
+};
+
+/**
+ * Widget class of the Scene Query Browser window, where all available scene queries are shown for the currently visualized frame
+ * in a Scene Outliner kind of way
+ */
+class CHAOSVD_API SChaosVDSceneQueryBrowser : public SCompoundWidget
+{
+public:
+	SLATE_BEGIN_ARGS(SChaosVDSceneQueryBrowser)
+		{
+		}
+
+
+	SLATE_END_ARGS()
+
+	virtual ~SChaosVDSceneQueryBrowser() override;
+
+	void Construct(const FArguments& InArgs, TWeakPtr<FChaosVDScene> Scene, TWeakPtr<FEditorModeTools> EditorModeTools);
+	void RegisterSceneEvents();
+	void UnregisterSceneEvents();
+
+	static inline const FName ToolBarName = FName("ChaosVD.SceneQueryBrowser.ToolBar");
+
+	virtual void Tick(const FGeometry& AllottedGeometry, const double InCurrentTime, const float InDeltaTime) override;
+
+protected:
+
+	void SelectSceneQuery(const TSharedPtr<const FChaosVDQueryDataWrapper>& InQuery, ESelectInfo::Type Type);
+	void SelectSceneQuery(const TSharedPtr<FChaosVDSceneQueryTreeItem>& SceneQueryTreeItem, ESelectInfo::Type Type);
+	void HandleExternalSelectionEvent(const TSharedPtr<FChaosVDSolverDataSelectionHandle>& InDataSelectionHandle);
+	void HandleSolverVisibilityChanged(int32 SolverID, bool bNewVisibility);
+
+	void RegisterMainToolbarMenu();
+
+	TSharedRef<SWidget> GenerateMainToolbarWidget();
+	TSharedRef<SWidget> GenerateSearchBarWidget();
+	TSharedRef<SWidget> GenerateQueryVisualizationModeWidget();
+
+	EVisibility GetUpdatesPausedMessageVisibility() const;
+
+	void HandleSearchTextChanged(const FText& NewText);
+	
+	void HandleSettingsChanged(UObject* SettingsObject);
+
+	void ApplySettingsChange();
+
+	void UpdateBrowserContents();
+
+	void HandleSceneUpdated();
+
+	bool CanUpdate() const;
+
+	bool GetQueryTreeWidgetEnabled() const;
+	
+	FText GetFilterStatusText() const;
+	FSlateColor GetFilterStatusTextColor() const;
+
+	TSharedRef<SWidget> GenerateQueriesPlaybackControls();
+
+	bool GetPlaybackControlsEnabled() const;
+	
+	void HandlePlaybackControlInput(EChaosVDPlaybackButtonsID InputID);
+
+	void HandlePlaybackQueryIndexUpdated(int32 NewIndex);
+
+	FText GetPlaybackQueryControlText() const;
+	int32 GetCurrentMinPlaybackQueryIndex() const;
+	int32 GetCurrentMaxPlaybackQueryIndex() const;
+	int32 GetCurrentPlaybackQueryIndex() const;
+
+	TSharedPtr<FChaosVDSceneQueryTreeItem> MakeSceneQueryTreeItem(const TSharedPtr<FChaosVDQueryDataWrapper>& InQueryData, const UChaosVDSceneQueryDataComponent* DataComponent);
+
+	bool IsQueryDataVisible(const TSharedRef<FChaosVDQueryDataWrapper>& InQueryData);
+
+	void UpdateAllTreeItemsVisibility();
+	void UpdateTreeItemVisibility(const TSharedPtr<FChaosVDSceneQueryTreeItem>& InTreeItem);
+	
+	void HandleTreeItemSelected(const TSharedPtr<FChaosVDSceneQueryTreeItem>& SelectedTreeItem, ESelectInfo::Type Type);
+	void HandleTreeItemFocused(const TSharedPtr<FChaosVDSceneQueryTreeItem>& FocusedTreeItem);
+
+	void ApplyFilterToData(TConstArrayView<TSharedPtr<FChaosVDSceneQueryTreeItem>> InDataSource, const TSharedRef<TArray<TSharedPtr<FChaosVDSceneQueryTreeItem>>>& OutFilteredData);
+
+	TSharedRef<FString> GetCachedStringFromName(FName Name);
+
+	bool GetCachedSolverVisibility(int32 SolverID);
+	
+	TSharedPtr<SChaosVDTimelineWidget> PlaybackControlsTimelineWidget;
+	
+	TWeakPtr<FChaosVDScene> SceneWeakPtr;
+	
+	TWeakPtr<FEditorModeTools> EditorModeToolsWeakPtr;
+
+	TWeakPtr<SChaosVDSceneQueryTree> SceneQueryTreeWidget;
+
+	TSharedPtr<TArray<TSharedPtr<FChaosVDSceneQueryTreeItem>>> FilteredCachedTreeItems;
+	TArray<TSharedPtr<FChaosVDSceneQueryTreeItem>> UnfilteredCachedTreeItems;
+	TMap<int32, TSharedPtr<FChaosVDSceneQueryTreeItem>> CachedTreeItemsByID;
+	TMap<int32, bool> CachedSolverVisibilityByID;
+
+	TMap<FName, TSharedPtr<FString>> CachedNameToStringMap;
+	
+	EChaosVDSQFrameVisualizationMode CurrentVisualizationMode = EChaosVDSQFrameVisualizationMode::AllEnabledQueries;
+
+	FText CurrentTextFilter;
+
+	int32 CurrentPlaybackIndex = 0;
+	
+	bool bIsUpToDate = false;
+	bool bNeedsToUpdateSettings = true;
+};
diff --git a/Engine/Plugins/ChaosVD/Source/ChaosVD/Private/Widgets/SChaosVDSceneQueryTree.cpp b/Engine/Plugins/ChaosVD/Source/ChaosVD/Private/Widgets/SChaosVDSceneQueryTree.cpp
new file mode 100644
--- /dev/null
+++ b/Engine/Plugins/ChaosVD/Source/ChaosVD/Private/Widgets/SChaosVDSceneQueryTree.cpp
@@ -0,0 +1,153 @@
+// Copyright Epic Games, Inc. All Rights Reserved.
+
+#include "SChaosVDSceneQueryTree.h"
+
+#include "SChaosVDSceneQueryTreeRow.h"
+#include "Widgets/Images/SImage.h"
+#include "Widgets/Text/STextBlock.h"
+
+#define LOCTEXT_NAMESPACE "ChaosVisualDebugger"
+
+SChaosVDSceneQueryTree::FColumNames SChaosVDSceneQueryTree::ColumnNames = SChaosVDSceneQueryTree::FColumNames();
+
+void SChaosVDSceneQueryTree::Construct(const FArguments& InArgs)
+{
+	QueryItemSelectedDelegate = InArgs._OnItemSelected;
+	QueryItemFocusedDelegate = InArgs._OnItemFocused;
+
+	constexpr float BottomPadding = 2.0f;
+	constexpr float NoPadding = 0.0f;
+	constexpr float VerticalPadding = 0.0f;
+
+	static FMargin ColumHeaderTextMargin(NoPadding,VerticalPadding);
+
+	ChildSlot
+	[
+		SNew(SVerticalBox)
+		+SVerticalBox::Slot()
+		.Padding(NoPadding, NoPadding, NoPadding, BottomPadding)
+		[
+			SAssignNew(SceneQueriesListWidget, STreeView<TSharedPtr<FChaosVDSceneQueryTreeItem>>)
+			.OnGenerateRow(this, &SChaosVDSceneQueryTree::GenerateSceneQueryDataRow)
+			.OnSelectionChanged(this, &SChaosVDSceneQueryTree::QueryTreeSelectionChanged)
+			.OnGetChildren(this, &SChaosVDSceneQueryTree::OnGetChildrenForQueryItem)
+			.TreeItemsSource(&InternalTreeItemSourceData)
+			.OnMouseButtonDoubleClick(this, &SChaosVDSceneQueryTree::HandleFocusRequest)
+			.SelectionMode(ESelectionMode::Type::Single)
+			.HighlightParentNodesForSelection(true)
+			.HeaderRow(
+						// TODO: Add sort support for the columns
+						SNew(SHeaderRow)
+						+SHeaderRow::Column(ColumnNames.Visibility)
+						.SortMode(EColumnSortMode::None)
+						.FixedWidth(24.f)
+						.HAlignHeader(HAlign_Left)
+						.VAlignHeader(VAlign_Center)
+						.HAlignCell(HAlign_Center)
+						.VAlignCell(VAlign_Center)
+						.HeaderContentPadding(FMargin(4.0f, 0.0f, 0.0f, 0.0f))
+						[
+							SNew(SImage)
+							.ColorAndOpacity(FSlateColor::UseForeground())
+							.Image(FAppStyle::Get().GetBrush("Level.VisibleIcon16x"))
+						]
+						+SHeaderRow::Column(ColumnNames.TraceTag)
+						.SortMode(EColumnSortMode::None)
+						.HAlignCell(HAlign_Left)
+						[
+							SNew(STextBlock)
+							.Margin(ColumHeaderTextMargin)
+							.Text(LOCTEXT("QueryListTagHeader", "Trace Tag"))
+						]
+						+SHeaderRow::Column(ColumnNames.TraceOwner)
+						.SortMode(EColumnSortMode::None)
+						.HAlignCell(HAlign_Left)
+						[
+							SNew(STextBlock)
+							.Margin(ColumHeaderTextMargin)
+							.Text(LOCTEXT("QueryListOwnerHeader", "Trace Owner"))
+						]
+						+SHeaderRow::Column(ColumnNames.QueryType)
+						.SortMode(EColumnSortMode::None)
+						.HAlignCell(HAlign_Left)
+						[
+							SNew(STextBlock)
+							.Margin(ColumHeaderTextMargin)
+							.Text(LOCTEXT("QueryListTypeHeader", "Query Type"))
+						]
+						+SHeaderRow::Column(ColumnNames.SolverName)
+						.SortMode(EColumnSortMode::None)
+						.HAlignCell(HAlign_Left)
+						[
+							SNew(STextBlock)
+							.Margin(ColumHeaderTextMargin)
+							.Text(LOCTEXT("QueryListSolverNameHeader", "Solver Name"))
+						]
+					)
+		]
+	];
+}
+
+TSharedRef<ITableRow> SChaosVDSceneQueryTree::GenerateSceneQueryDataRow(TSharedPtr<FChaosVDSceneQueryTreeItem> SceneQueryData, const TSharedRef<STableViewBase>& OwnerTable)
+{
+	if (!SceneQueryData)
+	{
+		return
+		SNew(STableRow< TSharedPtr<FString> >, OwnerTable)
+		[
+			SNew(SVerticalBox)
+			+SVerticalBox::Slot()
+			[
+				SNew(STextBlock)
+				.Text(LOCTEXT("SChaosVDSceneQueryListErrorMessage", "Failed to read data for solver."))
+			]
+		];
+	}
+
+	return SNew(SChaosVDSceneQueryTreeRow, OwnerTable)
+			.Item(SceneQueryData);
+}
+
+void SChaosVDSceneQueryTree::QueryTreeSelectionChanged(TSharedPtr<FChaosVDSceneQueryTreeItem> SelectedQuery, ESelectInfo::Type Type)
+{
+	QueryItemSelectedDelegate.ExecuteIfBound(SelectedQuery, Type);
+}
+
+void SChaosVDSceneQueryTree::OnGetChildrenForQueryItem(TSharedPtr<FChaosVDSceneQueryTreeItem> QueryEntry, TArray<TSharedPtr<FChaosVDSceneQueryTreeItem>>& OutQueries)
+{
+	OutQueries.Append(QueryEntry->SubItems);
+}
+
+void SChaosVDSceneQueryTree::SetExternalSourceData(const TSharedPtr<TArray<TSharedPtr<FChaosVDSceneQueryTreeItem>>>& InUpdatedSceneQueryDataSource)
+{
+	if (ExternalTreeItemSourceData != InUpdatedSceneQueryDataSource)
+	{
+		ExternalTreeItemSourceData = InUpdatedSceneQueryDataSource;
+		SceneQueriesListWidget->SetTreeItemsSource(ExternalTreeItemSourceData.Get());
+	}
+
+	for (const TSharedPtr<FChaosVDSceneQueryTreeItem>& TreeItem : *ExternalTreeItemSourceData)
+	{
+		SceneQueriesListWidget->SetItemExpansion(TreeItem, true);
+	}
+
+	SceneQueriesListWidget->RebuildList();
+	SceneQueriesListWidget->RequestTreeRefresh();
+}
+
+void SChaosVDSceneQueryTree::SelectItem(const TSharedPtr<FChaosVDSceneQueryTreeItem>& ItemToSelect, ESelectInfo::Type Type)
+{
+	SceneQueriesListWidget->SetSelection(ItemToSelect, Type);
+}
+
+void SChaosVDSceneQueryTree::Tick(const FGeometry& AllottedGeometry, const double InCurrentTime, const float InDeltaTime)
+{
+	SCompoundWidget::Tick(AllottedGeometry, InCurrentTime, InDeltaTime);
+}
+
+void SChaosVDSceneQueryTree::HandleFocusRequest(TSharedPtr<FChaosVDSceneQueryTreeItem> InFocusedItem)
+{
+	QueryItemFocusedDelegate.ExecuteIfBound(InFocusedItem);
+}
+
+#undef LOCTEXT_NAMESPACE
diff --git a/Engine/Plugins/ChaosVD/Source/ChaosVD/Private/Widgets/SChaosVDSceneQueryTree.h b/Engine/Plugins/ChaosVD/Source/ChaosVD/Private/Widgets/SChaosVDSceneQueryTree.h
new file mode 100644
--- /dev/null
+++ b/Engine/Plugins/ChaosVD/Source/ChaosVD/Private/Widgets/SChaosVDSceneQueryTree.h
@@ -0,0 +1,79 @@
+// Copyright Epic Games, Inc. All Rights Reserved.
+
+#pragma once
+
+#include "Containers/Array.h"
+#include "Templates/SharedPointer.h"
+#include "Widgets/SCompoundWidget.h"
+#include "Widgets/Views/STableViewBase.h"
+#include "Widgets/Views/STreeView.h"
+
+class ITableRow;
+class UChaosVDSceneQueryDataComponent;
+
+struct FChaosVDQueryDataWrapper;
+
+struct FChaosVDSceneQueryTreeItem
+{
+	TWeakPtr<FChaosVDQueryDataWrapper> ItemWeakPtr;
+	TArray<TSharedPtr<FChaosVDSceneQueryTreeItem>> SubItems;
+	int32 QueryID = INDEX_NONE;
+	FName OwnerSolverName = FName("Invalid");
+	int32 OwnerSolverID = INDEX_NONE;
+	bool bIsVisible = true;
+};
+
+DECLARE_DELEGATE_TwoParams(FChaosVDQueryTreeItemSelected, const TSharedPtr<FChaosVDSceneQueryTreeItem>&, ESelectInfo::Type)
+DECLARE_DELEGATE_OneParam(FChaosVDQueryTreeItemFocused, const TSharedPtr<FChaosVDSceneQueryTreeItem>&)
+
+/**
+ * Tree Widget used to represent the recorded scene queries hierachy
+ */
+class SChaosVDSceneQueryTree : public SCompoundWidget
+{
+public:
+	SLATE_BEGIN_ARGS(SChaosVDSceneQueryTree) {}
+		SLATE_EVENT(FChaosVDQueryTreeItemSelected, OnItemSelected)
+		SLATE_EVENT(FChaosVDQueryTreeItemFocused, OnItemFocused)
+	SLATE_END_ARGS()
+	
+	void Construct(const FArguments& InArgs);
+
+	void SetExternalSourceData(const TSharedPtr<TArray<TSharedPtr<FChaosVDSceneQueryTreeItem>>>& InUpdatedSceneQueryDataSource);
+
+	void SelectItem(const TSharedPtr<FChaosVDSceneQueryTreeItem>& ItemToSelect, ESelectInfo::Type Type);
+
+	virtual void Tick(const FGeometry& AllottedGeometry, const double InCurrentTime, const float InDeltaTime) override;
+
+protected:
+
+	void HandleFocusRequest(TSharedPtr<FChaosVDSceneQueryTreeItem> InFocusedItem);
+
+	TSharedRef<ITableRow> GenerateSceneQueryDataRow(TSharedPtr<FChaosVDSceneQueryTreeItem> SceneQueryData, const TSharedRef<STableViewBase>& OwnerTable);
+
+	void QueryTreeSelectionChanged(TSharedPtr<FChaosVDSceneQueryTreeItem> SelectedQuery, ESelectInfo::Type Type);
+
+	void OnGetChildrenForQueryItem(TSharedPtr<FChaosVDSceneQueryTreeItem> QueryEntry, TArray<TSharedPtr<FChaosVDSceneQueryTreeItem>>& OutQueries);
+
+	TSharedPtr<STreeView<TSharedPtr<FChaosVDSceneQueryTreeItem>>> SceneQueriesListWidget;
+
+	TSharedPtr<TArray<TSharedPtr<FChaosVDSceneQueryTreeItem>>> ExternalTreeItemSourceData;
+
+	TArray<TSharedPtr<FChaosVDSceneQueryTreeItem>> InternalTreeItemSourceData;
+
+	FChaosVDQueryTreeItemSelected QueryItemSelectedDelegate;
+	FChaosVDQueryTreeItemFocused QueryItemFocusedDelegate;
+
+public:
+
+	struct FColumNames
+	{
+		const FName Visibility = FName("Visibility");
+		const FName TraceTag = FName("TraceTag");
+		const FName TraceOwner = FName("TraceOwner");
+		const FName QueryType = FName("QueryType");
+		const FName SolverName = FName("SolverName");
+	};
+
+	static FColumNames ColumnNames;
+};
diff --git a/Engine/Plugins/ChaosVD/Source/ChaosVD/Private/Widgets/SChaosVDSceneQueryTreeRow.cpp b/Engine/Plugins/ChaosVD/Source/ChaosVD/Private/Widgets/SChaosVDSceneQueryTreeRow.cpp
new file mode 100644
--- /dev/null
+++ b/Engine/Plugins/ChaosVD/Source/ChaosVD/Private/Widgets/SChaosVDSceneQueryTreeRow.cpp
@@ -0,0 +1,120 @@
+// Copyright Epic Games, Inc. All Rights Reserved.
+
+#include "SChaosVDSceneQueryTreeRow.h"
+
+#include "SChaosVDSceneQueryTree.h"
+#include "SlateOptMacros.h"
+#include "DataWrappers/ChaosVDQueryDataWrappers.h"
+#include "Widgets/Images/SImage.h"
+#include "Widgets/Layout/SBox.h"
+#include "Widgets/SBoxPanel.h"
+#include "Widgets/Text/STextBlock.h"
+
+#define LOCTEXT_NAMESPACE "ChaosVisualDebugger"
+
+void SChaosVDSceneQueryTreeRow::Construct(const FArguments& InArgs, const TSharedRef<STableViewBase>& InOwnerTableView)
+{
+	Item = InArgs._Item;
+	FSuperRowType::FArguments Args = FSuperRowType::FArguments()
+									.Style(&FAppStyle::Get().GetWidgetStyle<FTableRowStyle>("SceneOutliner.TableViewRow"));
+
+	SMultiColumnTableRow<TSharedPtr<const FChaosVDSceneQueryTreeItem>>::Construct(Args, InOwnerTableView);
+}
+
+TSharedRef<SWidget> SChaosVDSceneQueryTreeRow::GenerateWidgetForColumn(const FName& ColumnName)
+{
+	if (!Item)
+	{
+		return SNullWidget::NullWidget;
+	}
+
+	TSharedPtr<const FChaosVDQueryDataWrapper> QueryDataPtr = Item->ItemWeakPtr.Pin();
+	if (!QueryDataPtr)
+	{
+		return SNullWidget::NullWidget;
+	}
+
+	if (ColumnName == SChaosVDSceneQueryTree::ColumnNames.TraceTag)
+	{
+		constexpr float NoPadding = 0.0f;
+		constexpr float ExpanderLeftPadding = 6.0f;
+		constexpr float ExpanderIndentAmount = 12.0f;
+		// The first column gets the tree expansion arrow for this row
+		return SNew(SBox)
+			[
+				SNew(SHorizontalBox)
+				+SHorizontalBox::Slot()
+				.AutoWidth()
+				.Padding(ExpanderLeftPadding, NoPadding, NoPadding, NoPadding)
+				[
+					SNew(SExpanderArrow, SharedThis(this)).IndentAmount(ExpanderIndentAmount)
+				]
+				+SHorizontalBox::Slot()
+				.FillWidth(1.0f)
+				[
+					GenerateTextWidgetFromName(QueryDataPtr->CollisionQueryParams.TraceTag)
+				]
+			];
+	}
+
+	if (ColumnName == SChaosVDSceneQueryTree::ColumnNames.TraceOwner)
+	{
+		return GenerateTextWidgetFromName(QueryDataPtr->CollisionQueryParams.OwnerTag);
+	}
+
+	if (ColumnName == SChaosVDSceneQueryTree::ColumnNames.QueryType)
+	{
+		return GenerateTextWidgetFromText(UEnum::GetDisplayValueAsText(QueryDataPtr->Type));
+	}
+
+	if (ColumnName == SChaosVDSceneQueryTree::ColumnNames.SolverName)
+	{
+		return GenerateTextWidgetFromName(Item->OwnerSolverName);
+	}
+
+	if (ColumnName == SChaosVDSceneQueryTree::ColumnNames.Visibility)
+	{
+		return SNew(SImage)
+				.IsEnabled(false)
+				.Visibility(this, &SChaosVDSceneQueryTreeRow::GetVisibilityIconVisibility)
+				.ColorAndOpacity(FSlateColor::UseForeground())
+				.ToolTipText(LOCTEXT("SceneUqreyBrowserItemVisibilityToolTip","Visibility is controlled by the Visualization Flags menu"))
+				.Image(this, &SChaosVDSceneQueryTreeRow::GetVisibilityIconForCurrentItem);
+	}
+
+	return SNullWidget::NullWidget;
+}
+
+EVisibility SChaosVDSceneQueryTreeRow::GetVisibilityIconVisibility() const
+{
+	// For now, we only want to show the visibility icon to indicate the item is hidden.
+	// TODO: Support Hover like the visibility widget in the scene outliner
+	return Item && Item->bIsVisible ? EVisibility::Hidden : EVisibility::Visible;
+}
+
+TSharedRef<SWidget> SChaosVDSceneQueryTreeRow::GenerateTextWidgetFromName(FName Name)
+{
+	return GenerateTextWidgetFromText(FText::FromName(Name));
+}
+
+TSharedRef<SWidget> SChaosVDSceneQueryTreeRow::GenerateTextWidgetFromText(const FText& Text)
+{
+	constexpr float MarginLeft = 4.0f;
+	constexpr float NoMargin = 0.0f;
+
+	return SNew(STextBlock)
+			.Margin(FMargin(MarginLeft,NoMargin,NoMargin,NoMargin))
+			.Text(Text);
+}
+
+const FSlateBrush* SChaosVDSceneQueryTreeRow::GetVisibilityIconForCurrentItem() const
+{
+	if (!Item)
+	{
+		return nullptr;
+	}
+
+	return Item->bIsVisible ? FAppStyle::Get().GetBrush(TEXT("Level.VisibleIcon16x")) : FAppStyle::Get().GetBrush(TEXT("Level.NotVisibleIcon16x"));
+}
+
+#undef LOCTEXT_NAMESPACE
\ No newline at end of file
diff --git a/Engine/Plugins/ChaosVD/Source/ChaosVD/Private/Widgets/SChaosVDSceneQueryTreeRow.h b/Engine/Plugins/ChaosVD/Source/ChaosVD/Private/Widgets/SChaosVDSceneQueryTreeRow.h
new file mode 100644
--- /dev/null
+++ b/Engine/Plugins/ChaosVD/Source/ChaosVD/Private/Widgets/SChaosVDSceneQueryTreeRow.h
@@ -0,0 +1,40 @@
+// Copyright Epic Games, Inc. All Rights Reserved.
+
+#pragma once
+
+#include "Widgets/SCompoundWidget.h"
+#include "Widgets/Views/SListView.h"
+#include "Widgets/Views/STableRow.h"
+
+class SChaosVDSceneQueryBrowser;
+struct FChaosVDQueryDataWrapper;
+struct FChaosVDSceneQueryTreeItem;
+
+/**
+ * Widget used to represent a row on the Scene Query Browser tree view
+ */
+class SChaosVDSceneQueryTreeRow : public SMultiColumnTableRow<TSharedPtr<const FChaosVDSceneQueryTreeItem>>
+{
+public:
+	SLATE_BEGIN_ARGS(SChaosVDSceneQueryTreeRow)
+		{
+		}
+		SLATE_ARGUMENT(TSharedPtr<FChaosVDSceneQueryTreeItem>, Item)
+
+	SLATE_END_ARGS()
+
+	void Construct(const FArguments& InArgs, const TSharedRef<STableViewBase>& InOwnerTableView);
+
+	virtual TSharedRef<SWidget> GenerateWidgetForColumn(const FName& ColumnName) override;
+
+private:
+
+	EVisibility GetVisibilityIconVisibility() const;
+	
+	TSharedRef<SWidget> GenerateTextWidgetFromName(FName Name);
+	TSharedRef<SWidget> GenerateTextWidgetFromText(const FText& Text);
+
+	const FSlateBrush* GetVisibilityIconForCurrentItem() const;
+
+	TSharedPtr<FChaosVDSceneQueryTreeItem> Item;
+};
