[ChaosVD] Changed default colors to be less harsh to the eyes.

We were using the same default colors as the runtime debug draw implementation. The issue is that CVD's Visualization is mostly solid shapes, so what it works great for wireframes, it is too bright and saturated for solid shapes.

In addition to chnage the defaults, this CL also:

- Reduces the Directional light from 10 to 4 Lux
- Added a postprocessing volume to enable an ambient light cubemap. This is needed so geometry that is shaded, it is not shaded completly dark which makes it difficult to see.

#virtualized

diff --git a/Engine/Plugins/ChaosVD/Config/DefaultChaosVD.ini b/Engine/Plugins/ChaosVD/Config/DefaultChaosVD.ini
--- a/Engine/Plugins/ChaosVD/Config/DefaultChaosVD.ini
+++ b/Engine/Plugins/ChaosVD/Config/DefaultChaosVD.ini
@@ -1,2 +1,8 @@
 [SectionsToSave]
-bCanSaveAllSections=true
\ No newline at end of file
+bCanSaveAllSections=true
+
+[/Script/ChaosVD.ChaosVDParticleVisualizationColorSettings]
+ParticleColorMode=State
+ColorsByShapeType=(SimpleTypeColor=(B=115,G=158,R=0,A=255),ConvexColor=(B=66,G=228,R=240,A=255),HeightFieldColor=(B=233,G=180,R=86,A=255),TriangleMeshColor=(B=0,G=94,R=213,A=255),LevelSetColor=(B=167,G=121,R=204,A=255))
+ColorsByParticleState=(DynamicColor=(B=98,G=246,R=253,A=255),SleepingColor=(B=80,G=92,R=231,A=255),KinematicColor=(B=178,G=114,R=0,A=255),StaticColor=(B=156,G=159,R=150,A=255))
+ColorsByClientServer=(ServerBaseColor=(B=80,G=92,R=231,A=255),ClientBaseColor=(B=178,G=114,R=0,A=255))
\ No newline at end of file
diff --git a/Engine/Plugins/ChaosVD/Config/Engine.ini b/Engine/Plugins/ChaosVD/Config/Engine.ini
--- a/Engine/Plugins/ChaosVD/Config/Engine.ini
+++ b/Engine/Plugins/ChaosVD/Config/Engine.ini
@@ -3,4 +3,5 @@ QueryOnlyMeshesMaterial=/ChaosVD/Materials/ChaosVDQueryOnlyMesh.ChaosVDQueryOnly
 SimOnlyMeshesMaterial=/ChaosVD/Materials/ChasoVDSimMesh.ChasoVDSimMesh
 InstancedMeshesMaterial=/ChaosVD/Materials/ChaosVDInstanced.ChaosVDInstanced
 InstancedMeshesQueryOnlyMaterial=/ChaosVD/Materials/ChaosVDInstancedQueryOnly.ChaosVDInstancedQueryOnly
-SkySphereActorClass=/ChaosVD/BP_CVD_SkySphere.BP_CVD_SkySphere_C
\ No newline at end of file
+SkySphereActorClass=/ChaosVD/BP_CVD_SkySphere.BP_CVD_SkySphere_C
+AmbientCubeMapTexture=/ChaosVD/Textures/DaylightAmbientCubemap.DaylightAmbientCubemap
\ No newline at end of file
diff --git a/Engine/Plugins/ChaosVD/Source/ChaosVD/Private/ChaosVDScene.cpp b/Engine/Plugins/ChaosVD/Source/ChaosVD/Private/ChaosVDScene.cpp
--- a/Engine/Plugins/ChaosVD/Source/ChaosVD/Private/ChaosVDScene.cpp
+++ b/Engine/Plugins/ChaosVD/Source/ChaosVD/Private/ChaosVDScene.cpp
@@ -31,6 +31,9 @@
 #include "Settings/ChaosVDCoreSettings.h"
 #include "UObject/Package.h"
 #include "Actors/ChaosVDGeometryContainer.h"
+#include "Components/LightComponent.h"
+#include "Engine/PostProcessVolume.h"
+#include "Engine/TextureCube.h"
 
 #define LOCTEXT_NAMESPACE "ChaosVisualDebugger"
 
@@ -63,8 +66,6 @@ void FChaosVDScene::Initialize()
 
 	InitializeSelectionSets();
 	
-	PhysicsVDWorld = CreatePhysicsVDWorld();
-	
 	StreamableManager = MakeShared<FStreamableManager>();
 
 	if (UChaosVDCoreSettings* Settings = FChaosVDSettingsManager::Get().GetSettingsObject<UChaosVDCoreSettings>())
@@ -75,7 +76,10 @@ void FChaosVDScene::Initialize()
 		StreamableManager->RequestSyncLoad(Settings->SimOnlyMeshesMaterial.ToSoftObjectPath());
 		StreamableManager->RequestSyncLoad(Settings->InstancedMeshesMaterial.ToSoftObjectPath());
 		StreamableManager->RequestSyncLoad(Settings->InstancedMeshesQueryOnlyMaterial.ToSoftObjectPath());
+		StreamableManager->RequestSyncLoad(Settings->AmbientCubeMapTexture.ToSoftObjectPath());
 	}
+	
+	PhysicsVDWorld = CreatePhysicsVDWorld();
 
 	GeometryGenerator = MakeShared<FChaosVDGeometryBuilder>();
 
@@ -558,6 +562,8 @@ void FChaosVDScene::CreateBaseLights(UWorld* TargetWorld) const
 			DirectionalLightActor->SetCastShadows(false);
 			DirectionalLightActor->SetMobility(EComponentMobility::Movable);
 			DirectionalLightActor->SetActorLocation(SpawnPosition);
+			
+			DirectionalLightActor->SetBrightness(4.0f);
 
 			DirectionalLightActor->SetFolderPath(LightingFolderPath);
 
@@ -577,6 +583,32 @@ void FChaosVDScene::CreateBaseLights(UWorld* TargetWorld) const
 	}
 }
 
+void FChaosVDScene::CreatePostProcessingVolumes(UWorld* TargetWorld)
+{
+	const FName LightingFolderPath("ChaosVisualDebugger/Lighting");
+
+	if (const UChaosVDCoreSettings* Settings = FChaosVDSettingsManager::Get().GetSettingsObject<UChaosVDCoreSettings>())
+	{
+		APostProcessVolume* PostProcessingVolume = TargetWorld->SpawnActor<APostProcessVolume>();
+		if (ensure(PostProcessingVolume))
+		{
+			PostProcessingVolume->SetFolderPath(LightingFolderPath);
+			PostProcessingVolume->Settings.bOverride_AmbientCubemapIntensity = true;
+			PostProcessingVolume->Settings.AmbientCubemapIntensity = 0.3f;
+			PostProcessingVolume->bUnbound = true;
+			PostProcessingVolume->bEnabled = true;
+
+			UTextureCube* AmbientCubemap = Settings->AmbientCubeMapTexture.Get();
+			if (ensure(AmbientCubemap))
+			{
+				PostProcessingVolume->Settings.AmbientCubemap = AmbientCubemap;
+			}
+			
+			PostProcessingVolume->MarkComponentsRenderStateDirty();
+		}
+	}
+}
+
 AActor* FChaosVDScene::CreateMeshComponentsContainer(UWorld* TargetWorld)
 {
 	const FName GeometryFolderPath("ChaosVisualDebugger/GeneratedMeshComponents");
@@ -614,6 +646,7 @@ UWorld* FChaosVDScene::CreatePhysicsVDWorld()
 
 	CreateBaseLights(NewWorld);
 	CreateMeshComponentsContainer(NewWorld);
+	CreatePostProcessingVolumes(NewWorld);
 
 	ActorDestroyedHandle = NewWorld->AddOnActorDestroyedHandler(FOnActorDestroyed::FDelegate::CreateRaw(this, &FChaosVDScene::HandleActorDestroyed));
 	
diff --git a/Engine/Plugins/ChaosVD/Source/ChaosVD/Private/ChaosVDScene.h b/Engine/Plugins/ChaosVD/Source/ChaosVD/Private/ChaosVDScene.h
--- a/Engine/Plugins/ChaosVD/Source/ChaosVD/Private/ChaosVDScene.h
+++ b/Engine/Plugins/ChaosVD/Source/ChaosVD/Private/ChaosVDScene.h
@@ -162,6 +162,8 @@ private:
 
 	void CreateBaseLights(UWorld* TargetWorld) const;
 
+	void CreatePostProcessingVolumes(UWorld* TargetWorld);
+
 	/** Creates an actor that will contain all solver data for the provided Solver ID*/
 	AChaosVDSolverInfoActor* GetOrCreateSolverInfoActor(int32 SolverID);
 
diff --git a/Engine/Plugins/ChaosVD/Source/ChaosVD/Private/Settings/ChaosVDCoreSettings.h b/Engine/Plugins/ChaosVD/Source/ChaosVD/Private/Settings/ChaosVDCoreSettings.h
--- a/Engine/Plugins/ChaosVD/Source/ChaosVD/Private/Settings/ChaosVDCoreSettings.h
+++ b/Engine/Plugins/ChaosVD/Source/ChaosVD/Private/Settings/ChaosVDCoreSettings.h
@@ -11,6 +11,7 @@
 
 class UChaosVDCoreSettings;
 class UMaterial;
+class UTextureCube;
 
 DECLARE_MULTICAST_DELEGATE_OneParam(FChaosVDSettingChanged, UObject* SettingsObject)
 
@@ -84,4 +85,7 @@ public:
 
 	UPROPERTY(Config)
 	FSoftClassPath SkySphereActorClass;
+
+	UPROPERTY(Config)
+	TSoftObjectPtr<UTextureCube> AmbientCubeMapTexture;
 };
diff --git a/Engine/Plugins/ChaosVD/Source/ChaosVD/Private/Settings/ChaosVDParticleVisualizationSettings.cpp b/Engine/Plugins/ChaosVD/Source/ChaosVD/Private/Settings/ChaosVDParticleVisualizationSettings.cpp
--- a/Engine/Plugins/ChaosVD/Source/ChaosVD/Private/Settings/ChaosVDParticleVisualizationSettings.cpp
+++ b/Engine/Plugins/ChaosVD/Source/ChaosVD/Private/Settings/ChaosVDParticleVisualizationSettings.cpp
@@ -154,17 +154,26 @@ FColor FChaosDebugDrawColorsByShapeType::GetColorFromShapeType(Chaos::EImplicitO
 
 FColor FChaosDebugDrawColorsByClientServer::GetColorFromState(bool bIsServer, EChaosVDObjectStateType State) const
 {
-	switch (State)
+	if (State == EChaosVDObjectStateType::Uninitialized)
 	{
-	case EChaosVDObjectStateType::Sleeping:
-		return bIsServer ? ServerSleepingColor : ClientSleepingColor;
-	case EChaosVDObjectStateType::Kinematic:
-		return bIsServer ? ServerColor : ClientColor;
-	case EChaosVDObjectStateType::Static:
-		return bIsServer ? ServerColor : ClientColor;
-	case EChaosVDObjectStateType::Dynamic:
-		return bIsServer ? ServerDynamicColor : ClientDynamicColor;
-	default:
 		return FColor::Purple;
 	}
+
+	constexpr float IntensityFactor = 1.0f / static_cast<float>(EChaosVDObjectStateType::Count);
+
+	// Make sure static is always darker than sleeping
+	if (State == EChaosVDObjectStateType::Static)
+	{
+
+		constexpr float StaticStateIntensity = IntensityFactor * (static_cast<float>(EChaosVDObjectStateType::Sleeping) * 0.2f);
+		return GetColorAtIntensity(bIsServer ? ServerBaseColor : ClientBaseColor, StaticStateIntensity);
+	}
+
+	const float Intensity = IntensityFactor * static_cast<int32>(State);
+	return GetColorAtIntensity(bIsServer ? ServerBaseColor : ClientBaseColor, Intensity);
+}
+
+FColor FChaosDebugDrawColorsByClientServer::GetColorAtIntensity(const FColor& InColor, float Intensity) const
+{
+	return (FLinearColor(InColor) * Intensity).ToFColor(true);
 }
diff --git a/Engine/Plugins/ChaosVD/Source/ChaosVD/Private/Settings/ChaosVDParticleVisualizationSettings.h b/Engine/Plugins/ChaosVD/Source/ChaosVD/Private/Settings/ChaosVDParticleVisualizationSettings.h
--- a/Engine/Plugins/ChaosVD/Source/ChaosVD/Private/Settings/ChaosVDParticleVisualizationSettings.h
+++ b/Engine/Plugins/ChaosVD/Source/ChaosVD/Private/Settings/ChaosVDParticleVisualizationSettings.h
@@ -80,32 +80,17 @@ USTRUCT()
 struct FChaosDebugDrawColorsByClientServer
 {
 	GENERATED_BODY()
+	
+	/** Color used for server shapes that are not awake or sleeping dynamic */
+    UPROPERTY(config, EditAnywhere, Category=DebugDraw)
+    FColor ServerBaseColor = FColor(50, 0, 0);
 
 	/** Color used for server shapes that are not awake or sleeping dynamic */
 	UPROPERTY(config, EditAnywhere, Category=DebugDraw)
-	FColor ServerColor = FColor(50, 0, 0); 
-
-	/** Color used for server shapes that are awake dynamic */
-	UPROPERTY(config, EditAnywhere, Category = DebugDraw)
-	FColor ServerDynamicColor = FColor(150, 0, 0);
-
-	/** Color used for server shapes that are sleeping dynamics */
-	UPROPERTY(config, EditAnywhere, Category = DebugDraw)
-	FColor ServerSleepingColor = FColor(10, 0, 0);
-
-	/** Color used for client shapes that are not awake or sleeping dynamic */
-	UPROPERTY(config, EditAnywhere, Category = DebugDraw)
-	FColor ClientColor = FColor(0, 0, 50);
-
-	/** Color used for server shapes that are awake dynamic */
-	UPROPERTY(config, EditAnywhere, Category = DebugDraw)
-	FColor ClientDynamicColor = FColor(0, 0, 150);
-
-	/** Color used for client shapes that are sleeping dynamics */
-	UPROPERTY(config, EditAnywhere, Category = DebugDraw)
-	FColor ClientSleepingColor = FColor(0, 0, 100);
+	FColor ClientBaseColor = FColor(0, 0, 50); 
 
 	FColor GetColorFromState(bool bIsServer, EChaosVDObjectStateType State) const;
+	FColor GetColorAtIntensity(const FColor& InColor, float Intensity) const;
 };
 
 UENUM()
