[ChaosVD] Fix Static particles being too dark on the Color by Client/Server mode.

When I tested the intensity factor for my previous CL, I was using the old light intencity, which was set to be twice as bright.

Long term this probably needs to be implemented with a curve that can be changed via settings instead of a fixed C++ value.

diff --git a/Engine/Plugins/ChaosVD/Source/ChaosVD/Private/Settings/ChaosVDParticleVisualizationSettings.cpp b/Engine/Plugins/ChaosVD/Source/ChaosVD/Private/Settings/ChaosVDParticleVisualizationSettings.cpp
--- a/Engine/Plugins/ChaosVD/Source/ChaosVD/Private/Settings/ChaosVDParticleVisualizationSettings.cpp
+++ b/Engine/Plugins/ChaosVD/Source/ChaosVD/Private/Settings/ChaosVDParticleVisualizationSettings.cpp
@@ -165,7 +165,7 @@ FColor FChaosDebugDrawColorsByClientServer::GetColorFromState(bool bIsServer, EC
 	if (State == EChaosVDObjectStateType::Static)
 	{
 
-		constexpr float StaticStateIntensity = IntensityFactor * (static_cast<float>(EChaosVDObjectStateType::Sleeping) * 0.2f);
+		constexpr float StaticStateIntensity = IntensityFactor * (static_cast<float>(EChaosVDObjectStateType::Sleeping) * 0.6f);
 		return GetColorAtIntensity(bIsServer ? ServerBaseColor : ClientBaseColor, StaticStateIntensity);
 	}
 
