[ChaosVD] Fix for issue where object might incorrectly be unloaded if certain conditions are met

This is fallout of the fix for UE-214135. As part of that fix, now when a CVD file is loaded, if we get particle (or other) data outside a Solver Stage scope, we auto generate a Solver Stage entry named "In-Between" stage.

There are valid situations where we get data outside a solver stage, like static particles creation. The issue is that CVD assumes a keyframe will have all particles, therefore if a particle is no longer seen it will proceed to remove it. This is only done with keyframes.

The issue here is if we get data after the Last recorded Solver Stage for a Frame was closed, AND that frame is a key frame. The auto-generated solver stage will have a lot of data missing as it wasn't really part of the keyframe.

The fix is simple, we just need to take into account the Explict Solver stage flag in addition to the Key Frame tag when determining if we need to remove particles from the visualization.

If a particle destroyed event is recorded, that will still be respected.

diff --git a/Engine/Plugins/ChaosVD/Source/ChaosVD/Private/ChaosVDScene.cpp b/Engine/Plugins/ChaosVD/Source/ChaosVD/Private/ChaosVDScene.cpp
--- a/Engine/Plugins/ChaosVD/Source/ChaosVD/Private/ChaosVDScene.cpp
+++ b/Engine/Plugins/ChaosVD/Source/ChaosVD/Private/ChaosVDScene.cpp
@@ -222,7 +222,7 @@ void FChaosVDScene::UpdateFromRecordedStepData(const int32 SolverID, const FChao
 	for (const TPair<int32, TObjectPtr<AChaosVDParticleActor>>& ParticleActorWithID : AllSolverParticlesByID)
 	{
 		// If we are playing back a keyframe, the scene should only contain what it is in the recorded data
-		const bool bShouldDestroyParticleAnyway = InFrameData.bIsKeyFrame && !ParticlesIDsInRecordedStepData.Contains(ParticleActorWithID.Key);
+		const bool bShouldDestroyParticleAnyway = InFrameData.bIsKeyFrame && EnumHasAnyFlags(InRecordedStepData.StageFlags, EChaosVDSolverStageFlags::ExplicitStage) && !ParticlesIDsInRecordedStepData.Contains(ParticleActorWithID.Key);
 		
 		if (bShouldDestroyParticleAnyway || InFrameData.ParticlesDestroyedIDs.Contains(ParticleActorWithID.Key))
 		{
