[ChaosVD] Added a simple Recent Files menu and changed the color of the polygon edge lines to white for when a wireframe view mode is enabled.

Also, added the new Viewmode Lit Wireframe that was added to the normal level editor (It is the same as the old menu where you had to enable the mesh edge show flag from the show menu, but now is accesible in two clicks)

diff --git a/Engine/Plugins/ChaosVD/Source/ChaosVD/Private/ChaosVDEngine.cpp b/Engine/Plugins/ChaosVD/Source/ChaosVD/Private/ChaosVDEngine.cpp
--- a/Engine/Plugins/ChaosVD/Source/ChaosVD/Private/ChaosVDEngine.cpp
+++ b/Engine/Plugins/ChaosVD/Source/ChaosVD/Private/ChaosVDEngine.cpp
@@ -6,6 +6,8 @@
 #include "ChaosVDPlaybackController.h"
 #include "ChaosVDRuntimeModule.h"
 #include "ChaosVDScene.h"
+#include "ChaosVDSettingsManager.h"
+#include "Settings/ChaosVDMiscSettings.h"
 #include "Trace/ChaosVDTraceManager.h"
 
 void FChaosVDEngine::Initialize()
@@ -64,12 +66,46 @@ void FChaosVDEngine::DeInitialize()
 	bIsInitialized = false;
 }
 
+void FChaosVDEngine::UpdateRecentFilesList(const FString& InFilename)
+{
+	UChaosVDMiscSettings* MiscSettings = FChaosVDSettingsManager::Get().GetSettingsObject<UChaosVDMiscSettings>();
+	if (!MiscSettings)
+	{
+		return;
+	}
+
+	FDateTime CurrentTime = FDateTime::UtcNow();
+	if (FChaosVDRecentFile* RecentProject = MiscSettings->RecentFiles.FindByKey(InFilename))
+	{
+		RecentProject->LastOpenTime = CurrentTime;
+	}
+	else
+	{
+		FChaosVDRecentFile RecentFileEntry(InFilename, CurrentTime);
+		MiscSettings->RecentFiles.Emplace(RecentFileEntry);
+	}
+	
+	MiscSettings->RecentFiles.Sort(FChaosVDRecentFile::FRecentFilesSortPredicate());
+
+	if (MiscSettings->RecentFiles.Num() > MiscSettings->MaxRecentFilesNum)
+	{
+		MiscSettings->RecentFiles.SetNum(MiscSettings->MaxRecentFilesNum);
+	}
+
+	MiscSettings->SaveConfig();
+}
+
 void FChaosVDEngine::LoadRecording(const FString& FilePath)
 {
 	FChaosVDTraceSessionDescriptor NewSessionFromFileDescriptor;
 	NewSessionFromFileDescriptor.SessionName = FChaosVDModule::Get().GetTraceManager()->LoadTraceFile(FilePath);
 	NewSessionFromFileDescriptor.bIsLiveSession = false;
 
+	if (NewSessionFromFileDescriptor.IsValid())
+	{
+		UpdateRecentFilesList(FilePath);
+	}
+
 	SetCurrentSession(NewSessionFromFileDescriptor);
 }
 
diff --git a/Engine/Plugins/ChaosVD/Source/ChaosVD/Private/ChaosVDEngine.h b/Engine/Plugins/ChaosVD/Source/ChaosVD/Private/ChaosVDEngine.h
--- a/Engine/Plugins/ChaosVD/Source/ChaosVD/Private/ChaosVDEngine.h
+++ b/Engine/Plugins/ChaosVD/Source/ChaosVD/Private/ChaosVDEngine.h
@@ -39,6 +39,8 @@ public:
 	
 private:
 
+	void UpdateRecentFilesList(const FString& InFilename);
+
 	FGuid InstanceGUID;
 
 	FChaosVDTraceSessionDescriptor CurrentSessionDescriptor;
diff --git a/Engine/Plugins/ChaosVD/Source/ChaosVD/Private/ChaosVDPlaybackController.cpp b/Engine/Plugins/ChaosVD/Source/ChaosVD/Private/ChaosVDPlaybackController.cpp
--- a/Engine/Plugins/ChaosVD/Source/ChaosVD/Private/ChaosVDPlaybackController.cpp
+++ b/Engine/Plugins/ChaosVD/Source/ChaosVD/Private/ChaosVDPlaybackController.cpp
@@ -53,7 +53,7 @@ FChaosVDPlaybackController::~FChaosVDPlaybackController()
 
 bool FChaosVDPlaybackController::LoadChaosVDRecordingFromTraceSession(const FChaosVDTraceSessionDescriptor& InSessionDescriptor)
 {
-	if (!ensure(!InSessionDescriptor.SessionName.IsEmpty()))
+	if (!ensure(InSessionDescriptor.IsValid()))
 	{
 		return false;
 	}
diff --git a/Engine/Plugins/ChaosVD/Source/ChaosVD/Private/ChaosVDScene.cpp b/Engine/Plugins/ChaosVD/Source/ChaosVD/Private/ChaosVDScene.cpp
--- a/Engine/Plugins/ChaosVD/Source/ChaosVD/Private/ChaosVDScene.cpp
+++ b/Engine/Plugins/ChaosVD/Source/ChaosVD/Private/ChaosVDScene.cpp
@@ -573,11 +573,24 @@ void FChaosVDScene::CreateBaseLights(UWorld* TargetWorld) const
 			{
 				SkySphere->SetActorLocation(SpawnPosition);
 				SkySphere->SetFolderPath(LightingFolderPath);
+				
 				if (SkySphere->Implements<UChaosVDSkySphereInterface>())
 				{
 					FEditorScriptExecutionGuard AllowEditorScriptGuard;
 					IChaosVDSkySphereInterface::Execute_SetDirectionalLightSource(SkySphere, DirectionalLightActor);
 				}
+
+				// Keep it dark to reduce visual noise.
+				// TODO: We should hide these components altogether when we switch to a unlit wireframe mode 
+				const TSet<UActorComponent*>& Components = SkySphere->GetComponents();
+				for (UActorComponent* Component : Components)
+				{
+					if (UStaticMeshComponent* AsStaticMeshComponent = Cast<UStaticMeshComponent>(Component))
+					{
+						AsStaticMeshComponent->bOverrideWireframeColor = true;
+						AsStaticMeshComponent->WireframeColorOverride = FColor::Black;
+					}
+				}
 			}
 		}
 	}
diff --git a/Engine/Plugins/ChaosVD/Source/ChaosVD/Private/Components/ChaosVDInstancedStaticMeshComponent.h b/Engine/Plugins/ChaosVD/Source/ChaosVD/Private/Components/ChaosVDInstancedStaticMeshComponent.h
--- a/Engine/Plugins/ChaosVD/Source/ChaosVD/Private/Components/ChaosVDInstancedStaticMeshComponent.h
+++ b/Engine/Plugins/ChaosVD/Source/ChaosVD/Private/Components/ChaosVDInstancedStaticMeshComponent.h
@@ -19,6 +19,9 @@ public:
 		SetRemoveSwap();
 		SetCanEverAffectNavigation(false);
 		bHasPerInstanceHitProxies = true;
+
+		bOverrideWireframeColor = true;
+		WireframeColorOverride = FColor::White;
 	}
 	
 	// BEGIN IChaosVDGeometryDataComponent Interface
diff --git a/Engine/Plugins/ChaosVD/Source/ChaosVD/Private/Components/ChaosVDStaticMeshComponent.h b/Engine/Plugins/ChaosVD/Source/ChaosVD/Private/Components/ChaosVDStaticMeshComponent.h
--- a/Engine/Plugins/ChaosVD/Source/ChaosVD/Private/Components/ChaosVDStaticMeshComponent.h
+++ b/Engine/Plugins/ChaosVD/Source/ChaosVD/Private/Components/ChaosVDStaticMeshComponent.h
@@ -19,6 +19,8 @@ public:
 	{
 		SetCanEverAffectNavigation(false);
 		bNavigationRelevant = false;
+		bOverrideWireframeColor = true;
+		WireframeColorOverride = FColor::White;
 	}
 
 	// BEGIN IChaosVDGeometryDataComponent Interface
diff --git a/Engine/Plugins/ChaosVD/Source/ChaosVD/Private/Settings/ChaosVDMiscSettings.h b/Engine/Plugins/ChaosVD/Source/ChaosVD/Private/Settings/ChaosVDMiscSettings.h
new file mode 100644
--- /dev/null
+++ b/Engine/Plugins/ChaosVD/Source/ChaosVD/Private/Settings/ChaosVDMiscSettings.h
@@ -0,0 +1,70 @@
+// Copyright Epic Games, Inc. All Rights Reserved.
+
+#pragma once
+
+#include "ChaosVDCoreSettings.h"
+#include "ChaosVDMiscSettings.generated.h"
+
+/** Structure with the details about a recently open CVD file */
+USTRUCT()
+struct FChaosVDRecentFile
+{
+	GENERATED_BODY()
+
+	/** Path to the CVD File */
+	UPROPERTY(config)
+	FString FileName;
+
+	/** Timestamp of the last time the CVD opened this file */
+	UPROPERTY(config)
+	FDateTime LastOpenTime;
+
+	FChaosVDRecentFile()
+	{}
+
+	FChaosVDRecentFile(const FString& InProjectName, FDateTime InLastOpenTime)
+		: FileName(InProjectName)
+		, LastOpenTime(InLastOpenTime)
+	{}
+
+	struct FRecentFilesSortPredicate
+	{
+		bool operator()( const FChaosVDRecentFile& A, const FChaosVDRecentFile& B ) const
+		{
+			return B.LastOpenTime < A.LastOpenTime;
+		}
+	};
+
+	bool operator<(const FChaosVDRecentFile& Other) const
+	{
+		return LastOpenTime < Other.LastOpenTime;
+	}
+
+	bool operator==(const FChaosVDRecentFile& Other) const
+	{
+		return FileName == Other.FileName;
+	}
+
+	bool operator==(const FString& OtherProjectName) const
+	{
+		return FileName == OtherProjectName;
+	}
+};
+
+/**
+ * General non-core CVD settings
+ */
+UCLASS(config=ChaosVD, PerObjectConfig)
+class CHAOSVD_API UChaosVDMiscSettings : public UChaosVDSettingsObjectBase
+{
+	GENERATED_BODY()
+
+public:
+	/** List of recently opened files */
+	UPROPERTY(config)
+	TArray<FChaosVDRecentFile> RecentFiles;
+
+	/** Maximum number of recent files we can keep track of*/
+	UPROPERTY(config)
+	int32 MaxRecentFilesNum = 10;
+};
diff --git a/Engine/Plugins/ChaosVD/Source/ChaosVD/Private/Trace/ChaosVDTraceManager.h b/Engine/Plugins/ChaosVD/Source/ChaosVD/Private/Trace/ChaosVDTraceManager.h
--- a/Engine/Plugins/ChaosVD/Source/ChaosVD/Private/Trace/ChaosVDTraceManager.h
+++ b/Engine/Plugins/ChaosVD/Source/ChaosVD/Private/Trace/ChaosVDTraceManager.h
@@ -19,6 +19,8 @@ struct FChaosVDTraceSessionDescriptor
 {
 	FString SessionName;
 	bool bIsLiveSession = false;
+
+	bool IsValid() const { return !SessionName.IsEmpty(); }
 };
 
 /** Manager class used by Chaos VD to interact/control UE Trace systems */
diff --git a/Engine/Plugins/ChaosVD/Source/ChaosVD/Private/Widgets/SChaosVDEditorViewportViewMenu.cpp b/Engine/Plugins/ChaosVD/Source/ChaosVD/Private/Widgets/SChaosVDEditorViewportViewMenu.cpp
--- a/Engine/Plugins/ChaosVD/Source/ChaosVD/Private/Widgets/SChaosVDEditorViewportViewMenu.cpp
+++ b/Engine/Plugins/ChaosVD/Source/ChaosVD/Private/Widgets/SChaosVDEditorViewportViewMenu.cpp
@@ -32,6 +32,7 @@ void SChaosVDEditorViewportViewMenu::RegisterMenus() const
 					{
 						Section.AddMenuEntry(BaseViewportActions.LitMode, UViewModeUtils::GetViewModeDisplayName(VMI_Lit));
 						Section.AddMenuEntry(BaseViewportActions.UnlitMode, UViewModeUtils::GetViewModeDisplayName(VMI_Unlit));
+						Section.AddMenuEntry(BaseViewportActions.LitWireframeMode, UViewModeUtils::GetViewModeDisplayName(VMI_Lit_Wireframe));
 						Section.AddMenuEntry(BaseViewportActions.WireframeMode, UViewModeUtils::GetViewModeDisplayName(VMI_BrushWireframe));
 					}
 				}
diff --git a/Engine/Plugins/ChaosVD/Source/ChaosVD/Private/Widgets/SChaosVDMainTab.cpp b/Engine/Plugins/ChaosVD/Source/ChaosVD/Private/Widgets/SChaosVDMainTab.cpp
--- a/Engine/Plugins/ChaosVD/Source/ChaosVD/Private/Widgets/SChaosVDMainTab.cpp
+++ b/Engine/Plugins/ChaosVD/Source/ChaosVD/Private/Widgets/SChaosVDMainTab.cpp
@@ -36,6 +36,8 @@
 #include "Components/ChaosVDSolverCharacterGroundConstraintDataComponent.h"
 #include "Components/ChaosVDSolverCollisionDataComponent.h"
 #include "Components/ChaosVDSolverJointConstraintDataComponent.h"
+#include "HAL/FileManager.h"
+#include "Settings/ChaosVDMiscSettings.h"
 #include "Styling/StyleColors.h"
 #include "Styling/ToolBarStyle.h"
 #include "TabSpawers/ChaosVDSceneQueryBrowserTab.h"
@@ -314,14 +316,9 @@ void SChaosVDMainTab::GenerateMainWindowMenu()
 	MenuBarBuilder.AddPullDownMenu(
 		LOCTEXT("FileMenuLabel", "File"),
 		FText::GetEmpty(),
-		FNewMenuDelegate::CreateLambda([this](FMenuBuilder& MenuBuilder)
+		FNewMenuDelegate::CreateSPLambda(this, [this](FMenuBuilder& MenuBuilder)
 		{
-			MenuBuilder.AddMenuEntry(LOCTEXT("OpenFileMenuLabel", "OpenFile"), FText(), FSlateIcon(),
-			FUIAction(
-				FExecuteAction::CreateLambda([this]()
-				{
-					BrowseAndOpenChaosVDRecording();
-				})));
+			MenuBuilder.AddSubMenu(LOCTEXT("RecentFilesMenuLabel", "Recent Files"), LOCTEXT("RecentFilesMenuLabelToolTip", "Shows a list of recently used CVD Files"), FNewMenuDelegate::CreateSP(this, &SChaosVDMainTab::GenerateRecentFilesMenu));
 		}),
 		"File"
 	);
@@ -339,6 +336,22 @@ void SChaosVDMainTab::GenerateMainWindowMenu()
 	TabManager->SetMenuMultiBox(MenuBarBuilder.GetMultiBox(), MenuBarBuilder.MakeWidget());
 }
 
+void SChaosVDMainTab::GenerateRecentFilesMenu(FMenuBuilder& MenuBuilder)
+{
+	if (UChaosVDMiscSettings* MiscSettings = FChaosVDSettingsManager::Get().GetSettingsObject<UChaosVDMiscSettings>())
+	{
+		MiscSettings->RecentFiles.Sort(FChaosVDRecentFile::FRecentFilesSortPredicate());
+
+		for (const FChaosVDRecentFile& RecentFile : MiscSettings->RecentFiles)
+		{
+			const FText DisplayName = FText::FromString(FPaths::GetBaseFilename(*RecentFile.FileName));
+			const FText Tooltip = FText::FromString(IFileManager::Get().ConvertToAbsolutePathForExternalAppForRead(*RecentFile.FileName) );
+			MenuBuilder.AddMenuEntry(DisplayName, Tooltip, FSlateIcon(FChaosVDStyle::Get().GetStyleSetName(), "OpenFileIcon"), FUIAction(FExecuteAction::CreateSPLambda(this, [this, FileNameCopy = RecentFile.FileName](){ LoadCVDFile(FileNameCopy); })),
+					NAME_None, EUserInterfaceActionType::Button);
+		}
+	}
+}
+
 FReply SChaosVDMainTab::BrowseAndOpenChaosVDRecording()
 {
 	const TSharedRef<SChaosBrowseTraceFileSourceModal> SessionBrowserModal = SNew(SChaosBrowseTraceFileSourceModal);
@@ -427,6 +440,18 @@ TSharedRef<SWidget> SChaosVDMainTab::GenerateMainToolbarWidget()
 	return UToolMenus::Get()->GenerateWidget(MainToolBarName, MenuContext);
 }
 
+void SChaosVDMainTab::LoadCVDFile(const FString& InFilename)
+{
+	if (ensure(InFilename.EndsWith(TEXT("utrace"))))
+	{
+		GetChaosVDEngineInstance()->LoadRecording(InFilename);
+	}
+	else
+	{
+		UE_LOG(LogChaosVDEditor, Error, TEXT("[%s] Invalid file extension | Only UTrace files are supported | Filename [%s]"), ANSI_TO_TCHAR(__FUNCTION__), *InFilename)			
+	}
+}
+
 void SChaosVDMainTab::BrowseChaosVDRecordingFromFolder(FStringView FolderPath)
 {
 	TArray<FString> OutOpenFilenames;
@@ -450,10 +475,7 @@ void SChaosVDMainTab::BrowseChaosVDRecordingFromFolder(FStringView FolderPath)
 
 	if (OutOpenFilenames.Num() > 0)
 	{
-		if (OutOpenFilenames[0].EndsWith(TEXT("utrace")))
-		{
-			GetChaosVDEngineInstance()->LoadRecording(OutOpenFilenames[0]);
-		}
+		LoadCVDFile(OutOpenFilenames[0]);
 	}
 }
 
@@ -465,7 +487,7 @@ bool SChaosVDMainTab::ConnectToLiveSession(int32 SessionID, const FString Sessio
 
 	bool bSuccess = false;
 
-	if (NewSessionFromFileDescriptor.SessionName.IsEmpty())
+	if (!NewSessionFromFileDescriptor.IsValid())
 	{
 		// If it failed we want to clean the current session name, so it is ok calling it either way
 		GetChaosVDEngineInstance()->SetCurrentSession(FChaosVDTraceSessionDescriptor());
diff --git a/Engine/Plugins/ChaosVD/Source/ChaosVD/Private/Widgets/SChaosVDMainTab.h b/Engine/Plugins/ChaosVD/Source/ChaosVD/Private/Widgets/SChaosVDMainTab.h
--- a/Engine/Plugins/ChaosVD/Source/ChaosVD/Private/Widgets/SChaosVDMainTab.h
+++ b/Engine/Plugins/ChaosVD/Source/ChaosVD/Private/Widgets/SChaosVDMainTab.h
@@ -72,6 +72,8 @@ public:
 	
 	bool ConnectToLiveSession(int32 SessionID, FString SessionAddress) const;
 
+	void LoadCVDFile(const FString& InFilename);
+
 private:
 	
 	void SetUpDisableCPUThrottlingDelegate();
@@ -90,6 +92,7 @@ private:
 	TSharedRef<FTabManager::FLayout> GenerateMainLayout();
 
 	void GenerateMainWindowMenu();
+	void GenerateRecentFilesMenu(FMenuBuilder& MenuBuilder);
 
 	FReply BrowseAndOpenChaosVDRecording();
 
diff --git a/Engine/Plugins/ChaosVD/Source/ChaosVD/Private/Widgets/SChaosVDViewportToolbar.cpp b/Engine/Plugins/ChaosVD/Source/ChaosVD/Private/Widgets/SChaosVDViewportToolbar.cpp
--- a/Engine/Plugins/ChaosVD/Source/ChaosVD/Private/Widgets/SChaosVDViewportToolbar.cpp
+++ b/Engine/Plugins/ChaosVD/Source/ChaosVD/Private/Widgets/SChaosVDViewportToolbar.cpp
@@ -228,7 +228,6 @@ TSharedRef<SWidget> SChaosVDViewportToolbar::GenerateShowMenu() const
 																		.IncludeFlag(FEngineShowFlags::SF_AntiAliasing)
 																		.IncludeFlag(FEngineShowFlags::SF_Grid)
 																		.IncludeFlag(FEngineShowFlags::SF_Translucency)
-																		.IncludeFlag(FEngineShowFlags::SF_MeshEdges)
 																		.IncludeFlag(FEngineShowFlags::SF_HitProxies)
 																		.IncludeFlag(FEngineShowFlags::SF_Fog)
 																		.IncludeFlag(FEngineShowFlags::SF_Pivot);
