[ChaosVD] Scene Query Data recording and inspection improvements

- Added supporting code to record the reject reason for when a shape is evaluated but not processed as a hit (this new data is not serialized yet, that will be done in a followup CL as it requires a version bump).

- Created a customization to process the recorded Sim and Query filter data and present  it like the collision response matrix used in the rest of the editor. This also handles conversion from Chaos formatted flags into the UE formatted flags (like the Collision Enabled enum).

- Created customization for the scene query collision response params to show the recorded flags array as the collision response matrix used in the rest of the editor.

- Created customization for the scene query object response params to show the recorded integer as a list of flags per object type

- Made Scene query visualization to use World priority instead of foreground by default, and exposed that as a setting on the existing visualization settings menu for scene queries

- Fixed an issue where the query mode in CVD was reported incorrectly (off by one in the mode enum)

- Fixed an issue where the selected particle details panel might not updated its values if the geometry shape data mode is selected

- Fixed an issue introduced in a previuos CL, where if you scrub with a delta of +1 frame, some destroyed particles might never be removed

diff --git a/Engine/Plugins/ChaosVD/Source/ChaosVD/Private/ChaosVDCollisionDataDetailsTab.cpp b/Engine/Plugins/ChaosVD/Source/ChaosVD/Private/ChaosVDCollisionDataDetailsTab.cpp
--- a/Engine/Plugins/ChaosVD/Source/ChaosVD/Private/ChaosVDCollisionDataDetailsTab.cpp
+++ b/Engine/Plugins/ChaosVD/Source/ChaosVD/Private/ChaosVDCollisionDataDetailsTab.cpp
@@ -21,10 +21,17 @@ TSharedRef<SDockTab> FChaosVDCollisionDataDetailsTab::HandleTabSpawnRequest(cons
 		.Label(LOCTEXT("CollisionInspectorTab", "Collision Data Inspector"))
 		.ToolTipText(LOCTEXT("CollisionInspectorTabTip", "See the details of the any collision data for the selected object"));
 
-	DetailsPanelTab->SetContent
-	(
-		SAssignNew(CollisionDataInspector, SChaosVDCollisionDataInspector, GetChaosVDScene())
-	);
+	if (const TSharedPtr<SChaosVDMainTab> MainTabPtr = OwningTabWidget.Pin())
+	{
+		DetailsPanelTab->SetContent
+		(
+			SAssignNew(CollisionDataInspector, SChaosVDCollisionDataInspector, GetChaosVDScene(), MainTabPtr.ToSharedRef())
+		);
+	}
+	else
+	{
+		DetailsPanelTab->SetContent(GenerateErrorWidget());
+	}
 
 	DetailsPanelTab->SetTabIcon(FSlateIcon(FAppStyle::GetAppStyleSetName(), "CollisionAnalyzer.TabIcon").GetIcon());
 
diff --git a/Engine/Plugins/ChaosVD/Source/ChaosVD/Private/ChaosVDConstraintDataInspectorTab.cpp b/Engine/Plugins/ChaosVD/Source/ChaosVD/Private/ChaosVDConstraintDataInspectorTab.cpp
--- a/Engine/Plugins/ChaosVD/Source/ChaosVD/Private/ChaosVDConstraintDataInspectorTab.cpp
+++ b/Engine/Plugins/ChaosVD/Source/ChaosVD/Private/ChaosVDConstraintDataInspectorTab.cpp
@@ -23,7 +23,7 @@ TSharedRef<SDockTab> FChaosVDConstraintDataInspectorTab::HandleTabSpawnRequest(c
 	{
 		DetailsPanelTab->SetContent
 		(
-			SAssignNew(ConstraintDataInspector, SChaosVDConstraintDataInspector, GetChaosVDScene())
+			SAssignNew(ConstraintDataInspector, SChaosVDConstraintDataInspector, GetChaosVDScene(), MainTabPtr.ToSharedRef())
 		);
 	}
 	else
diff --git a/Engine/Plugins/ChaosVD/Source/ChaosVD/Private/ChaosVDGeometryDataComponent.h b/Engine/Plugins/ChaosVD/Source/ChaosVD/Private/ChaosVDGeometryDataComponent.h
--- a/Engine/Plugins/ChaosVD/Source/ChaosVD/Private/ChaosVDGeometryDataComponent.h
+++ b/Engine/Plugins/ChaosVD/Source/ChaosVD/Private/ChaosVDGeometryDataComponent.h
@@ -82,19 +82,19 @@ struct FChaosVDMeshDataInstanceState
 	GENERATED_BODY()
 
 	/** Recorded Shape instance Data */
-	UPROPERTY(VisibleAnywhere, Category="Recorded GeometryData", meta=(ShowOnlyInnerProperties))
+	UPROPERTY(VisibleAnywhere, Category="Recorded GeometryData")
 	FChaosVDShapeCollisionData CollisionData;
 
 	/** Minimum set of data about the recorded implicit object */
-	UPROPERTY(VisibleAnywhere, Category="Recorded GeometryData", meta=(ShowOnlyInnerProperties))
+	UPROPERTY(VisibleAnywhere, Category="Recorded GeometryData")
 	FChaosVDImplicitObjectBasicView ImplicitObjectInfo;
 
 	/* CVD Debug - Current world transform used to render this Mesh */
-	UPROPERTY(VisibleAnywhere, Category="CVD GeometryData")
+	UPROPERTY(VisibleAnywhere, Category="CVD GeometryData", meta=(EditCondition=bShowCVDDebugData, EditConditionHides))
 	FTransform CurrentWorldTransform;
 
 	/* CVD Debug - Current mesh component type to render this Mesh */
-	UPROPERTY(VisibleAnywhere, Category="CVD GeometryData")
+	UPROPERTY(VisibleAnywhere, Category="CVD GeometryData", meta=(EditCondition=bShowCVDDebugData, EditConditionHides))
 	EChaosVDMeshComponent MeshComponentType = EChaosVDMeshComponent::Invalid;
 
 	/* CVD Debug - Pointer to the mesh component used to render this Mesh */
@@ -102,24 +102,29 @@ struct FChaosVDMeshDataInstanceState
 	TObjectPtr<UMeshComponent> MeshComponent;
 
 	/* CVD Debug - Instance index of mesh component used to render this Mesh */
-	UPROPERTY(VisibleAnywhere, Category="CVD GeometryData")
+	UPROPERTY(VisibleAnywhere, Category="CVD GeometryData", meta=(EditCondition=bShowCVDDebugData, EditConditionHides))
 	int32 MeshInstanceIndex = INDEX_NONE;
 
 	/* CVD Debug - Color used to render this mesh */
-	UPROPERTY(VisibleAnywhere, Category="CVD GeometryData")
+	UPROPERTY(VisibleAnywhere, Category="CVD GeometryData", meta=(EditCondition=bShowCVDDebugData, EditConditionHides))
 	FLinearColor CurrentGeometryColor = FLinearColor(ForceInitToZero);
 
 	/* CVD Debug - Id of the particle this geometry belongs */
-	UPROPERTY(VisibleAnywhere, Category="CVD GeometryData")
+	UPROPERTY(VisibleAnywhere, Category="CVD GeometryData", meta=(EditCondition=bShowCVDDebugData, EditConditionHides))
 	int32 OwningParticleID = INDEX_NONE;
 
 	/* CVD Debug - Id of the solver this geometry belongs */
-	UPROPERTY(VisibleAnywhere, Category="CVD GeometryData")
+	UPROPERTY(VisibleAnywhere, Category="CVD GeometryData", meta=(EditCondition=bShowCVDDebugData, EditConditionHides))
 	int32 OwningSolverID = INDEX_NONE;
 
+	UPROPERTY(VisibleAnywhere, Category="CVD GeometryData", meta=(EditCondition=bShowCVDDebugData, EditConditionHides))
 	bool bIsVisible = true;
 
+	UPROPERTY(VisibleAnywhere, Category="CVD GeometryData", meta=(EditCondition=bShowCVDDebugData, EditConditionHides))
 	bool bIsSelected = false;
+
+	UPROPERTY(EditAnywhere, Category="CVD GeometryData")
+	bool bShowCVDDebugData = false;
 };
 
 /** Handle that provides access to a specific mesh instance on a CVD Mesh component (instanced or static)*/
diff --git a/Engine/Plugins/ChaosVD/Source/ChaosVD/Private/ChaosVDModule.cpp b/Engine/Plugins/ChaosVD/Source/ChaosVD/Private/ChaosVDModule.cpp
--- a/Engine/Plugins/ChaosVD/Source/ChaosVD/Private/ChaosVDModule.cpp
+++ b/Engine/Plugins/ChaosVD/Source/ChaosVD/Private/ChaosVDModule.cpp
@@ -8,19 +8,14 @@
 #include "ChaosVDSettingsManager.h"
 #include "ChaosVDStyle.h"
 #include "ChaosVDTabsIDs.h"
-#include "DetailsCustomizations/ChaosVDGeometryComponentCustomization.h"
-#include "DetailsCustomizations/ChaosVDParticleDataWrapperCustomization.h"
-#include "DetailsCustomizations/ChaosVDQueryDataWrappersCustomizationDetails.h"
 #include "Misc/App.h"
-#include "Misc/CommandLine.h"
 #include "Misc/Guid.h"
-#include "PropertyEditorModule.h"
 #include "Trace/ChaosVDTraceManager.h"
 #include "Widgets/Docking/SDockTab.h"
 #include "Widgets/SChaosVDMainTab.h"
 #include "WorkspaceMenuStructure.h"
 #include "WorkspaceMenuStructureModule.h"
-#include "DetailsCustomizations/ChaosVDSelectionMultipleViewCustomization.h"
+#include "DetailsCustomizations/ChaosVDShapeDataCustomization.h"
 
 #define LOCTEXT_NAMESPACE "ChaosVisualDebugger"
 
@@ -48,8 +43,6 @@ void FChaosVDModule::StartupModule()
 	
 	FChaosVDCommands::Register();
 
-	RegisterClassesCustomDetails();
-
 	FGlobalTabmanager::Get()->RegisterNomadTabSpawner(FChaosVDTabID::ChaosVisualDebuggerTab, FOnSpawnTab::CreateRaw(this, &FChaosVDModule::SpawnMainTab))
 								.SetDisplayName(LOCTEXT("VisualDebuggerTabTitle", "Chaos Visual Debugger"))
 								.SetTooltipText(LOCTEXT("VisualDebuggerTabDesc", "Opens the Chaos Visual Debugger window"))
@@ -93,22 +86,6 @@ void FChaosVDModule::ShutdownModule()
 	FChaosVDSettingsManager::TearDown();
 }
 
-void FChaosVDModule::RegisterClassesCustomDetails() const
-{
-	FPropertyEditorModule& PropertyModule = FModuleManager::GetModuleChecked<FPropertyEditorModule>("PropertyEditor");
-	PropertyModule.RegisterCustomClassLayout("ChaosVDParticleActor", FOnGetDetailCustomizationInstance::CreateStatic(&FChaosVDParticleActorCustomization::MakeInstance));
-	PropertyModule.RegisterCustomClassLayout("ChaosVDInstancedStaticMeshComponent", FOnGetDetailCustomizationInstance::CreateStatic(&FChaosVDGeometryComponentCustomization::MakeInstance));
-	PropertyModule.RegisterCustomClassLayout("ChaosVDStaticMeshComponent", FOnGetDetailCustomizationInstance::CreateStatic(&FChaosVDGeometryComponentCustomization::MakeInstance));
-	PropertyModule.RegisterCustomClassLayout("ChaosVDQueryVisitStep", FOnGetDetailCustomizationInstance::CreateStatic(&FChaosVDQueryVisitDataCustomization::MakeInstance));
-	PropertyModule.RegisterCustomClassLayout("ChaosVDQueryDataWrapper", FOnGetDetailCustomizationInstance::CreateStatic(&FChaosVDQueryDataWrapperCustomization::MakeInstance));
-	PropertyModule.RegisterCustomClassLayout("ChaosVDSelectionMultipleView", FOnGetDetailCustomizationInstance::CreateStatic(&FChaosVDSelectionMultipleViewCustomization::MakeInstance));
-
-	//TODO: Rename FChaosVDParticleDataWrapperCustomization to something generic as currently works with any type that wants to hide properties of type FChaosVDWrapperDataBase with invalid data.
-	// Or another option is create a new custom layout intended to be generic from the get go
-	PropertyModule.RegisterCustomPropertyTypeLayout("ChaosVDQueryDataWrapper", FOnGetPropertyTypeCustomizationInstance::CreateStatic(&FChaosVDParticleDataWrapperCustomization::MakeInstance));
-	PropertyModule.RegisterCustomPropertyTypeLayout("ChaosVDQueryVisitStep", FOnGetPropertyTypeCustomizationInstance::CreateStatic(&FChaosVDParticleDataWrapperCustomization::MakeInstance));
-}
-
 void FChaosVDModule::SpawnCVDTab()
 {
 	if (IsStandaloneChaosVisualDebugger())
diff --git a/Engine/Plugins/ChaosVD/Source/ChaosVD/Private/ChaosVDObjectDetailsTab.cpp b/Engine/Plugins/ChaosVD/Source/ChaosVD/Private/ChaosVDObjectDetailsTab.cpp
--- a/Engine/Plugins/ChaosVD/Source/ChaosVD/Private/ChaosVDObjectDetailsTab.cpp
+++ b/Engine/Plugins/ChaosVD/Source/ChaosVD/Private/ChaosVDObjectDetailsTab.cpp
@@ -50,25 +50,33 @@ TSharedRef<SDockTab> FChaosVDObjectDetailsTab::HandleTabSpawnRequest(const FSpaw
 			SolverDataSelectionObject->GetDataSelectionChangedDelegate().AddSP(this, &FChaosVDObjectDetailsTab::HandleSolverDataSelectionChange);
 		}
 
-		DetailsPanelTab->SetContent
-		(
-			SNew(SVerticalBox)
-			+SVerticalBox::Slot()
-			[
-				SAssignNew(DetailsPanelView, SChaosVDDetailsView)
-			]
-			+SVerticalBox::Slot()
-			.AutoHeight()
-			[
-				GenerateShowCollisionDataButton().ToSharedRef()
-			]
-		);
-
-		// If we closed the tab and opened it again with an object already selected, try to restore the selected object view
-		if (DetailsPanelView.IsValid() && CurrentSelectedObject.IsValid())
+		if (const TSharedPtr<SChaosVDMainTab> MainTabPtr = OwningTabWidget.Pin())
 		{
-			DetailsPanelView->SetSelectedObject(CurrentSelectedObject.Get());
+			DetailsPanelTab->SetContent
+			(
+				SNew(SVerticalBox)
+				+SVerticalBox::Slot()
+				[
+					SAssignNew(DetailsPanelView, SChaosVDDetailsView, MainTabPtr.ToSharedRef())
+				]
+				+SVerticalBox::Slot()
+				.AutoHeight()
+				[
+					GenerateShowCollisionDataButton().ToSharedRef()
+				]
+			);
+
+			// If we closed the tab and opened it again with an object already selected, try to restore the selected object view
+			if (DetailsPanelView.IsValid() && CurrentSelectedObject.IsValid())
+			{
+				DetailsPanelView->SetSelectedObject(CurrentSelectedObject.Get());
+			}
 		}
+		else
+		{
+			DetailsPanelTab->SetContent(GenerateErrorWidget());
+		}
+
 	}
 	else
 	{
diff --git a/Engine/Plugins/ChaosVD/Source/ChaosVD/Private/ChaosVDParticleActorCustomization.cpp b/Engine/Plugins/ChaosVD/Source/ChaosVD/Private/ChaosVDParticleActorCustomization.cpp
--- a/Engine/Plugins/ChaosVD/Source/ChaosVD/Private/ChaosVDParticleActorCustomization.cpp
+++ b/Engine/Plugins/ChaosVD/Source/ChaosVD/Private/ChaosVDParticleActorCustomization.cpp
@@ -2,45 +2,60 @@
 
 #include "ChaosVDParticleActorCustomization.h"
 
+#include "ChaosVDEngine.h"
 #include "ChaosVDModule.h"
 #include "ChaosVDParticleActor.h"
+#include "ChaosVDScene.h"
 #include "DetailCategoryBuilder.h"
 #include "DetailLayoutBuilder.h"
 #include "DetailWidgetRow.h"
 #include "IStructureDetailsView.h"
 #include "DataWrappers/ChaosVDParticleDataWrapper.h"
 #include "DetailsCustomizations/ChaosVDDetailsCustomizationUtils.h"
+#include "Widgets/SChaosVDMainTab.h"
 
 #define LOCTEXT_NAMESPACE "ChaosVisualDebugger"
 
-FChaosVDParticleActorCustomization::FChaosVDParticleActorCustomization()
+FChaosVDParticleActorCustomization::FChaosVDParticleActorCustomization(const TSharedPtr<SChaosVDMainTab>& InMainTab)
 {
 	AllowedCategories.Add(FChaosVDParticleActorCustomization::ParticleDataCategoryName);
 	AllowedCategories.Add(FChaosVDParticleActorCustomization::GeometryCategoryName);
+
+	MainTabWeakPtr = InMainTab;
+
+	ResetCachedView();
 }
 
 FChaosVDParticleActorCustomization::~FChaosVDParticleActorCustomization()
 {
-	if (CurrentObservedActor.Get())
-	{
-		CurrentObservedActor->OnParticleDataUpdated().Unbind();
-	}
+	RegisterCVDScene(nullptr);
 }
 
-TSharedRef<IDetailCustomization> FChaosVDParticleActorCustomization::MakeInstance()
+TSharedRef<IDetailCustomization> FChaosVDParticleActorCustomization::MakeInstance(TSharedRef<SChaosVDMainTab> InMainTab)
 {
-	return MakeShareable( new FChaosVDParticleActorCustomization );
+	return MakeShared<FChaosVDParticleActorCustomization>(InMainTab);
 }
 
 void FChaosVDParticleActorCustomization::CustomizeDetails(IDetailLayoutBuilder& DetailBuilder)
 {
 	FChaosVDDetailsCustomizationUtils::HideAllCategories(DetailBuilder, AllowedCategories);
 
+	TSharedPtr<SChaosVDMainTab> MainTabPtr =  MainTabWeakPtr.Pin(); 
+	TSharedPtr<FChaosVDScene> Scene = MainTabPtr ? MainTabPtr->GetChaosVDEngineInstance()->GetCurrentScene() : nullptr;
+
+	RegisterCVDScene(Scene);
+
+	if (!Scene)
+	{
+		ResetCachedView();
+		return;
+	}
+
 	// We keep the particle data we need to visualize as a shared ptr because copying it each frame we advance/rewind to to an struct that lives in the particle actor it is not cheap.
 	// Having a struct details view to which we set that pointer data each time the data in the particle is updated (meaning we assigned another ptr from the recording)
 	// seems to be more expensive because it has to rebuild the entire layout from scratch.
 	// So a middle ground I found is to have a Particle Data struct in this customization instance, which we add as external property. Then each time the particle data is updated we copy the data over.
-	// This allow us to only perform the copy just for the particle that is being inspected and not every particle updated in that frame.
+	// This allows us to only perform the copy just for the particle that is being inspected and not every particle updated in that frame.
 
 	TArray<TWeakObjectPtr<UObject>> SelectedObjects;
 	DetailBuilder.GetObjectsBeingCustomized(SelectedObjects);
@@ -49,27 +64,26 @@ void FChaosVDParticleActorCustomization::CustomizeDetails(IDetailLayoutBuilder&
 		//TODO: Add support for multi-selection.
 		if (!ensure(SelectedObjects.Num() == 1))
 		{
-			UE_LOG(LogChaosVDEditor, Warning, TEXT("[%s] [%d] objects were selectioned but this customization panel only support single object selection."), ANSI_TO_TCHAR(__FUNCTION__), SelectedObjects.Num())
+			UE_LOG(LogChaosVDEditor, Warning, TEXT("[%s] [%d] objects were selected but this customization panel only support single object selection."), ANSI_TO_TCHAR(__FUNCTION__), SelectedObjects.Num())
 		}
+		
+		AChaosVDParticleActor* CurrentActor = CurrentObservedActor.Get();
+		AChaosVDParticleActor* SelectedActor = Cast<AChaosVDParticleActor>(SelectedObjects[0]);
 
-		if (AChaosVDParticleActor* CurrentActor = CurrentObservedActor.Get())
+		if (CurrentActor && CurrentActor != SelectedActor)
 		{
-			CachedParticleData = FChaosVDParticleDataWrapper();
-			CachedGeometryDataInstanceCopy = FChaosVDMeshDataInstanceState();
-			CurrentActor->OnParticleDataUpdated().Unbind();
-			CurrentActor = nullptr;
+			ResetCachedView();
 		}
 		
-		if (AChaosVDParticleActor* ParticleActor = Cast<AChaosVDParticleActor>(SelectedObjects[0]))
+		if (SelectedActor)
 		{
-			CurrentObservedActor = ParticleActor;
-			CurrentObservedActor->OnParticleDataUpdated().BindRaw(this, &FChaosVDParticleActorCustomization::HandleParticleDataUpdated);
+			CurrentObservedActor = SelectedActor;
 
-			HandleParticleDataUpdated();
+			HandleSceneUpdated();
 
 			TSharedPtr<IPropertyHandle> InspectedDataPropertyHandlePtr;
 
-			if (TSharedPtr<FChaosVDMeshDataInstanceHandle> SelectedGeometryInstance = ParticleActor->GetSelectedMeshInstance().Pin())
+			if (TSharedPtr<FChaosVDMeshDataInstanceHandle> SelectedGeometryInstance = SelectedActor->GetSelectedMeshInstance().Pin())
 			{
 				InspectedDataPropertyHandlePtr = AddExternalStructure(CachedGeometryDataInstanceCopy, DetailBuilder, GeometryCategoryName, LOCTEXT("GeometryShapeDataStructName", "Geometry Shape Data"));
 			}
@@ -85,15 +99,18 @@ void FChaosVDParticleActorCustomization::CustomizeDetails(IDetailLayoutBuilder&
 			}
 		}
 	}
+	else
+	{
+		ResetCachedView();
+	}
 }
 
-void FChaosVDParticleActorCustomization::HandleParticleDataUpdated()
+void FChaosVDParticleActorCustomization::HandleSceneUpdated()
 {
 	AChaosVDParticleActor* ParticleActor = CurrentObservedActor.Get();
 	if (!ParticleActor)
 	{
-		CachedParticleData = FChaosVDParticleDataWrapper();
-		CachedGeometryDataInstanceCopy = FChaosVDMeshDataInstanceState();
+		ResetCachedView();
 		return;
 	}
 
@@ -115,4 +132,30 @@ void FChaosVDParticleActorCustomization::HandleParticleDataUpdated()
 	}
 }
 
+void FChaosVDParticleActorCustomization::ResetCachedView()
+{
+	CurrentObservedActor = nullptr;
+	CachedParticleData = FChaosVDParticleDataWrapper();
+	CachedGeometryDataInstanceCopy = FChaosVDMeshDataInstanceState();
+}
+
+void FChaosVDParticleActorCustomization::RegisterCVDScene(const TSharedPtr<FChaosVDScene>& InScene)
+{
+	TSharedPtr<FChaosVDScene> CurrentScene = SceneWeakPtr.Pin();
+	if (InScene != CurrentScene)
+	{
+		if (CurrentScene)
+		{
+			CurrentScene->OnSceneUpdated().RemoveAll(this);
+		}
+
+		if (InScene)
+		{
+			InScene->OnSceneUpdated().AddSP(this, &FChaosVDParticleActorCustomization::HandleSceneUpdated);
+		}
+
+		SceneWeakPtr = InScene;
+	}
+}
+
 #undef LOCTEXT_NAMESPACE
diff --git a/Engine/Plugins/ChaosVD/Source/ChaosVD/Private/ChaosVDParticleActorCustomization.h b/Engine/Plugins/ChaosVD/Source/ChaosVD/Private/ChaosVDParticleActorCustomization.h
--- a/Engine/Plugins/ChaosVD/Source/ChaosVD/Private/ChaosVDParticleActorCustomization.h
+++ b/Engine/Plugins/ChaosVD/Source/ChaosVD/Private/ChaosVDParticleActorCustomization.h
@@ -10,6 +10,7 @@
 #include "PropertyHandle.h"
 #include "DataWrappers/ChaosVDParticleDataWrapper.h"
 
+class SChaosVDMainTab;
 struct FChaosVDParticleDataWrapper;
 class AChaosVDParticleActor;
 
@@ -17,13 +18,13 @@ class AChaosVDParticleActor;
 class FChaosVDParticleActorCustomization : public IDetailCustomization
 {
 public:
-	FChaosVDParticleActorCustomization();
+	FChaosVDParticleActorCustomization(const TSharedPtr<SChaosVDMainTab>& InMainTab);
 	virtual ~FChaosVDParticleActorCustomization() override;
 
 	inline static FName ParticleDataCategoryName = FName("Particle Data");
 	inline static FName GeometryCategoryName = FName("Geometry Shape Data");
 
-	static TSharedRef<IDetailCustomization> MakeInstance();
+	static TSharedRef<IDetailCustomization> MakeInstance(TSharedRef<SChaosVDMainTab> InMainTab);
 
 	virtual void CustomizeDetails(IDetailLayoutBuilder& DetailBuilder) override;
 
@@ -35,12 +36,20 @@ private:
 	TSet<FName> AllowedCategories;
 	TWeakObjectPtr<AChaosVDParticleActor> CurrentObservedActor;
 
-	void HandleParticleDataUpdated();
+	void ResetCachedView();
+
+	void RegisterCVDScene(const TSharedPtr<FChaosVDScene>& InScene);
+
+	void HandleSceneUpdated();
 
 	/* Copy of the last known geometry shape data structure of a selected particle and mesh instance -  Used to avoid rebuild the layout every time we change frame in CVD */
 	FChaosVDParticleDataWrapper CachedParticleData;
 	/* Copy of the last known particle data structure of a selected particle -  Used to avoid rebuild the layout every time we change frame in CVD */
 	FChaosVDMeshDataInstanceState CachedGeometryDataInstanceCopy;
+	
+	TWeakPtr<FChaosVDScene> SceneWeakPtr = nullptr;
+
+	TWeakPtr<SChaosVDMainTab> MainTabWeakPtr = nullptr;
 };
 
 template <typename TStruct>
diff --git a/Engine/Plugins/ChaosVD/Source/ChaosVD/Private/ChaosVDSceneQueryDataInspectorTab.cpp b/Engine/Plugins/ChaosVD/Source/ChaosVD/Private/ChaosVDSceneQueryDataInspectorTab.cpp
--- a/Engine/Plugins/ChaosVD/Source/ChaosVD/Private/ChaosVDSceneQueryDataInspectorTab.cpp
+++ b/Engine/Plugins/ChaosVD/Source/ChaosVD/Private/ChaosVDSceneQueryDataInspectorTab.cpp
@@ -25,7 +25,7 @@ TSharedRef<SDockTab> FChaosVDSceneQueryDataInspectorTab::HandleTabSpawnRequest(c
 	{
 		DetailsPanelTab->SetContent
 		(
-			SAssignNew(SceneQueryDataInspector, SChaosVDSceneQueryDataInspector, GetChaosVDScene(), MainTabPtr->GetEditorModeManager().AsWeak())
+			SAssignNew(SceneQueryDataInspector, SChaosVDSceneQueryDataInspector, GetChaosVDScene(), MainTabPtr.ToSharedRef())
 		);
 	}
 	else
diff --git a/Engine/Plugins/ChaosVD/Source/ChaosVD/Private/DetailsCustomizations/ChaosVDDetailsCustomizationUtils.cpp b/Engine/Plugins/ChaosVD/Source/ChaosVD/Private/DetailsCustomizations/ChaosVDDetailsCustomizationUtils.cpp
--- a/Engine/Plugins/ChaosVD/Source/ChaosVD/Private/DetailsCustomizations/ChaosVDDetailsCustomizationUtils.cpp
+++ b/Engine/Plugins/ChaosVD/Source/ChaosVD/Private/DetailsCustomizations/ChaosVDDetailsCustomizationUtils.cpp
@@ -2,9 +2,18 @@
 
 #include "DetailsCustomizations/ChaosVDDetailsCustomizationUtils.h"
 
+#include "ChaosVDModule.h"
+#include "ChaosVDScene.h"
 #include "DetailLayoutBuilder.h"
+#include "DetailWidgetRow.h"
 #include "DataWrappers/ChaosVDParticleDataWrapper.h"
+#include "Engine/EngineTypes.h"
+#include "IDetailGroup.h"
+#include "Widgets/Input/SCheckBox.h"
+#include "Widgets/Layout/SBox.h"
+#include "Widgets/Text/STextBlock.h"
 
+#define LOCTEXT_NAMESPACE "ChaosVisualDebugger"
 
 void FChaosVDDetailsCustomizationUtils::HideAllCategories(IDetailLayoutBuilder& DetailBuilder, const TSet<FName>& AllowedCategories)
 {
@@ -87,3 +96,198 @@ bool FChaosVDDetailsCustomizationUtils::HasValidCVDWrapperData(const TSharedPtr<
 
 	return true;
 }
+
+TSharedPtr<FChaosVDCollisionChannelsInfoContainer> FChaosVDDetailsCustomizationUtils::BuildDefaultCollisionChannelInfo()
+{
+	TSharedPtr<FChaosVDCollisionChannelsInfoContainer> NewCollisionChannelsInfoContainer = MakeShared<FChaosVDCollisionChannelsInfoContainer>();
+
+	// Build the default channels name & type using the enum metadata
+	UEnum * Enum = StaticEnum<ECollisionChannel>();
+	if (!ensure(Enum))
+	{
+		UE_LOG(LogChaosVDEditor, Error, TEXT("[%s] Failed to generate fallback collision channels entry "), ANSI_TO_TCHAR(__FUNCTION__));
+		return NewCollisionChannelsInfoContainer;
+	}
+
+	const int32 NumEnum = Enum->NumEnums();
+	constexpr int32 ExpectedChannels = FChaosVDDetailsCustomizationUtils::GetMaxCollisionChannelIndex();
+
+	if (NumEnum < ExpectedChannels)
+	{
+		UE_LOG(LogChaosVDEditor, Error, TEXT("[%s] Failed to generate fallback collision channels entry "), ANSI_TO_TCHAR(__FUNCTION__));
+		return NewCollisionChannelsInfoContainer;
+	}
+	
+	const FString KeyName = TEXT("DisplayName");
+	const FString TraceType = TEXT("TraceQuery");
+
+	for (int32 ChannelIndex = 0; ChannelIndex < ExpectedChannels; ++ChannelIndex)
+	{
+		FChaosVDCollisionChannelInfo Info;
+		Info.DisplayName = Enum->GetDisplayNameTextByIndex(ChannelIndex).ToString();
+		Info.CollisionChannel = static_cast<ECollisionChannel>(ChannelIndex);
+		Info.bIsTraceType = Enum->GetMetaData(*TraceType, ChannelIndex) == TEXT("1");
+
+		NewCollisionChannelsInfoContainer->CustomChannelsNames[ChannelIndex] = Info;
+	}
+
+	return NewCollisionChannelsInfoContainer;
+}
+
+void FChaosVDDetailsCustomizationUtils::CreateCollisionChannelsMatrixRow(int32 ChannelIndex, const FChaosVDCollisionChannelStateGetter& InChannelStateGetter, const FText& InChannelName, IDetailGroup& CollisionGroup, const float RowWidthCustomization)
+{
+	// Currently all details panel in CVD are read-only
+	constexpr bool bEnabledState = false;
+
+	CollisionGroup.AddWidgetRow()
+	.NameContent()
+	[
+		SNew(STextBlock)
+		.IsEnabled(bEnabledState)
+		.Text(InChannelName)
+		.Font(IDetailLayoutBuilder::GetDetailFont())
+	]
+	.ValueContent()
+	[
+		SNew(SHorizontalBox)
+		.IsEnabled(bEnabledState)
+		+SHorizontalBox::Slot()
+		.AutoWidth()
+		.VAlign(VAlign_Center)
+		[
+			CreateCollisionResponseMatrixCheckbox(InChannelStateGetter, ChannelIndex, ECollisionResponse::ECR_Ignore, RowWidthCustomization)
+		]
+		+SHorizontalBox::Slot()
+		.AutoWidth()
+		.VAlign(VAlign_Center)
+		[
+			CreateCollisionResponseMatrixCheckbox(InChannelStateGetter, ChannelIndex, ECollisionResponse::ECR_Overlap, RowWidthCustomization)
+		]
+		+SHorizontalBox::Slot()
+		.AutoWidth()
+		.VAlign(VAlign_Center)
+		[
+			CreateCollisionResponseMatrixCheckbox(InChannelStateGetter, ChannelIndex, ECollisionResponse::ECR_Block, RowWidthCustomization)
+		]
+	];
+}
+
+void FChaosVDDetailsCustomizationUtils::BuildCollisionChannelMatrix(const FChaosVDCollisionChannelStateGetter& InCollisionChannelStateGetter, TConstArrayView<FChaosVDCollisionChannelInfo> CollisionChannelsInfo, IDetailGroup& ParentCategoryGroup)
+{
+	constexpr float RowWidthCustomization = 50;
+
+	ParentCategoryGroup.AddWidgetRow()
+	.ValueContent()
+	.MaxDesiredWidth(0)
+	.MinDesiredWidth(0)
+	[
+		SNew(SHorizontalBox)
+		+SHorizontalBox::Slot()
+		.AutoWidth()
+		[
+			SNew(SBox)
+			.WidthOverride(RowWidthCustomization)
+			.HAlign(HAlign_Left)
+			.Content()
+			[
+				SNew(STextBlock)
+				.Text(LOCTEXT("IgnoreCollisionLabel", "Ignore"))
+				.Font(IDetailLayoutBuilder::GetDetailFontBold())
+			]
+		]
+		+SHorizontalBox::Slot()
+		.AutoWidth()
+		[
+			SNew(SBox)
+			.HAlign(HAlign_Left)
+			.WidthOverride(RowWidthCustomization)
+			.Content()
+			[
+				SNew(STextBlock)
+				.Text(LOCTEXT("OverlapCollisionLabel", "Overlap"))
+				.Font(IDetailLayoutBuilder::GetDetailFontBold())
+			]
+		]
+		+SHorizontalBox::Slot()
+		.AutoWidth()
+		[
+			SNew(STextBlock)
+			.Text(LOCTEXT("BlockCollisionLabel", "Block"))
+			.Font(IDetailLayoutBuilder::GetDetailFontBold())
+		]
+	];
+
+	constexpr bool bStartExpanded = true;
+	IDetailGroup& TraceResponsesGroup = ParentCategoryGroup.AddGroup(TEXT("CollisionTraceResponses"), LOCTEXT("CollisionTraceResponsesLabel", "Trace Responses"), bStartExpanded);
+	TraceResponsesGroup.EnableReset(false);
+
+	constexpr int32 ExpectedChannels = FChaosVDDetailsCustomizationUtils::GetMaxCollisionChannelIndex();
+
+	for (int32 ChannelIndex = 0; ChannelIndex < ExpectedChannels; ++ChannelIndex)
+	{
+		const FChaosVDCollisionChannelInfo& ChannelInfo = CollisionChannelsInfo[ChannelIndex];
+		if (ChannelInfo.bIsTraceType)
+		{
+			FChaosVDDetailsCustomizationUtils::CreateCollisionChannelsMatrixRow(ChannelIndex, InCollisionChannelStateGetter, FText::FromString(ChannelInfo.DisplayName), TraceResponsesGroup, RowWidthCustomization);
+		}
+	}
+
+	IDetailGroup& CollisionResponsesGroup = ParentCategoryGroup.AddGroup(TEXT("CollisionObjectResponses"), LOCTEXT("CollisionObjectResponses", "Object Responses"), bStartExpanded);
+	TraceResponsesGroup.EnableReset(false);
+
+	for (int32 ChannelIndex = 0; ChannelIndex < ExpectedChannels; ++ChannelIndex)
+	{
+		const FChaosVDCollisionChannelInfo& ChannelInfo = CollisionChannelsInfo[ChannelIndex];
+		if (!ChannelInfo.bIsTraceType)
+		{
+			FChaosVDDetailsCustomizationUtils::CreateCollisionChannelsMatrixRow(ChannelIndex, InCollisionChannelStateGetter, FText::FromString(ChannelInfo.DisplayName), CollisionResponsesGroup, RowWidthCustomization);
+		}
+	}
+}
+
+TSharedRef<SWidget> FChaosVDDetailsCustomizationUtils::CreateCollisionResponseMatrixCheckbox(const FChaosVDCollisionChannelStateGetter& InStateGetter, int32 ChannelIndex, ECollisionResponse TargetResponse, float Width)
+{
+	return SNew(SBox)
+		.WidthOverride(Width)
+		.Content()
+		[
+			SNew(SCheckBox)
+			.IsChecked_Lambda([GetterCopy = InStateGetter, ChannelIndex, TargetResponse]()
+			{
+				if (GetterCopy.IsBound())
+				{
+					return GetterCopy.Execute(ChannelIndex) == TargetResponse ? ECheckBoxState::Checked : ECheckBoxState::Unchecked;
+				}
+				return ECheckBoxState::Undetermined;
+			})
+		];
+}
+
+void FChaosVDDetailsCustomizationUtils::AddWidgetRowForCheckboxValue(TAttribute<ECheckBoxState>&& State, const FText& InValueName, IDetailGroup& DetailGroup)
+{
+	DetailGroup.AddWidgetRow()
+	.IsEnabled(false)
+	.NameContent()
+	[
+		SNew(STextBlock)
+		.Text(InValueName)
+		.Font(IDetailLayoutBuilder::GetDetailFont())
+	]
+	.ValueContent()
+	[
+		SNew(SBox)
+		.WidthOverride(50)
+		.Content()
+		[
+			SNew(SCheckBox)
+			.IsChecked(State)
+		]
+	];
+}
+
+FText FChaosVDDetailsCustomizationUtils::GetDefaultCollisionChannelsUseWarningMessage()
+{
+	return LOCTEXT("EngineDefaultsWarningMessageBox", "The following names are the Engine's default channel names. \n Some might be incorrect or missing");
+}
+
+#undef LOCTEXT_NAMESPACE
diff --git a/Engine/Plugins/ChaosVD/Source/ChaosVD/Private/DetailsCustomizations/ChaosVDDetailsCustomizationUtils.h b/Engine/Plugins/ChaosVD/Source/ChaosVD/Private/DetailsCustomizations/ChaosVDDetailsCustomizationUtils.h
--- a/Engine/Plugins/ChaosVD/Source/ChaosVD/Private/DetailsCustomizations/ChaosVDDetailsCustomizationUtils.h
+++ b/Engine/Plugins/ChaosVD/Source/ChaosVD/Private/DetailsCustomizations/ChaosVDDetailsCustomizationUtils.h
@@ -2,14 +2,76 @@
 
 #pragma once
 
+#include "PropertyHandle.h"
+#include "Chaos/CollisionFilterData.h"
 #include "Containers/Set.h"
+#include "Delegates/DelegateCombinations.h"
 #include "Templates/SharedPointer.h"
 #include "UObject/NameTypes.h"
+#include "Engine/EngineTypes.h"
 
+class SChaosVDMainTab;
+struct FCollisionFilterData;
+enum ECollisionResponse : int;
+enum class ECheckBoxState : uint8;
+struct FChaosVDRecording;
+class FChaosVDScene;
+class IDetailsView;
+struct FChaosVDCollisionChannelInfo;
+struct FChaosVDCollisionChannelsInfoContainer;
+class IDetailGroup;
 class IPropertyHandle;
 class IDetailLayoutBuilder;
 class FName;
 
+DECLARE_DELEGATE_RetVal_OneParam(ECollisionResponse, FChaosVDCollisionChannelStateGetter, int32 /* ChannelIndex*/);
+
+template<typename PropertyType>
+struct FChaosVDDetailsPropertyDataHandle
+{
+	explicit FChaosVDDetailsPropertyDataHandle(const TSharedRef<IPropertyHandle>& InPropertyHandle);
+
+	PropertyType* GetDataInstance();
+	
+private:
+	TSharedPtr<IPropertyHandle> PropertyHandle;
+	PropertyType* DataInstance = nullptr;
+};
+
+template <typename PropertyType>
+FChaosVDDetailsPropertyDataHandle<PropertyType>::FChaosVDDetailsPropertyDataHandle(const TSharedRef<IPropertyHandle>& InPropertyHandle)
+{
+	PropertyHandle = InPropertyHandle;
+	
+	FProperty* Property = PropertyHandle->GetProperty();
+	if (!Property)
+	{
+		return;
+	}
+
+	const FStructProperty* StructProperty = CastField<FStructProperty>(Property);
+	if (StructProperty && StructProperty->Struct && StructProperty->Struct->IsChildOf(PropertyType::StaticStruct()))
+	{
+		void* Data = nullptr;
+		PropertyHandle->GetValueData(Data);
+		if (Data)
+		{
+			DataInstance = static_cast<PropertyType*>(Data);
+		}
+	}
+}
+
+template <typename PropertyType>
+PropertyType* FChaosVDDetailsPropertyDataHandle<PropertyType>::GetDataInstance()
+{
+	if (!PropertyHandle)
+	{
+		return nullptr;
+	}
+
+	return DataInstance;
+}
+
 /** Helper Class for CVD Custom Details view */
 class FChaosVDDetailsCustomizationUtils
 {
@@ -40,4 +102,39 @@ public:
 	 * @param bOutIsCVDBaseDataStruct Set to true if the property handle provided is from a CVD Wrapper Data Base. All other types will be deemed valid
 	 */
 	static bool HasValidCVDWrapperData(const TSharedPtr<IPropertyHandle>& InPropertyHandle, bool& bOutIsCVDBaseDataStruct);
+
+	static void CreateCollisionChannelsMatrixRow(int32 ChannelIndex, const FChaosVDCollisionChannelStateGetter& InChannelStateGetter, const FText& InChannelName, IDetailGroup& CollisionGroup, const float RowWidthCustomization);
+	
+	static TSharedPtr<FChaosVDCollisionChannelsInfoContainer> BuildDefaultCollisionChannelInfo();
+	
+	static void BuildCollisionChannelMatrix(const FChaosVDCollisionChannelStateGetter& InCollisionChannelStateGetter, TConstArrayView<FChaosVDCollisionChannelInfo> CollisionChannelsInfo, IDetailGroup& ParentCategoryGroup);
+
+	static TSharedRef<SWidget> CreateCollisionResponseMatrixCheckbox(const FChaosVDCollisionChannelStateGetter& InStateGetter, int32 ChannelIndex, ECollisionResponse TargetResponse, float Width);
+
+	static constexpr int32 GetMaxCollisionChannelIndex();
+
+	template<typename CVDCollisionFilteringData>
+	static FCollisionFilterData ConvertToEngineFilteringData(const CVDCollisionFilteringData& InCVDFilteringData);
+
+	static void AddWidgetRowForCheckboxValue(TAttribute<ECheckBoxState>&& State, const FText& InValueName, IDetailGroup& DetailGroup);
+
+	static FText GetDefaultCollisionChannelsUseWarningMessage();
 };
+
+template <typename CVDCollisionFilteringData>
+FCollisionFilterData FChaosVDDetailsCustomizationUtils::ConvertToEngineFilteringData(const CVDCollisionFilteringData& InCVDFilteringData)
+{
+	FCollisionFilterData EngineFilteringData;
+	EngineFilteringData.Word0 = InCVDFilteringData.Word0;
+	EngineFilteringData.Word1 = InCVDFilteringData.Word1;
+	EngineFilteringData.Word2 = InCVDFilteringData.Word2;
+	EngineFilteringData.Word3 = InCVDFilteringData.Word3;
+
+	return EngineFilteringData;
+}
+
+constexpr int32 FChaosVDDetailsCustomizationUtils::GetMaxCollisionChannelIndex()
+{
+	// Skip the deprecated channel which is last
+	return ECollisionChannel::ECC_MAX - 1;
+}
diff --git a/Engine/Plugins/ChaosVD/Source/ChaosVD/Private/DetailsCustomizations/ChaosVDQueryDataWrappersCustomizationDetails.cpp b/Engine/Plugins/ChaosVD/Source/ChaosVD/Private/DetailsCustomizations/ChaosVDQueryDataWrappersCustomizationDetails.cpp
--- a/Engine/Plugins/ChaosVD/Source/ChaosVD/Private/DetailsCustomizations/ChaosVDQueryDataWrappersCustomizationDetails.cpp
+++ b/Engine/Plugins/ChaosVD/Source/ChaosVD/Private/DetailsCustomizations/ChaosVDQueryDataWrappersCustomizationDetails.cpp
@@ -2,18 +2,29 @@
 
 #include "DetailsCustomizations/ChaosVDQueryDataWrappersCustomizationDetails.h"
 
+#include "ChaosVDEngine.h"
+#include "ChaosVDModule.h"
+#include "ChaosVDRecording.h"
+#include "ChaosVDScene.h"
+#include "CollisionQueryParams.h"
 #include "DataWrappers/ChaosVDQueryDataWrappers.h"
 #include "DetailsCustomizations/ChaosVDDetailsCustomizationUtils.h"
 #include "DetailLayoutBuilder.h"
+#include "IDetailChildrenBuilder.h"
+#include "IDetailGroup.h"
 #include "PropertyCustomizationHelpers.h"
 #include "PropertyHandle.h"
+#include "Widgets/SChaosVDMainTab.h"
+#include "Widgets/SChaosVDWarningMessageBox.h"
+#include "Widgets/Input/SCheckBox.h"
+#include "Widgets/Layout/SBox.h"
+#include "Widgets/Text/STextBlock.h"
 
 #define LOCTEXT_NAMESPACE "ChaosVisualDebugger"
 
 TSharedRef<IPropertyTypeCustomization> FChaosVDQueryDataWrappersCustomizationDetails::MakeInstance()
 {
 	return MakeShareable(new FChaosVDQueryDataWrappersCustomizationDetails());
-
 }
 
 void FChaosVDQueryDataWrappersCustomizationDetails::CustomizeChildren(TSharedRef<IPropertyHandle> StructPropertyHandle, IDetailChildrenBuilder& StructBuilder, IPropertyTypeCustomizationUtils& StructCustomizationUtils)
@@ -75,4 +86,185 @@ void FChaosVDQueryDataWrapperCustomization::CustomizeDetails(IDetailLayoutBuilde
 	FChaosVDDetailsCustomizationUtils::HideInvalidCVDDataWrapperProperties(PotentialPropertiesToHide, DetailBuilder);
 }
 
+FChaosVDCollisionChannelsCustomizationBase::FChaosVDCollisionChannelsCustomizationBase(const TSharedPtr<SChaosVDMainTab>& InMainTab)
+{
+	MainTabWeakPtr = InMainTab;
+	// Fill with the Engine defaults. When a CVD file is loaded we will update it with any new data is available
+	UpdateCollisionChannelsInfoCache({});
+}
+
+void FChaosVDCollisionChannelsCustomizationBase::CustomizeChildren(TSharedRef<IPropertyHandle> StructPropertyHandle,IDetailChildrenBuilder& StructBuilder, IPropertyTypeCustomizationUtils& StructCustomizationUtils)
+{
+	TSharedPtr<SChaosVDMainTab> MainTabPtr =  MainTabWeakPtr.Pin(); 
+	TSharedPtr<FChaosVDScene> Scene = MainTabPtr ? MainTabPtr->GetChaosVDEngineInstance()->GetCurrentScene() : nullptr;
+
+	if (TSharedPtr<FChaosVDRecording> CVDRecording = Scene ? Scene->LoadedRecording : nullptr)
+	{
+		UpdateCollisionChannelsInfoCache(CVDRecording->GetCollisionChannelsInfoContainer());
+	}
+	else
+	{
+		UpdateCollisionChannelsInfoCache(nullptr);
+	}
+}
+
+void FChaosVDCollisionChannelsCustomizationBase::UpdateCollisionChannelsInfoCache(TSharedPtr<FChaosVDCollisionChannelsInfoContainer> NewCollisionChannelsInfo)
+{
+	if (NewCollisionChannelsInfo)
+	{
+		CachedCollisionChannelInfos = NewCollisionChannelsInfo;
+		bChannelInfoBuiltFromDefaults = false;
+	}
+	else
+	{
+		// Fallback to engine channels name using the enum metadata
+		CachedCollisionChannelInfos = FChaosVDDetailsCustomizationUtils::BuildDefaultCollisionChannelInfo();
+		bChannelInfoBuiltFromDefaults = true;
+	}
+}
+
+ECollisionResponse FChaosVDCollisionResponseParamsCustomization::GetCurrentCollisionResponseForChannel(int32 ChannelIndex) const
+{
+	constexpr int32 MaxChannels = FChaosVDDetailsCustomizationUtils::GetMaxCollisionChannelIndex();
+
+	if (ChannelIndex >= MaxChannels)
+	{
+		return ECollisionResponse::ECR_MAX;
+	}
+
+	if (FChaosVDCollisionResponseParams* Data = CurrentPropertyDataHandle ? CurrentPropertyDataHandle->GetDataInstance() : nullptr)
+	{
+		return static_cast<ECollisionResponse>(Data->FlagsPerChannel[ChannelIndex]);
+	}
+
+	return ECollisionResponse::ECR_MAX;
+}
+
+TSharedRef<IPropertyTypeCustomization> FChaosVDCollisionResponseParamsCustomization::MakeInstance(TSharedRef<SChaosVDMainTab> MainTab)
+{
+	return MakeShared<FChaosVDCollisionResponseParamsCustomization>(MainTab);
+}
+
+void FChaosVDCollisionResponseParamsCustomization::CustomizeChildren(TSharedRef<IPropertyHandle> StructPropertyHandle, IDetailChildrenBuilder& StructBuilder, IPropertyTypeCustomizationUtils& StructCustomizationUtils)
+{
+	FChaosVDCollisionChannelsCustomizationBase::CustomizeChildren(StructPropertyHandle, StructBuilder, StructCustomizationUtils);
+
+	if (!CachedCollisionChannelInfos)
+	{
+		return;
+	}
+	
+	CurrentPropertyDataHandle = MakeShared<FChaosVDDetailsPropertyDataHandle<FChaosVDCollisionResponseParams>>(StructPropertyHandle);
+	FChaosVDCollisionResponseParams* CollisionResponseParams = CurrentPropertyDataHandle->GetDataInstance();
+
+	if (!CollisionResponseParams)
+	{
+		CurrentPropertyDataHandle = nullptr;
+		return;
+	}
+
+	IDetailGroup& CollisionGroup = StructBuilder.AddGroup(TEXT("CollisionResponseParams"), LOCTEXT("CollisionResponseQueryParamsLabel", "Collision Response Query Params"));
+	CollisionGroup.EnableReset(false);
+
+	CollisionGroup.HeaderRow()
+	.NameContent()
+	[
+		SNew(SHorizontalBox)
+		+ SHorizontalBox::Slot()
+		.HAlign(HAlign_Right)
+		[ 
+			SNew(STextBlock)
+			.Text(LOCTEXT("CollisionResponsesLabel", "Collision Response Query params"))
+			.Font(IDetailLayoutBuilder::GetDetailFont())
+		]
+	];
+
+	if (bChannelInfoBuiltFromDefaults)
+	{
+		CollisionGroup.AddWidgetRow()
+		.WholeRowContent()
+		[
+			SNew(SChaosVDWarningMessageBox).WarningText(FChaosVDDetailsCustomizationUtils::GetDefaultCollisionChannelsUseWarningMessage())
+		];
+	}
+	
+	FChaosVDCollisionChannelStateGetter CollisionChannelStateGetter;
+	CollisionChannelStateGetter.BindSP(this, &FChaosVDCollisionResponseParamsCustomization::GetCurrentCollisionResponseForChannel);
+	
+	FChaosVDDetailsCustomizationUtils::BuildCollisionChannelMatrix(CollisionChannelStateGetter,
+																	{ CachedCollisionChannelInfos->CustomChannelsNames, UE_ARRAY_COUNT(CachedCollisionChannelInfos->CustomChannelsNames)},
+																	CollisionGroup);
+}
+
+TSharedRef<IPropertyTypeCustomization> FChaosVDCollisionObjectParamsCustomization::MakeInstance(TSharedRef<SChaosVDMainTab> MainTab)
+{
+	return MakeShared<FChaosVDCollisionObjectParamsCustomization>(MainTab);
+}
+
+void FChaosVDCollisionObjectParamsCustomization::CustomizeChildren(TSharedRef<IPropertyHandle> StructPropertyHandle, IDetailChildrenBuilder& StructBuilder, IPropertyTypeCustomizationUtils& StructCustomizationUtils)
+{
+	FChaosVDCollisionChannelsCustomizationBase::CustomizeChildren(StructPropertyHandle, StructBuilder, StructCustomizationUtils);
+	
+	if (!CachedCollisionChannelInfos)
+	{
+		return;
+	}
+	
+	CurrentPropertyDataHandle = MakeShared<FChaosVDDetailsPropertyDataHandle<FChaosVDCollisionObjectQueryParams>>(StructPropertyHandle);
+	FChaosVDCollisionObjectQueryParams* CollisionResponseParams = CurrentPropertyDataHandle->GetDataInstance();
+
+	if (!CollisionResponseParams)
+	{
+		CurrentPropertyDataHandle = nullptr;
+		return;
+	}
+
+	IDetailGroup& CollisionGroup = StructBuilder.AddGroup(TEXT("CollisionObjectResponseParams"), LOCTEXT("CollisionObjectResponseQueryParamsLabel", "Collision Response Query Params") );
+	CollisionGroup.EnableReset(false);
+
+	constexpr int32 MaxChannels = FChaosVDDetailsCustomizationUtils::GetMaxCollisionChannelIndex();
+	for (int32 ChannelIndex = 0 ; ChannelIndex < MaxChannels; ChannelIndex++)
+	{
+		FChaosVDCollisionChannelInfo& ChannelInfo = CachedCollisionChannelInfos->CustomChannelsNames[ChannelIndex];
+
+		if (ChannelInfo.bIsTraceType)
+		{
+			continue;
+		}
+
+		// Currently, all details panel in CVD are read only
+		constexpr bool bIsEditable = false;
+
+		CollisionGroup.AddWidgetRow()
+		.NameContent()
+		[
+			SNew(STextBlock)
+			.Text(FText::FromString(CachedCollisionChannelInfos->CustomChannelsNames[ChannelIndex].DisplayName))
+			.Font(IDetailLayoutBuilder::GetDetailFont())
+		]
+		.ValueContent()
+		[
+			SNew(SBox)
+			.IsEnabled(bIsEditable)
+			.WidthOverride(50)
+			.Content()
+			[
+				SNew(SCheckBox)
+				.IsChecked_Raw(this, &FChaosVDCollisionObjectParamsCustomization::GetCurrentObjectFlag, ChannelIndex)
+			]
+		];
+	}
+}
+
+ECheckBoxState FChaosVDCollisionObjectParamsCustomization::GetCurrentObjectFlag(int32 ChannelIndex) const
+{
+	FChaosVDCollisionObjectQueryParams* CollisionObjectResponseParams = CurrentPropertyDataHandle ? CurrentPropertyDataHandle->GetDataInstance() : nullptr;
+	if (!CollisionObjectResponseParams)
+	{
+		return ECheckBoxState::Undetermined;
+	}
+
+	return (CollisionObjectResponseParams->ObjectTypesToQuery & ECC_TO_BITFIELD(ChannelIndex)) != 0 ? ECheckBoxState::Checked : ECheckBoxState::Unchecked;
+}
+
 #undef LOCTEXT_NAMESPACE
diff --git a/Engine/Plugins/ChaosVD/Source/ChaosVD/Private/DetailsCustomizations/ChaosVDQueryDataWrappersCustomizationDetails.h b/Engine/Plugins/ChaosVD/Source/ChaosVD/Private/DetailsCustomizations/ChaosVDQueryDataWrappersCustomizationDetails.h
--- a/Engine/Plugins/ChaosVD/Source/ChaosVD/Private/DetailsCustomizations/ChaosVDQueryDataWrappersCustomizationDetails.h
+++ b/Engine/Plugins/ChaosVD/Source/ChaosVD/Private/DetailsCustomizations/ChaosVDQueryDataWrappersCustomizationDetails.h
@@ -1,8 +1,15 @@
 ﻿// Copyright Epic Games, Inc. All Rights Reserved.
 
 #pragma once
+#include "ChaosVDDetailsCustomizationUtils.h"
 #include "IDetailCustomization.h"
+#include "IDetailGroup.h"
 #include "IPropertyTypeCustomization.h"
+#include "DataWrappers/ChaosVDQueryDataWrappers.h"
+
+class SChaosVDMainTab;
+enum ECollisionResponse : int;
+struct FChaosVDCollisionChannelInfo;
 
 /** Custom property layout for the ChaosVD SQ Data wrapper struct */
 class FChaosVDQueryDataWrappersCustomizationDetails : public IPropertyTypeCustomization
@@ -36,3 +43,62 @@ public:
 
 	virtual void CustomizeDetails(IDetailLayoutBuilder& DetailBuilder) override;
 };
+
+/** Custom details panel for the ChaosVD SQ Data Collision Response View */
+class FChaosVDCollisionChannelsCustomizationBase : public IPropertyTypeCustomization
+{
+public:
+
+	FChaosVDCollisionChannelsCustomizationBase(const TSharedPtr<SChaosVDMainTab>& InMainTab);
+
+	virtual void CustomizeHeader(TSharedRef<IPropertyHandle> StructPropertyHandle, FDetailWidgetRow& HeaderRow, IPropertyTypeCustomizationUtils& StructCustomizationUtils) override {}
+	virtual void CustomizeChildren(TSharedRef<IPropertyHandle> StructPropertyHandle, IDetailChildrenBuilder& StructBuilder, IPropertyTypeCustomizationUtils& StructCustomizationUtils) override;
+
+protected:
+	void UpdateCollisionChannelsInfoCache(TSharedPtr<FChaosVDCollisionChannelsInfoContainer> NewCollisionChannelsInfo);
+	
+	TSharedPtr<FChaosVDCollisionChannelsInfoContainer> CachedCollisionChannelInfos;
+	bool bChannelInfoBuiltFromDefaults = true;
+
+	TWeakPtr<SChaosVDMainTab> MainTabWeakPtr = nullptr;
+};
+
+/** Custom details panel for the ChaosVD SQ Data Collision Response View */
+class FChaosVDCollisionResponseParamsCustomization : public FChaosVDCollisionChannelsCustomizationBase
+{
+public:
+	FChaosVDCollisionResponseParamsCustomization(const TSharedPtr<SChaosVDMainTab>& InMainTab) : FChaosVDCollisionChannelsCustomizationBase(InMainTab)
+	{
+	}
+
+	static TSharedRef<IPropertyTypeCustomization> MakeInstance(TSharedRef<SChaosVDMainTab> MainTab);
+
+	virtual void CustomizeHeader(TSharedRef<IPropertyHandle> StructPropertyHandle, FDetailWidgetRow& HeaderRow, IPropertyTypeCustomizationUtils& StructCustomizationUtils) override {}
+	virtual void CustomizeChildren(TSharedRef<IPropertyHandle> StructPropertyHandle, IDetailChildrenBuilder& StructBuilder, IPropertyTypeCustomizationUtils& StructCustomizationUtils) override;
+
+private:
+	
+	ECollisionResponse GetCurrentCollisionResponseForChannel(int32 ChannelIndex) const;
+
+	TSharedPtr<FChaosVDDetailsPropertyDataHandle<FChaosVDCollisionResponseParams>> CurrentPropertyDataHandle = nullptr;
+};
+
+/** Custom details panel for the ChaosVD SQ Data Collision Object Response View */
+class FChaosVDCollisionObjectParamsCustomization : public FChaosVDCollisionChannelsCustomizationBase
+{
+public:
+	FChaosVDCollisionObjectParamsCustomization(const TSharedPtr<SChaosVDMainTab>& InMainTab) : FChaosVDCollisionChannelsCustomizationBase(InMainTab)
+	{
+	}
+
+	static TSharedRef<IPropertyTypeCustomization> MakeInstance(TSharedRef<SChaosVDMainTab> MainTab);
+
+	virtual void CustomizeHeader(TSharedRef<IPropertyHandle> StructPropertyHandle, FDetailWidgetRow& HeaderRow, IPropertyTypeCustomizationUtils& StructCustomizationUtils) override {}
+	virtual void CustomizeChildren(TSharedRef<IPropertyHandle> StructPropertyHandle, IDetailChildrenBuilder& StructBuilder, IPropertyTypeCustomizationUtils& StructCustomizationUtils) override;
+
+private:
+	
+	ECheckBoxState GetCurrentObjectFlag(int32 ChannelIndex) const;
+
+	TSharedPtr<FChaosVDDetailsPropertyDataHandle<FChaosVDCollisionObjectQueryParams>> CurrentPropertyDataHandle = nullptr;
+};
diff --git a/Engine/Plugins/ChaosVD/Source/ChaosVD/Private/DetailsCustomizations/ChaosVDShapeDataCustomization.cpp b/Engine/Plugins/ChaosVD/Source/ChaosVD/Private/DetailsCustomizations/ChaosVDShapeDataCustomization.cpp
new file mode 100644
--- /dev/null
+++ b/Engine/Plugins/ChaosVD/Source/ChaosVD/Private/DetailsCustomizations/ChaosVDShapeDataCustomization.cpp
@@ -0,0 +1,376 @@
+﻿// Copyright Epic Games, Inc. All Rights Reserved.
+
+#include "ChaosVDShapeDataCustomization.h"
+
+#include "ChaosVDDetailsCustomizationUtils.h"
+#include "ChaosVDEngine.h"
+#include "ChaosVDRecording.h"
+#include "ChaosVDScene.h"
+#include "DetailCategoryBuilder.h"
+#include "DetailLayoutBuilder.h"
+#include "DetailWidgetRow.h"
+#include "IDetailChildrenBuilder.h"
+#include "IDetailGroup.h"
+#include "PropertyHandle.h"
+#include "SEnumCombo.h"
+#include "Engine/EngineTypes.h"
+#include "Physics/PhysicsFiltering.h"
+#include "Widgets/SChaosVDMainTab.h"
+#include "Widgets/SChaosVDWarningMessageBox.h"
+#include "Widgets/Input/SCheckBox.h"
+#include "Widgets/Input/SSpinBox.h"
+#include "Widgets/Layout/SBox.h"
+#include "Widgets/Text/STextBlock.h"
+
+#define LOCTEXT_NAMESPACE "ChaosVisualDebugger"
+
+TSharedRef<IPropertyTypeCustomization> FChaosVDShapeDataCustomization::MakeInstance(TSharedRef<SChaosVDMainTab> MainTab)
+{
+	return MakeShared<FChaosVDShapeDataCustomization>(MainTab);
+}
+
+FChaosVDShapeDataCustomization::FChaosVDShapeDataCustomization(const TSharedPtr<SChaosVDMainTab>& InMainTab)
+{
+	MainTabWeakPtr = InMainTab;
+}
+
+FChaosVDShapeDataCustomization::~FChaosVDShapeDataCustomization()
+{
+	RegisterCVDScene(nullptr);
+}
+
+void FChaosVDShapeDataCustomization::CustomizeHeader(TSharedRef<IPropertyHandle> StructPropertyHandle, FDetailWidgetRow& HeaderRow, IPropertyTypeCustomizationUtils& StructCustomizationUtils)
+{
+	constexpr bool bPropagateToChildren = true;
+	HeaderRow.OverrideResetToDefault(FResetToDefaultOverride::Hide(bPropagateToChildren));
+}
+
+void FChaosVDShapeDataCustomization::CustomizeChildren(TSharedRef<IPropertyHandle> StructPropertyHandle, IDetailChildrenBuilder& StructBuilder, IPropertyTypeCustomizationUtils& StructCustomizationUtils)
+{
+	CurrentShapeDataHandle = nullptr;
+
+	TSharedPtr<SChaosVDMainTab> MainTabPtr =  MainTabWeakPtr.Pin(); 
+	TSharedPtr<FChaosVDScene> Scene = MainTabPtr ? MainTabPtr->GetChaosVDEngineInstance()->GetCurrentScene() : nullptr;
+
+	RegisterCVDScene(Scene);
+
+	if (!Scene)
+	{
+		return;
+	}
+
+	if (TSharedPtr<FChaosVDRecording> CVDRecording = Scene->LoadedRecording)
+	{
+		UpdateCollisionChannelsInfoCache(CVDRecording->GetCollisionChannelsInfoContainer());
+	}
+	else
+	{
+		UpdateCollisionChannelsInfoCache(nullptr);
+	}
+
+	if (!CachedCollisionChannelInfos)
+	{
+		return;
+	}
+
+	IDetailLayoutBuilder& ParentLayoutBuilder = StructBuilder.GetParentCategory().GetParentLayout();
+	IDetailCategoryBuilder& CollisionCategoryBuilder = ParentLayoutBuilder.EditCategory(TEXT("ShapeCollisionSettings"), LOCTEXT("ShapeDataCollisionSettingsLabel", "Shape Collision Settings"));
+
+	CurrentShapeDataHandle = MakeShared<FChaosVDDetailsPropertyDataHandle<FChaosVDShapeCollisionData>>(StructPropertyHandle);
+
+	FChaosVDShapeCollisionData* ShapeData = CurrentShapeDataHandle->GetDataInstance();
+
+	if (!ShapeData || !ShapeData->bIsValid)
+	{
+		StructBuilder
+		.AddCustomRow(FText::GetEmpty())
+		.WholeRowContent()
+		[
+			SNew(SChaosVDWarningMessageBox).WarningText(LOCTEXT("InvalidShapeDataWarningMessageBox", "Warning : Failed to load shape data for the selected shape! "))
+		];
+
+		return;
+	}
+
+	CacheCurrentCollisionChannelData();
+	
+	TSharedPtr<IPropertyHandle> CollisionTraceTypeHandle = StructPropertyHandle->GetChildHandle(GET_MEMBER_NAME_CHECKED(FChaosVDShapeCollisionData, CollisionTraceType));
+	CollisionCategoryBuilder.AddProperty(CollisionTraceTypeHandle.ToSharedRef());
+	
+	constexpr bool bIsEditable = false;
+
+	const FText& CollisionEnabledRowLabel = LOCTEXT("CollisionEnabledLabel", "Collision Enabled");
+	CollisionCategoryBuilder.AddCustomRow(CollisionEnabledRowLabel)
+	.NameContent()
+	[
+		SNew(STextBlock)
+		.Text(CollisionEnabledRowLabel)
+		.Font(IDetailLayoutBuilder::GetDetailFont())
+	]
+	.ValueContent()
+	[
+		SNew(SBox)
+		[ 
+			SNew(STextBlock)
+			.Text_Raw(this, &FChaosVDShapeDataCustomization::GetCurrentCollisionEnabledAsText) 
+			.IsEnabled(bIsEditable)
+			.Font(IDetailLayoutBuilder::GetDetailFontBold())
+		]
+	];
+
+	const FText& ComplexCollisionRowLabel = LOCTEXT("ComplexCollisionLabel", "Is Complex");
+	CollisionCategoryBuilder.AddCustomRow(ComplexCollisionRowLabel)
+	.NameContent()
+	[
+		SNew(STextBlock)
+		.Text(ComplexCollisionRowLabel)
+		.Font(IDetailLayoutBuilder::GetDetailFont())
+	]
+	.ValueContent()
+	[
+		SNew(SBox)
+		.IsEnabled(bIsEditable)
+		.WidthOverride(50)
+		.VAlign(VAlign_Center)
+		.Content()
+		[
+			SNew(SCheckBox)
+			.IsChecked_Raw(this, &FChaosVDShapeDataCustomization::GetCurrentIsComplexCollision)
+		]
+	];
+
+	BuildCollisionFilerDataLayout(ECollisionFilterDataType::Query, LOCTEXT("QueryFilterDataRowLabel", "Query Filter Data"), ParentLayoutBuilder, ECollisionFilterLayoutFlags::StartExpanded);
+
+	BuildCollisionFilerDataLayout(ECollisionFilterDataType::Sim, LOCTEXT("SimFilterDataRowLabel", "Sim Filter Data"), ParentLayoutBuilder, ECollisionFilterLayoutFlags::StartExpanded);
+}
+
+void FChaosVDShapeDataCustomization::BuildCollisionFilerDataLayout(ECollisionFilterDataType FilterDataType, const FText& InDetailsGroupLabel, IDetailLayoutBuilder& ParentLayoutBuilder, ECollisionFilterLayoutFlags LayoutFlags)
+{
+	const bool bStartExpanded = EnumHasAnyFlags(LayoutFlags, ECollisionFilterLayoutFlags::StartExpanded);
+	const bool bIsEditable = EnumHasAnyFlags(LayoutFlags, ECollisionFilterLayoutFlags::IsEditable);
+	const bool bIsAdvanced = EnumHasAnyFlags(LayoutFlags, ECollisionFilterLayoutFlags::IsAdvanced);
+
+	IDetailCategoryBuilder& FilterDataCategory = ParentLayoutBuilder.EditCategory(FName(InDetailsGroupLabel.ToString()), InDetailsGroupLabel);
+
+	if (bChannelInfoBuiltFromDefaults)
+	{
+		FilterDataCategory.AddCustomRow(FText::GetEmpty())
+		.WholeRowContent()
+		[
+			SNew(SChaosVDWarningMessageBox).WarningText(FChaosVDDetailsCustomizationUtils::GetDefaultCollisionChannelsUseWarningMessage())
+		];
+	}
+
+	const FText CollisionChannelRowLabel = LOCTEXT("FilterDataCollisionChannelLabel", "Collision Channel");
+	FilterDataCategory.AddCustomRow(CollisionChannelRowLabel)
+	.NameContent()
+	[
+		SNew(STextBlock)
+		.Text(CollisionChannelRowLabel)
+		.Font(IDetailLayoutBuilder::GetDetailFont())
+	]
+	.ValueContent()
+	[
+		SNew(STextBlock)
+		.Text_Raw(this, &FChaosVDShapeDataCustomization::GetCurrentCollisionChannelAsText, FilterDataType)
+		.IsEnabled(bIsEditable)
+		.Font(IDetailLayoutBuilder::GetDetailFontBold())
+	];
+
+	const FText ExtraFilterRowLabel = LOCTEXT("FilterDataExtraFilterLabel", "Extra Filter");
+	FilterDataCategory.AddCustomRow(ExtraFilterRowLabel)
+	.NameContent()
+	[
+		SNew(STextBlock)
+		.Text(ExtraFilterRowLabel)
+		.Font(IDetailLayoutBuilder::GetDetailFont())
+	]
+	.ValueContent()
+	[
+		SNew(SSpinBox<uint8>)
+		.Value_Raw(this, &FChaosVDShapeDataCustomization::GetCurrentExtraFilter, FilterDataType)
+		.IsEnabled(bIsEditable)
+		.Font(IDetailLayoutBuilder::GetDetailFont())
+	];
+	
+	AddExtraCollisionFilterFlags(FilterDataType, LOCTEXT("FilterDataExtraCollisionFilterTags", "Chaos Collision Filter Flags"), FilterDataCategory);
+
+	IDetailGroup& CollisionResponseGroup = FilterDataCategory.AddGroup(TEXT("FilterDataCollisionResponseGroup"), LOCTEXT("FilterDataCollisionResponseGroupLabel", "Collision Response Flags"), bIsAdvanced, bStartExpanded);
+	CollisionResponseGroup.EnableReset(false);
+
+	FChaosVDCollisionChannelStateGetter CollisionSimChannelStateGetter;
+	CollisionSimChannelStateGetter.BindSPLambda(this, [this, FilterDataType](int32 ChannelIndex)
+	{
+		return GetCurrentCollisionResponseForChannel(ChannelIndex, FilterDataType);
+	});
+
+	FChaosVDDetailsCustomizationUtils::BuildCollisionChannelMatrix(CollisionSimChannelStateGetter,
+	                                                               { CachedCollisionChannelInfos->CustomChannelsNames, UE_ARRAY_COUNT(CachedCollisionChannelInfos->CustomChannelsNames)},
+	                                                               CollisionResponseGroup);
+}
+
+void FChaosVDShapeDataCustomization::UpdateCollisionChannelsInfoCache(const TSharedPtr<FChaosVDCollisionChannelsInfoContainer>& NewCollisionChannelsInfo)
+{
+	if (NewCollisionChannelsInfo)
+	{
+		CachedCollisionChannelInfos = NewCollisionChannelsInfo;
+		bChannelInfoBuiltFromDefaults = false;
+	}
+	else
+	{
+		// Fallback to engine channels name using the enum metadata
+		CachedCollisionChannelInfos = FChaosVDDetailsCustomizationUtils::BuildDefaultCollisionChannelInfo();
+		bChannelInfoBuiltFromDefaults = true;
+	}
+}
+FText FChaosVDShapeDataCustomization::GetCurrentCollisionEnabledAsText() const
+{
+	if (!CurrentShapeDataHandle)
+	{
+		return FText::GetEmpty();
+	}
+
+	FChaosVDShapeCollisionData* Data = CurrentShapeDataHandle->GetDataInstance();
+
+	return Data ? UEnum::GetDisplayValueAsText(CollisionEnabledFromFlags(Data->bQueryCollision, Data->bSimCollision, Data->bIsProbe)) : FText::GetEmpty();
+}
+
+ECheckBoxState FChaosVDShapeDataCustomization::GetCurrentIsComplexCollision() const
+{
+	if (FChaosVDShapeCollisionData* Data = CurrentShapeDataHandle ? CurrentShapeDataHandle->GetDataInstance() : nullptr)
+	{
+		return Data->bIsComplex ? ECheckBoxState::Checked : ECheckBoxState::Unchecked;
+	}
+
+	return ECheckBoxState::Undetermined;
+}
+
+FText FChaosVDShapeDataCustomization::GetCurrentCollisionChannelAsText(ECollisionFilterDataType FilterDataType) const
+{
+	if (FChaosVDShapeCollisionData* Data = CurrentShapeDataHandle ? CurrentShapeDataHandle->GetDataInstance() : nullptr)
+	{
+		FChaosVDCollisionFilterData& FilterDataRef = FilterDataType == ECollisionFilterDataType::Query ? Data->QueryData : Data->SimData;
+		uint8 QueryExtraFilter = 0;
+		return UEnum::GetDisplayValueAsText(GetCollisionChannelAndExtraFilter(FilterDataRef.Word3, QueryExtraFilter));
+	}
+
+	return FText::GetEmpty();
+}
+
+uint8 FChaosVDShapeDataCustomization::GetCurrentExtraFilter(ECollisionFilterDataType FilterDataType) const
+{
+	if (FChaosVDShapeCollisionData* Data = CurrentShapeDataHandle ? CurrentShapeDataHandle->GetDataInstance() : nullptr)
+	{
+		FChaosVDCollisionFilterData& FilterDataRef = FilterDataType== ECollisionFilterDataType::Query ? Data->QueryData : Data->SimData;
+		uint8 QueryExtraFilter = 0;
+		GetCollisionChannelAndExtraFilter(FilterDataRef.Word3, QueryExtraFilter);
+	
+		return QueryExtraFilter;
+	}
+
+	return 0;
+}
+
+ECollisionResponse FChaosVDShapeDataCustomization::GetCurrentCollisionResponseForChannel(int32 ChannelIndex, ECollisionFilterDataType Type) const
+{
+	constexpr int32 MaxChannels = FChaosVDDetailsCustomizationUtils::GetMaxCollisionChannelIndex();
+
+	if (ChannelIndex >= MaxChannels)
+	{
+		return ECollisionResponse::ECR_MAX;
+	}
+
+	if (!CurrentShapeDataHandle->GetDataInstance())
+	{
+		return ECollisionResponse::ECR_MAX;
+	}
+
+	switch (Type)
+	{
+		case Query:
+			return static_cast<ECollisionResponse>(CachedQueryCollisionResponsesFlagsPerChannel.EnumArray[ChannelIndex]);
+		case Sim:
+			return static_cast<ECollisionResponse>(CachedSimCollisionResponsesFlagsPerChannel.EnumArray[ChannelIndex]);
+		default:
+			break;
+	}
+
+	return ECollisionResponse::ECR_MAX;
+}
+
+ECheckBoxState FChaosVDShapeDataCustomization::GetCurrentStateForFilteringFlag(ECollisionFilterDataType Type, Chaos::EFilterFlags FilteringFlag) const
+{
+	if (FChaosVDShapeCollisionData* Data = CurrentShapeDataHandle ? CurrentShapeDataHandle->GetDataInstance() : nullptr)
+	{
+		FChaosVDCollisionFilterData& FilterDataRef = Type == Query ? Data->QueryData : Data->SimData;
+		return FilterDataRef.Word3 & static_cast<uint8>(FilteringFlag) ? ECheckBoxState::Checked : ECheckBoxState::Unchecked;
+	}
+	
+	return ECheckBoxState::Undetermined;
+}
+void FChaosVDShapeDataCustomization::AddExtraCollisionFilterFlags(ECollisionFilterDataType Type, const FText& IntSectionLabel, IDetailCategoryBuilder& CategoryBuilder)
+{
+	IDetailGroup& ExtraFilterFlagsGroup = CategoryBuilder.AddGroup(FName(IntSectionLabel.ToString()), IntSectionLabel);
+	ExtraFilterFlagsGroup.EnableReset(false);
+
+	AddCollisionFilterFlagRow(Type, Chaos::EFilterFlags::SimpleCollision, ExtraFilterFlagsGroup);
+	AddCollisionFilterFlagRow(Type, Chaos::EFilterFlags::ComplexCollision, ExtraFilterFlagsGroup);
+	AddCollisionFilterFlagRow(Type, Chaos::EFilterFlags::ContactNotify, ExtraFilterFlagsGroup);
+	AddCollisionFilterFlagRow(Type, Chaos::EFilterFlags::CCD, ExtraFilterFlagsGroup);
+	AddCollisionFilterFlagRow(Type, Chaos::EFilterFlags::ModifyContacts, ExtraFilterFlagsGroup);
+	AddCollisionFilterFlagRow(Type, Chaos::EFilterFlags::StaticShape, ExtraFilterFlagsGroup);
+	AddCollisionFilterFlagRow(Type, Chaos::EFilterFlags::KinematicKinematicPairs, ExtraFilterFlagsGroup);
+}
+
+void FChaosVDShapeDataCustomization::AddCollisionFilterFlagRow(ECollisionFilterDataType Type, Chaos::EFilterFlags Flag, IDetailGroup& DetailGroup) const
+{
+	TAttribute<ECheckBoxState> CheckBoxStateAttribute;
+	CheckBoxStateAttribute.BindSPLambda(this, [this, Flag, Type]()
+	{
+		return GetCurrentStateForFilteringFlag(Type, Flag);
+	});
+
+	FChaosVDDetailsCustomizationUtils::AddWidgetRowForCheckboxValue(MoveTemp(CheckBoxStateAttribute), FText::FromString(Chaos::LexToString(Flag)), DetailGroup);
+}
+
+void FChaosVDShapeDataCustomization::CacheCurrentCollisionChannelData()
+{
+	if (FChaosVDShapeCollisionData* Data = CurrentShapeDataHandle ? CurrentShapeDataHandle->GetDataInstance() : nullptr)
+	{
+		FCollisionFilterData QueryFilterData = FChaosVDDetailsCustomizationUtils::ConvertToEngineFilteringData(Data->QueryData);
+		CachedQueryCollisionResponsesFlagsPerChannel = ExtractQueryCollisionResponseContainer(QueryFilterData);
+
+		FCollisionFilterData SimFilterData = FChaosVDDetailsCustomizationUtils::ConvertToEngineFilteringData(Data->QueryData);
+		CachedSimCollisionResponsesFlagsPerChannel = ExtractSimCollisionResponseContainer(SimFilterData);
+	}
+}
+
+void FChaosVDShapeDataCustomization::RegisterCVDScene(const TSharedPtr<FChaosVDScene>& InScene)
+{
+	TSharedPtr<FChaosVDScene> CurrentScene = SceneWeakPtr.Pin();
+	if (InScene != CurrentScene)
+	{
+		if (CurrentScene)
+		{
+			CurrentScene->OnSceneUpdated().RemoveAll(this);
+		}
+
+		if (InScene)
+		{
+			InScene->OnSceneUpdated().AddSP(this, &FChaosVDShapeDataCustomization::HandleSceneUpdated);
+		}
+
+		SceneWeakPtr = InScene;
+	}
+}
+
+void FChaosVDShapeDataCustomization::HandleSceneUpdated()
+{
+	// To avoid reconstructing the data channel flags every time a row in the collision channel response matrix is updated, we process it once and cache it
+	// but we need to update it each time the CVD scene is updated as the source data of the currents elected particle might have changed.
+	CacheCurrentCollisionChannelData();
+
+	// Note: Intentionally not calling UpdateCollisionChannelsInfoCache, that is only updated once when a recording is loaded, therefore we don't nee to do it again
+}
+
+#undef LOCTEXT_NAMESPACE
diff --git a/Engine/Plugins/ChaosVD/Source/ChaosVD/Private/DetailsCustomizations/ChaosVDShapeDataCustomization.h b/Engine/Plugins/ChaosVD/Source/ChaosVD/Private/DetailsCustomizations/ChaosVDShapeDataCustomization.h
new file mode 100644
--- /dev/null
+++ b/Engine/Plugins/ChaosVD/Source/ChaosVD/Private/DetailsCustomizations/ChaosVDShapeDataCustomization.h
@@ -0,0 +1,80 @@
+﻿// Copyright Epic Games, Inc. All Rights Reserved.
+
+#pragma once
+#include "ChaosVDDetailsCustomizationUtils.h"
+#include "IPropertyTypeCustomization.h"
+#include "DataWrappers/ChaosVDCollisionDataWrappers.h"
+#include "Engine/EngineTypes.h"
+
+class IDetailCategoryBuilder;
+class IDetailGroup;
+class SChaosVDMainTab;
+
+enum class ECheckBoxState : uint8;
+
+struct FChaosVDCollisionChannelsInfoContainer;
+
+/** Custom details panel for the ChaosVD SQ Data Collision Response View */
+class FChaosVDShapeDataCustomization : public IPropertyTypeCustomization
+{
+public:
+	static TSharedRef<IPropertyTypeCustomization> MakeInstance(TSharedRef<SChaosVDMainTab> MainTab);
+	
+	FChaosVDShapeDataCustomization(const TSharedPtr<SChaosVDMainTab>& InMainTab);
+	virtual ~FChaosVDShapeDataCustomization() override;
+
+	virtual void CustomizeHeader(TSharedRef<IPropertyHandle> StructPropertyHandle, FDetailWidgetRow& HeaderRow, IPropertyTypeCustomizationUtils& StructCustomizationUtils) override;
+	virtual void CustomizeChildren(TSharedRef<IPropertyHandle> StructPropertyHandle, IDetailChildrenBuilder& StructBuilder, IPropertyTypeCustomizationUtils& StructCustomizationUtils) override;
+
+private:
+	void UpdateCollisionChannelsInfoCache(const TSharedPtr<FChaosVDCollisionChannelsInfoContainer>& NewCollisionChannelsInfo);
+
+	FText GetCurrentCollisionEnabledAsText() const;
+
+	ECheckBoxState GetCurrentIsComplexCollision() const;
+
+	enum ECollisionFilterDataType
+	{
+		Query,
+		Sim
+	};
+
+	enum class ECollisionFilterLayoutFlags : uint8
+	{
+		None = 0,
+		StartExpanded = 1 << 0,
+		IsEditable = 1 << 1,
+		IsAdvanced = 1 << 2
+	};
+	FRIEND_ENUM_CLASS_FLAGS(ECollisionFilterLayoutFlags)
+
+	void BuildCollisionFilerDataLayout(ECollisionFilterDataType FilterDataType, const FText& InDetailsGroupLabel, IDetailLayoutBuilder& ParentLayoutBuilder, ECollisionFilterLayoutFlags LayoutFlags);
+
+	FText GetCurrentCollisionChannelAsText(ECollisionFilterDataType FilterDataType) const;
+	uint8 GetCurrentExtraFilter(ECollisionFilterDataType FilterDataType) const;
+
+	ECollisionResponse GetCurrentCollisionResponseForChannel(int32 ChannelIndex, ECollisionFilterDataType Type) const;
+
+	void AddCollisionFilterFlagRow(ECollisionFilterDataType Type, Chaos::EFilterFlags Flag, IDetailGroup& DetailGroup) const;
+	ECheckBoxState GetCurrentStateForFilteringFlag(ECollisionFilterDataType Type, Chaos::EFilterFlags FilteringFlag) const;
+
+	void AddExtraCollisionFilterFlags(ECollisionFilterDataType Type, const FText& IntSectionLabel, IDetailCategoryBuilder& CategoryBuilder);
+
+	void CacheCurrentCollisionChannelData();
+
+	void RegisterCVDScene(const TSharedPtr<FChaosVDScene>& InScene);
+
+	void HandleSceneUpdated();
+
+	TSharedPtr<FChaosVDCollisionChannelsInfoContainer> CachedCollisionChannelInfos;
+	bool bChannelInfoBuiltFromDefaults = true;
+	
+	FCollisionResponseContainer CachedSimCollisionResponsesFlagsPerChannel;
+	FCollisionResponseContainer CachedQueryCollisionResponsesFlagsPerChannel;
+
+	TSharedPtr<FChaosVDDetailsPropertyDataHandle<FChaosVDShapeCollisionData>> CurrentShapeDataHandle = nullptr;
+
+	TWeakPtr<FChaosVDScene> SceneWeakPtr = nullptr;
+
+	TWeakPtr<SChaosVDMainTab> MainTabWeakPtr = nullptr;
+};
diff --git a/Engine/Plugins/ChaosVD/Source/ChaosVD/Private/Settings/ChaosVDSceneQueryVisualizationSettings.h b/Engine/Plugins/ChaosVD/Source/ChaosVD/Private/Settings/ChaosVDSceneQueryVisualizationSettings.h
--- a/Engine/Plugins/ChaosVD/Source/ChaosVD/Private/Settings/ChaosVDSceneQueryVisualizationSettings.h
+++ b/Engine/Plugins/ChaosVD/Source/ChaosVD/Private/Settings/ChaosVDSceneQueryVisualizationSettings.h
@@ -46,6 +46,11 @@ public:
 	UPROPERTY(EditAnywhere, Category=DebugDraw)
 	bool bShowText = false;
 
+	/** The depth priority used for while drawing. Can be World or Foreground (with this one the shapes representing the query
+	 * will be drawn on top of the geometry and be always visible) */
+	UPROPERTY(config, EditAnywhere, Category=DebugDraw)
+	TEnumAsByte<ESceneDepthPriorityGroup> DepthPriority = ESceneDepthPriorityGroup::SDPG_World;
+
 	UPROPERTY()
 	EChaosVDSQFrameVisualizationMode CurrentVisualizationMode = EChaosVDSQFrameVisualizationMode::AllEnabledQueries;
 
diff --git a/Engine/Plugins/ChaosVD/Source/ChaosVD/Private/Trace/ChaosVDTraceProvider.cpp b/Engine/Plugins/ChaosVD/Source/ChaosVD/Private/Trace/ChaosVDTraceProvider.cpp
--- a/Engine/Plugins/ChaosVD/Source/ChaosVD/Private/Trace/ChaosVDTraceProvider.cpp
+++ b/Engine/Plugins/ChaosVD/Source/ChaosVD/Private/Trace/ChaosVDTraceProvider.cpp
@@ -10,6 +10,7 @@
 
 #include "Compression/OodleDataCompressionUtil.h"
 #include "DataProcessors/ChaosVDAABBTreeDataProcessor.h"
+#include "DataProcessors/ChaosVDCollisionChannelsInfoDataProcessor.h"
 #include "DataProcessors/ChaosVDDebugDrawBoxDataProcessor.h"
 #include "DataProcessors/ChaosVDDebugDrawImplicitObjectDataProcessor.h"
 #include "DataProcessors/ChaosVDDebugDrawLineDataProcessor.h"
@@ -335,6 +336,10 @@ void FChaosVDTraceProvider::RegisterDefaultDataProcessorsIfNeeded()
 	TSharedPtr<FChaosVDDebugDrawImplicitObjectDataProcessor> DebugDrawImplicitObjectDataProcessor = MakeShared<FChaosVDDebugDrawImplicitObjectDataProcessor>();
 	DebugDrawImplicitObjectDataProcessor->SetTraceProvider(AsShared());
 	RegisterDataProcessor(DebugDrawImplicitObjectDataProcessor);
+	
+	TSharedPtr<FChaosVDCollisionChannelsInfoDataProcessor> CollisionChannelsInfoDataProcessor = MakeShared<FChaosVDCollisionChannelsInfoDataProcessor>();
+	CollisionChannelsInfoDataProcessor->SetTraceProvider(AsShared());
+	RegisterDataProcessor(CollisionChannelsInfoDataProcessor);
 
 	bDefaultDataProcessorsRegistered = true;
 }
diff --git a/Engine/Plugins/ChaosVD/Source/ChaosVD/Private/Trace/DataProcessors/ChaosVDCollisionChannelsInfoDataProcessor.cpp b/Engine/Plugins/ChaosVD/Source/ChaosVD/Private/Trace/DataProcessors/ChaosVDCollisionChannelsInfoDataProcessor.cpp
new file mode 100644
--- /dev/null
+++ b/Engine/Plugins/ChaosVD/Source/ChaosVD/Private/Trace/DataProcessors/ChaosVDCollisionChannelsInfoDataProcessor.cpp
@@ -0,0 +1,29 @@
+﻿// Copyright Epic Games, Inc. All Rights Reserved.
+
+#include "ChaosVDCollisionChannelsInfoDataProcessor.h"
+
+FChaosVDCollisionChannelsInfoDataProcessor::FChaosVDCollisionChannelsInfoDataProcessor() : FChaosVDDataProcessorBase(FChaosVDCollisionChannelsInfoContainer::WrapperTypeName)
+{
+	
+}
+
+bool FChaosVDCollisionChannelsInfoDataProcessor::ProcessRawData(const TArray<uint8>& InData)
+{
+	FChaosVDDataProcessorBase::ProcessRawData(InData);
+
+	TSharedPtr<FChaosVDTraceProvider> ProviderSharedPtr = TraceProvider.Pin();
+	if (!ensure(ProviderSharedPtr.IsValid()))
+	{
+		return false;
+	}
+
+	TSharedPtr<FChaosVDCollisionChannelsInfoContainer> CollisionChannelsData = MakeShared<FChaosVDCollisionChannelsInfoContainer>();
+	const bool bSuccess = Chaos::VisualDebugger::ReadDataFromBuffer(InData, *CollisionChannelsData, ProviderSharedPtr.ToSharedRef());
+	
+	if (bSuccess)
+	{
+		ProviderSharedPtr->GetRecordingForSession()->SetCollisionChannelsInfoContainer(CollisionChannelsData);
+	}
+	
+	return bSuccess;
+}
diff --git a/Engine/Plugins/ChaosVD/Source/ChaosVD/Private/Trace/DataProcessors/ChaosVDCollisionChannelsInfoDataProcessor.h b/Engine/Plugins/ChaosVD/Source/ChaosVD/Private/Trace/DataProcessors/ChaosVDCollisionChannelsInfoDataProcessor.h
new file mode 100644
--- /dev/null
+++ b/Engine/Plugins/ChaosVD/Source/ChaosVD/Private/Trace/DataProcessors/ChaosVDCollisionChannelsInfoDataProcessor.h
@@ -0,0 +1,18 @@
+﻿// Copyright Epic Games, Inc. All Rights Reserved.
+
+#pragma once
+#include "ChaosVDDataProcessorBase.h"
+
+struct FChaosVDSolverFrameData;
+
+/**
+ * Data processor implementation that read any recorded collision channels info
+ */
+class FChaosVDCollisionChannelsInfoDataProcessor final : public FChaosVDDataProcessorBase
+{
+public:
+	explicit FChaosVDCollisionChannelsInfoDataProcessor();
+	
+	virtual bool ProcessRawData(const TArray<uint8>& InData) override;
+};
+
diff --git a/Engine/Plugins/ChaosVD/Source/ChaosVD/Private/Visualizers/ChaosVDSceneQueryDataComponentVisualizer.cpp b/Engine/Plugins/ChaosVD/Source/ChaosVD/Private/Visualizers/ChaosVDSceneQueryDataComponentVisualizer.cpp
--- a/Engine/Plugins/ChaosVD/Source/ChaosVD/Private/Visualizers/ChaosVDSceneQueryDataComponentVisualizer.cpp
+++ b/Engine/Plugins/ChaosVD/Source/ChaosVD/Private/Visualizers/ChaosVDSceneQueryDataComponentVisualizer.cpp
@@ -155,7 +155,7 @@ void FChaosVDSceneQueryDataComponentVisualizer::DrawLineTraceQuery(const UActorC
 		EndLocationToDraw = SceneQueryData.EndLocation;
 	}
 
-	FChaosVDDebugDrawUtils::DrawArrowVector(PDI, SceneQueryData.StartLocation, EndLocationToDraw, DebugText, VisualizationContext.DebugDrawColor, SDPG_Foreground);
+	FChaosVDDebugDrawUtils::DrawArrowVector(PDI, SceneQueryData.StartLocation, EndLocationToDraw, DebugText, VisualizationContext.DebugDrawColor, DebugDrawSettings->DepthPriority);
 
 	PDI->SetHitProxy(nullptr);
 
@@ -180,7 +180,7 @@ void FChaosVDSceneQueryDataComponentVisualizer::DrawOverlapQuery(const UActorCom
 	{
 		const FText DebugText = DebugDrawSettings->bShowText ? FText::FormatOrdered(LOCTEXT("OverlapDebugDrawText", "Type: Overlap \n Tag {1} \n Owner Tag {2}"), FText::AsCultureInvariant(SceneQueryData.CollisionQueryParams.TraceTag.ToString()), FText::AsCultureInvariant(SceneQueryData.CollisionQueryParams.OwnerTag.ToString()))
 								: FText::GetEmpty();
-		FChaosVDDebugDrawUtils::DrawImplicitObject(PDI, VisualizationContext.GeometryGenerator.Pin(), InputShapePtr, FTransform(SceneQueryData.GeometryOrientation, SceneQueryData.StartLocation), VisualizationContext.DebugDrawColor, DebugText, ESceneDepthPriorityGroup::SDPG_Foreground);
+		FChaosVDDebugDrawUtils::DrawImplicitObject(PDI, VisualizationContext.GeometryGenerator.Pin(), InputShapePtr, FTransform(SceneQueryData.GeometryOrientation, SceneQueryData.StartLocation), VisualizationContext.DebugDrawColor, DebugText, DebugDrawSettings->DepthPriority);
 	}
 
 	PDI->SetHitProxy(nullptr);
@@ -204,7 +204,7 @@ void FChaosVDSceneQueryDataComponentVisualizer::DrawSweepQuery(const UActorCompo
 	const Chaos::FConstImplicitObjectPtr InputShapePtr = VisualizationContext.InputGeometry;
 	if (ensure(InputShapePtr))
 	{
-		FChaosVDDebugDrawUtils::DrawImplicitObject(PDI, VisualizationContext.GeometryGenerator.Pin(), InputShapePtr, FTransform(SceneQueryData.GeometryOrientation, SceneQueryData.StartLocation), VisualizationContext.DebugDrawColor, FText::GetEmpty(), ESceneDepthPriorityGroup::SDPG_Foreground);
+		FChaosVDDebugDrawUtils::DrawImplicitObject(PDI, VisualizationContext.GeometryGenerator.Pin(), InputShapePtr, FTransform(SceneQueryData.GeometryOrientation, SceneQueryData.StartLocation), VisualizationContext.DebugDrawColor, FText::GetEmpty(), DebugDrawSettings->DepthPriority);
 
 		FVector EndLocationToDraw;
 		if (SceneQueryData.SQVisitData.IsValidIndex(SceneQueryData.CurrentVisitIndex))
@@ -218,7 +218,7 @@ void FChaosVDSceneQueryDataComponentVisualizer::DrawSweepQuery(const UActorCompo
 			EndLocationToDraw = SceneQueryData.EndLocation;
 		}
 
-		FChaosVDDebugDrawUtils::DrawImplicitObject(PDI, VisualizationContext.GeometryGenerator.Pin(), InputShapePtr, FTransform(SceneQueryData.GeometryOrientation, EndLocationToDraw), VisualizationContext.DebugDrawColor, FText::GetEmpty(), ESceneDepthPriorityGroup::SDPG_Foreground);
+		FChaosVDDebugDrawUtils::DrawImplicitObject(PDI, VisualizationContext.GeometryGenerator.Pin(), InputShapePtr, FTransform(SceneQueryData.GeometryOrientation, EndLocationToDraw), VisualizationContext.DebugDrawColor, FText::GetEmpty(), DebugDrawSettings->DepthPriority);
 	}
 
 	PDI->SetHitProxy(nullptr);
@@ -226,7 +226,7 @@ void FChaosVDSceneQueryDataComponentVisualizer::DrawSweepQuery(const UActorCompo
 	const FText DebugText = DebugDrawSettings->bShowText ? FText::FormatOrdered(LOCTEXT("SweepDebugDrawText", "Type: Sweep \n Tag {1} \n Owner Tag {2}"), FText::AsCultureInvariant(SceneQueryData.CollisionQueryParams.TraceTag.ToString()), FText::AsCultureInvariant(SceneQueryData.CollisionQueryParams.OwnerTag.ToString()))
 							: FText::GetEmpty();
 
-	FChaosVDDebugDrawUtils::DrawArrowVector(PDI, SceneQueryData.StartLocation, SceneQueryData.EndLocation, DebugText, VisualizationContext.DebugDrawColor, SDPG_Foreground);
+	FChaosVDDebugDrawUtils::DrawArrowVector(PDI, SceneQueryData.StartLocation, SceneQueryData.EndLocation, DebugText, VisualizationContext.DebugDrawColor, DebugDrawSettings->DepthPriority);
 
 	if (EnumHasAnyFlags(EChaosVDSceneQueryVisualizationFlags::DrawHits, static_cast<EChaosVDSceneQueryVisualizationFlags>(VisualizationContext.VisualizationFlags)))
 	{
@@ -267,13 +267,13 @@ void FChaosVDSceneQueryDataComponentVisualizer::DrawHits(const UActorComponent*
 		constexpr float CircleRadius = 5.0f;
 		constexpr int32 CircleSegments = 12;
 		constexpr float NormalScale = 10.5f;
-		FChaosVDDebugDrawUtils::DrawCircle(PDI, SQVisitData.HitData.WorldPosition, CircleRadius, CircleSegments, InColor, Thickness, Axes.GetUnitAxis(EAxis::Y), Axes.GetUnitAxis(EAxis::Z), HitPointDebugText,  ESceneDepthPriorityGroup::SDPG_Foreground);
-		FChaosVDDebugDrawUtils::DrawArrowVector(PDI, SQVisitData.HitData.WorldPosition, SQVisitData.HitData.WorldPosition + SQVisitData.HitData.FaceNormal * NormalScale, DebugDrawSettings->bShowText ? HitFaceNormalDebugText : FText::GetEmpty(), (FLinearColor(InColor) * 0.65f).ToFColorSRGB(), ESceneDepthPriorityGroup::SDPG_Foreground);
+		FChaosVDDebugDrawUtils::DrawCircle(PDI, SQVisitData.HitData.WorldPosition, CircleRadius, CircleSegments, InColor, Thickness, Axes.GetUnitAxis(EAxis::Y), Axes.GetUnitAxis(EAxis::Z), HitPointDebugText,  DebugDrawSettings->DepthPriority);
+		FChaosVDDebugDrawUtils::DrawArrowVector(PDI, SQVisitData.HitData.WorldPosition, SQVisitData.HitData.WorldPosition + SQVisitData.HitData.FaceNormal * NormalScale, DebugDrawSettings->bShowText ? HitFaceNormalDebugText : FText::GetEmpty(), (FLinearColor(InColor) * 0.65f).ToFColorSRGB(), DebugDrawSettings->DepthPriority);
 
 		// Hit Face Normal is not used in line traces
 		if (SceneQueryData.Type != EChaosVDSceneQueryType::RayCast)
 		{
-			FChaosVDDebugDrawUtils::DrawArrowVector(PDI, SQVisitData.HitData.WorldPosition, SQVisitData.HitData.WorldPosition + SQVisitData.HitData.WorldNormal * NormalScale, DebugDrawSettings->bShowText ? HitWorldNormalDebugText : FText::GetEmpty(), InColor, ESceneDepthPriorityGroup::SDPG_Foreground);
+			FChaosVDDebugDrawUtils::DrawArrowVector(PDI, SQVisitData.HitData.WorldPosition, SQVisitData.HitData.WorldPosition + SQVisitData.HitData.WorldNormal * NormalScale, DebugDrawSettings->bShowText ? HitWorldNormalDebugText : FText::GetEmpty(), InColor, DebugDrawSettings->DepthPriority);
 		}
 
 		if (TSharedPtr<FChaosVDSolverDataSelectionHandle> CurrentSelection = VisualizationContext.SolverDataSelectionObject->GetCurrentSelectionHandle())
@@ -290,7 +290,7 @@ void FChaosVDSceneQueryDataComponentVisualizer::DrawHits(const UActorComponent*
 				constexpr float HitSelectionBoxSize = CircleRadius * 1.2f;
 
 				FVector SelectionBoxExtents(HitSelectionBoxSize,HitSelectionBoxSize,HitSelectionBoxSize);
-				FChaosVDDebugDrawUtils::DrawBox(PDI, SelectionBoxExtents, FColor::Yellow, SelectionBoxTransform, FText::GetEmpty(), ESceneDepthPriorityGroup::SDPG_Foreground);
+				FChaosVDDebugDrawUtils::DrawBox(PDI, SelectionBoxExtents, FColor::Yellow, SelectionBoxTransform, FText::GetEmpty(), DebugDrawSettings->DepthPriority);
 			}
 		}
 
diff --git a/Engine/Plugins/ChaosVD/Source/ChaosVD/Private/Widgets/SChaosVDCollisionDataInspector.cpp b/Engine/Plugins/ChaosVD/Source/ChaosVD/Private/Widgets/SChaosVDCollisionDataInspector.cpp
--- a/Engine/Plugins/ChaosVD/Source/ChaosVD/Private/Widgets/SChaosVDCollisionDataInspector.cpp
+++ b/Engine/Plugins/ChaosVD/Source/ChaosVD/Private/Widgets/SChaosVDCollisionDataInspector.cpp
@@ -7,6 +7,7 @@
 #include "IStructureDetailsView.h"
 #include "Modules/ModuleManager.h"
 #include "PropertyEditorModule.h"
+#include "SChaosVDMainTab.h"
 #include "Visualizers/ChaosVDSolverCollisionDataComponentVisualizer.h"
 #include "Widgets/Input/SButton.h"
 #include "Widgets/Layout/SScrollBox.h"
@@ -253,7 +254,11 @@ TSharedRef<FName> SChaosVDCollisionDataInspector::GenerateNameForCollisionDataIt
 
 TSharedPtr<IStructureDetailsView> SChaosVDCollisionDataInspector::CreateCollisionDataDetailsView()
 {
-	FPropertyEditorModule& PropertyEditorModule = FModuleManager::GetModuleChecked<FPropertyEditorModule>("PropertyEditor");
+	TSharedPtr<SChaosVDMainTab> MainTabPtr = MainTabWeakPtr.Pin();
+	if (!MainTabPtr)
+	{
+		return nullptr;
+	}
 
 	FStructureDetailsViewArgs StructDetailsViewArgs;
 	FDetailsViewArgs DetailsViewArgs;
@@ -262,7 +267,7 @@ TSharedPtr<IStructureDetailsView> SChaosVDCollisionDataInspector::CreateCollisio
 	DetailsViewArgs.bAllowSearch = false;
 	DetailsViewArgs.bShowScrollBar = false;
 
-	return PropertyEditorModule.CreateStructureDetailView(DetailsViewArgs,StructDetailsViewArgs, nullptr);
+	return MainTabPtr->CreateStructureDetailsView(DetailsViewArgs,StructDetailsViewArgs, nullptr);
 }
 
 EVisibility SChaosVDCollisionDataInspector::GetDetailsSectionVisibility() const
diff --git a/Engine/Plugins/ChaosVD/Source/ChaosVD/Private/Widgets/SChaosVDConstraintDataInspector.cpp b/Engine/Plugins/ChaosVD/Source/ChaosVD/Private/Widgets/SChaosVDConstraintDataInspector.cpp
--- a/Engine/Plugins/ChaosVD/Source/ChaosVD/Private/Widgets/SChaosVDConstraintDataInspector.cpp
+++ b/Engine/Plugins/ChaosVD/Source/ChaosVD/Private/Widgets/SChaosVDConstraintDataInspector.cpp
@@ -5,6 +5,7 @@
 #include "ChaosVDScene.h"
 #include "IStructureDetailsView.h"
 #include "PropertyEditorModule.h"
+#include "SChaosVDMainTab.h"
 #include "Actors/ChaosVDSolverInfoActor.h"
 #include "Widgets/SChaosVDWarningMessageBox.h"
 #include "Modules/ModuleManager.h"
@@ -62,9 +63,10 @@ SChaosVDConstraintDataInspector::~SChaosVDConstraintDataInspector()
 	UnregisterSceneEvents();
 }
 
-void SChaosVDConstraintDataInspector::Construct(const FArguments& InArgs, const TWeakPtr<FChaosVDScene>& InScenePtr)
+void SChaosVDConstraintDataInspector::Construct(const FArguments& InArgs, const TWeakPtr<FChaosVDScene>& InScenePtr, const TSharedRef<SChaosVDMainTab>& InMainTab)
 {
 	SceneWeakPtr = InScenePtr;
+	MainTabWeakPtr = InMainTab;
 
 	RegisterSceneEvents();
 
@@ -336,9 +338,13 @@ const TSharedRef<FChaosVDSolverDataSelectionHandle>& SChaosVDConstraintDataInspe
 	return CurrentDataSelectionHandle;
 }
 
-TSharedPtr<IStructureDetailsView> SChaosVDConstraintDataInspector::CreateDataDetailsView()
+TSharedPtr<IStructureDetailsView> SChaosVDConstraintDataInspector::CreateDataDetailsView() const
 {
-	FPropertyEditorModule& PropertyEditorModule = FModuleManager::GetModuleChecked<FPropertyEditorModule>("PropertyEditor");
+	TSharedPtr<SChaosVDMainTab> MainTabPtr = MainTabWeakPtr.Pin();
+	if (!MainTabPtr)
+	{
+		return nullptr;
+	}
 
 	const FStructureDetailsViewArgs StructDetailsViewArgs;
 	FDetailsViewArgs DetailsViewArgs;
@@ -347,7 +353,7 @@ TSharedPtr<IStructureDetailsView> SChaosVDConstraintDataInspector::CreateDataDet
 	DetailsViewArgs.bAllowSearch = true;
 	DetailsViewArgs.bShowScrollBar = false;
 
-	return PropertyEditorModule.CreateStructureDetailView(DetailsViewArgs, StructDetailsViewArgs, nullptr);
+	return MainTabPtr->CreateStructureDetailsView(DetailsViewArgs, StructDetailsViewArgs, nullptr);
 }
 
 #undef LOCTEXT_NAMESPACE
diff --git a/Engine/Plugins/ChaosVD/Source/ChaosVD/Private/Widgets/SChaosVDConstraintDataInspector.h b/Engine/Plugins/ChaosVD/Source/ChaosVD/Private/Widgets/SChaosVDConstraintDataInspector.h
--- a/Engine/Plugins/ChaosVD/Source/ChaosVD/Private/Widgets/SChaosVDConstraintDataInspector.h
+++ b/Engine/Plugins/ChaosVD/Source/ChaosVD/Private/Widgets/SChaosVDConstraintDataInspector.h
@@ -5,6 +5,7 @@
 #include "Components/ChaosVDSolverJointConstraintDataComponent.h"
 #include "Widgets/SCompoundWidget.h"
 
+class SChaosVDMainTab;
 struct FChaosVDSolverDataSelectionHandle;
 struct FChaosVDJointConstraintSelectionHandle;
 class FChaosVDScene;
@@ -39,7 +40,7 @@ public:
 	virtual ~SChaosVDConstraintDataInspector() override;
 
 	/** Constructs this widget with InArgs */
-	void Construct(const FArguments& InArgs, const TWeakPtr<FChaosVDScene>& InScenePtr);
+	void Construct(const FArguments& InArgs, const TWeakPtr<FChaosVDScene>& InScenePtr, const TSharedRef<SChaosVDMainTab>& InMainTab);
 
 	/** Sets a new query data to be inspected */
 	virtual void SetConstraintDataToInspect(const TSharedPtr<FChaosVDSolverDataSelectionHandle>& InDataSelectionHandle);
@@ -54,7 +55,7 @@ protected:
 	virtual FText GetParticleName(EChaosVDParticlePairIndex ParticleSlot) const;
 	virtual FText GetParticleName(const EChaosVDParticlePairIndex ParticleSlot, const TSharedPtr<FChaosVDSolverDataSelectionHandle>& InSelectionHandle) const;
 
-	TSharedPtr<IStructureDetailsView> CreateDataDetailsView();
+	TSharedPtr<IStructureDetailsView> CreateDataDetailsView() const;
 
 	TSharedRef<SWidget> GenerateParticleSelectorButtons();
 	
@@ -90,5 +91,7 @@ protected:
 
 	TSharedPtr<FReadOnlyCopyStructOnScope> DataBeingInspectedCopy;
 
-	bool bIsUpToDate = true;;
+	bool bIsUpToDate = true;
+
+	TWeakPtr<SChaosVDMainTab> MainTabWeakPtr;
 };
diff --git a/Engine/Plugins/ChaosVD/Source/ChaosVD/Private/Widgets/SChaosVDDetailsView.cpp b/Engine/Plugins/ChaosVD/Source/ChaosVD/Private/Widgets/SChaosVDDetailsView.cpp
--- a/Engine/Plugins/ChaosVD/Source/ChaosVD/Private/Widgets/SChaosVDDetailsView.cpp
+++ b/Engine/Plugins/ChaosVD/Source/ChaosVD/Private/Widgets/SChaosVDDetailsView.cpp
@@ -4,11 +4,11 @@
 
 #include "ChaosVDParticleActor.h"
 #include "PropertyEditorModule.h"
-#include "SSubobjectEditor.h"
-#include "SSubobjectEditorModule.h"
+#include "SChaosVDMainTab.h"
 
-void SChaosVDDetailsView::Construct(const FArguments& InArgs)
+void SChaosVDDetailsView::Construct(const FArguments& InArgs, const TSharedRef<SChaosVDMainTab>& InMainTab)
 {
+	MainTabWeakPtr = InMainTab;
 	DetailsView = CreateObjectDetailsView();
 	StructDetailsView = CreateStructureDataDetailsView();
 
@@ -49,7 +49,12 @@ void SChaosVDDetailsView::SetSelectedStruct(const TSharedPtr<FStructOnScope>& Ne
 
 TSharedPtr<IDetailsView> SChaosVDDetailsView::CreateObjectDetailsView()
 {
-	FPropertyEditorModule& PropertyEditorModule = FModuleManager::GetModuleChecked<FPropertyEditorModule>("PropertyEditor");
+	TSharedPtr<SChaosVDMainTab> MainTabPtr = MainTabWeakPtr.Pin();
+	if (!MainTabPtr)
+	{
+		return nullptr;
+	}
+
 	FDetailsViewArgs DetailsViewArgs;
 	DetailsViewArgs.bUpdatesFromSelection = false;
 	DetailsViewArgs.bLockable = true;
@@ -57,23 +62,27 @@ TSharedPtr<IDetailsView> SChaosVDDetailsView::CreateObjectDetailsView()
 	DetailsViewArgs.NameAreaSettings = FDetailsViewArgs::HideNameArea;
 	DetailsViewArgs.bCustomFilterAreaLocation = false;
 	DetailsViewArgs.bShowSectionSelector = false;
-	DetailsViewArgs.bShowScrollBar = false;
+	DetailsViewArgs.bShowScrollBar = true;
 
-	return PropertyEditorModule.CreateDetailView(DetailsViewArgs);
+	return MainTabPtr->CreateDetailsView(DetailsViewArgs);
 }
 
-TSharedPtr<IStructureDetailsView> SChaosVDDetailsView::CreateStructureDataDetailsView()
+TSharedPtr<IStructureDetailsView> SChaosVDDetailsView::CreateStructureDataDetailsView() const
 {
-	FPropertyEditorModule& PropertyEditorModule = FModuleManager::GetModuleChecked<FPropertyEditorModule>("PropertyEditor");
+	TSharedPtr<SChaosVDMainTab> MainTabPtr = MainTabWeakPtr.Pin();
+	if (!MainTabPtr)
+	{
+		return nullptr;
+	}
 
 	const FStructureDetailsViewArgs StructDetailsViewArgs;
 	FDetailsViewArgs DetailsViewArgs;
 	DetailsViewArgs.bShowOptions = false;
 	DetailsViewArgs.bAllowFavoriteSystem = false;
 	DetailsViewArgs.bAllowSearch = false;
-	DetailsViewArgs.bShowScrollBar = false;
+	DetailsViewArgs.bShowScrollBar = true;
 
-	return PropertyEditorModule.CreateStructureDetailView(DetailsViewArgs,StructDetailsViewArgs, nullptr);
+	return MainTabPtr->CreateStructureDetailsView(DetailsViewArgs,StructDetailsViewArgs, nullptr);
 }
 
 EVisibility SChaosVDDetailsView::GetStructDetailsVisibility() const
diff --git a/Engine/Plugins/ChaosVD/Source/ChaosVD/Private/Widgets/SChaosVDDetailsView.h b/Engine/Plugins/ChaosVD/Source/ChaosVD/Private/Widgets/SChaosVDDetailsView.h
--- a/Engine/Plugins/ChaosVD/Source/ChaosVD/Private/Widgets/SChaosVDDetailsView.h
+++ b/Engine/Plugins/ChaosVD/Source/ChaosVD/Private/Widgets/SChaosVDDetailsView.h
@@ -8,6 +8,7 @@
 #include "UObject/WeakObjectPtrTemplates.h"
 #include "Widgets/SCompoundWidget.h"
 
+class SChaosVDMainTab;
 class IStructureDetailsView;
 class FSubobjectEditorTreeNode;
 class AActor;
@@ -29,7 +30,7 @@ public:
 	SLATE_END_ARGS()
 
 	/** Constructs this widget with InArgs */
-	void Construct(const FArguments& InArgs);
+	void Construct(const FArguments& InArgs, const TSharedRef<SChaosVDMainTab>& InMainTab);
 
 	/** Updates the current object this view details is viewing */
 	void SetSelectedObject(UObject* NewObject);
@@ -43,7 +44,7 @@ public:
 protected:
 
 	TSharedPtr<IDetailsView> CreateObjectDetailsView();
-	TSharedPtr<IStructureDetailsView> CreateStructureDataDetailsView();
+	TSharedPtr<IStructureDetailsView> CreateStructureDataDetailsView() const;
 
 	EVisibility GetStructDetailsVisibility() const;
 	EVisibility GetObjectDetailsVisibility() const;
@@ -55,6 +56,8 @@ protected:
 	TSharedPtr<IDetailsView> DetailsView;
 
 	TSharedPtr<IStructureDetailsView> StructDetailsView;
+
+	TWeakPtr<SChaosVDMainTab> MainTabWeakPtr;
 };
 
 template <typename TStruct>
diff --git a/Engine/Plugins/ChaosVD/Source/ChaosVD/Private/Widgets/SChaosVDMainTab.cpp b/Engine/Plugins/ChaosVD/Source/ChaosVD/Private/Widgets/SChaosVDMainTab.cpp
--- a/Engine/Plugins/ChaosVD/Source/ChaosVD/Private/Widgets/SChaosVDMainTab.cpp
+++ b/Engine/Plugins/ChaosVD/Source/ChaosVD/Private/Widgets/SChaosVDMainTab.cpp
@@ -13,6 +13,8 @@
 #include "ChaosVDSolversTracksTab.h"
 #include "ChaosVDCollisionDataDetailsTab.h"
 #include "ChaosVDConstraintDataInspectorTab.h"
+#include "ChaosVDParticleActor.h"
+#include "ChaosVDParticleActorCustomization.h"
 #include "ChaosVDSceneQueryDataInspectorTab.h"
 #include "ChaosVDStyle.h"
 #include "ChaosVDTabsIDs.h"
@@ -26,16 +28,24 @@
 #include "Framework/MultiBox/MultiBoxBuilder.h"
 #include "HAL/PlatformApplicationMisc.h"
 #include "IDesktopPlatform.h"
+#include "PropertyEditorModule.h"
 #include "Misc/MessageDialog.h"
 #include "StatusBarSubsystem.h"
 #include "ToolMenu.h"
 #include "ToolMenus.h"
 #include "Components/ChaosVDGenericDebugDrawDataComponent.h"
 #include "Components/ChaosVDGTAccelerationStructuresDataComponent.h"
+#include "Components/ChaosVDInstancedStaticMeshComponent.h"
 #include "Components/ChaosVDParticleDataComponent.h"
 #include "Components/ChaosVDSolverCharacterGroundConstraintDataComponent.h"
 #include "Components/ChaosVDSolverCollisionDataComponent.h"
 #include "Components/ChaosVDSolverJointConstraintDataComponent.h"
+#include "Components/ChaosVDStaticMeshComponent.h"
+#include "DetailsCustomizations/ChaosVDGeometryComponentCustomization.h"
+#include "DetailsCustomizations/ChaosVDParticleDataWrapperCustomization.h"
+#include "DetailsCustomizations/ChaosVDQueryDataWrappersCustomizationDetails.h"
+#include "DetailsCustomizations/ChaosVDSelectionMultipleViewCustomization.h"
+#include "DetailsCustomizations/ChaosVDShapeDataCustomization.h"
 #include "HAL/FileManager.h"
 #include "Settings/ChaosVDMiscSettings.h"
 #include "Styling/StyleColors.h"
@@ -452,6 +462,50 @@ void SChaosVDMainTab::LoadCVDFile(const FString& InFilename)
 	}
 }
 
+TSharedRef<IDetailsView> SChaosVDMainTab::CreateDetailsView(const FDetailsViewArgs& InDetailsViewArgs)
+{
+	FPropertyEditorModule& PropertyEditorModule = FModuleManager::GetModuleChecked<FPropertyEditorModule>("PropertyEditor");
+	TSharedRef<IDetailsView> DetailsView = PropertyEditorModule.CreateDetailView(InDetailsViewArgs);
+	SetCustomPropertyLayouts(&DetailsView.Get());
+
+	return DetailsView;
+}
+
+TSharedRef<IStructureDetailsView> SChaosVDMainTab::CreateStructureDetailsView(const FDetailsViewArgs& InDetailsViewArgs, const FStructureDetailsViewArgs& InStructureDetailsViewArgs, const TSharedPtr<FStructOnScope>& InStructData, const FText& CustomName)
+{
+	FPropertyEditorModule& PropertyEditorModule = FModuleManager::GetModuleChecked<FPropertyEditorModule>("PropertyEditor");
+
+	TSharedRef<IStructureDetailsView> DetailsView = PropertyEditorModule.CreateStructureDetailView(InDetailsViewArgs, InStructureDetailsViewArgs, InStructData, CustomName);
+	SetCustomPropertyLayouts(DetailsView->GetDetailsView());
+
+	return DetailsView;
+}
+
+void SChaosVDMainTab::SetCustomPropertyLayouts(IDetailsView* DetailsView)
+{
+	if (!DetailsView)
+	{
+		return;
+	}
+
+	FPropertyEditorModule& PropertyModule = FModuleManager::GetModuleChecked<FPropertyEditorModule>("PropertyEditor");
+	DetailsView->RegisterInstancedCustomPropertyLayout(AChaosVDParticleActor::StaticClass(), FOnGetDetailCustomizationInstance::CreateStatic(&FChaosVDParticleActorCustomization::MakeInstance, SharedThis(this)));
+	DetailsView->RegisterInstancedCustomPropertyLayout(UChaosVDInstancedStaticMeshComponent::StaticClass(), FOnGetDetailCustomizationInstance::CreateStatic(&FChaosVDGeometryComponentCustomization::MakeInstance));
+	DetailsView->RegisterInstancedCustomPropertyLayout(UChaosVDStaticMeshComponent::StaticClass(), FOnGetDetailCustomizationInstance::CreateStatic(&FChaosVDGeometryComponentCustomization::MakeInstance));
+	DetailsView->RegisterInstancedCustomPropertyLayout(FChaosVDQueryVisitStep::StaticStruct(), FOnGetDetailCustomizationInstance::CreateStatic(&FChaosVDQueryVisitDataCustomization::MakeInstance));
+	DetailsView->RegisterInstancedCustomPropertyLayout(FChaosVDQueryDataWrapper::StaticStruct(), FOnGetDetailCustomizationInstance::CreateStatic(&FChaosVDQueryDataWrapperCustomization::MakeInstance));
+	DetailsView->RegisterInstancedCustomPropertyLayout(FChaosVDSelectionMultipleView::StaticStruct(), FOnGetDetailCustomizationInstance::CreateStatic(&FChaosVDSelectionMultipleViewCustomization::MakeInstance));
+
+	//TODO: Rename FChaosVDParticleDataWrapperCustomization to something generic as currently works with any type that wants to hide properties of type FChaosVDWrapperDataBase with invalid data.
+	// Or another option is create a new custom layout intended to be generic from the get go
+	DetailsView->RegisterInstancedCustomPropertyTypeLayout(TEXT("ChaosVDQueryDataWrapper"), FOnGetPropertyTypeCustomizationInstance::CreateStatic(&FChaosVDParticleDataWrapperCustomization::MakeInstance));
+	DetailsView->RegisterInstancedCustomPropertyTypeLayout(TEXT("ChaosVDQueryVisitStep"), FOnGetPropertyTypeCustomizationInstance::CreateStatic(&FChaosVDParticleDataWrapperCustomization::MakeInstance));
+
+	DetailsView->RegisterInstancedCustomPropertyTypeLayout(TEXT("ChaosVDCollisionResponseParams"), FOnGetPropertyTypeCustomizationInstance::CreateStatic(&FChaosVDCollisionResponseParamsCustomization::MakeInstance, SharedThis(this)));
+	DetailsView->RegisterInstancedCustomPropertyTypeLayout(TEXT("ChaosVDCollisionObjectQueryParams"), FOnGetPropertyTypeCustomizationInstance::CreateStatic(&FChaosVDCollisionObjectParamsCustomization::MakeInstance, SharedThis(this)));
+	DetailsView->RegisterInstancedCustomPropertyTypeLayout(TEXT("ChaosVDShapeCollisionData"), FOnGetPropertyTypeCustomizationInstance::CreateStatic(&FChaosVDShapeDataCustomization::MakeInstance, SharedThis(this)));
+}
+
 void SChaosVDMainTab::BrowseChaosVDRecordingFromFolder(FStringView FolderPath)
 {
 	TArray<FString> OutOpenFilenames;
diff --git a/Engine/Plugins/ChaosVD/Source/ChaosVD/Private/Widgets/SChaosVDMainTab.h b/Engine/Plugins/ChaosVD/Source/ChaosVD/Private/Widgets/SChaosVDMainTab.h
--- a/Engine/Plugins/ChaosVD/Source/ChaosVD/Private/Widgets/SChaosVDMainTab.h
+++ b/Engine/Plugins/ChaosVD/Source/ChaosVD/Private/Widgets/SChaosVDMainTab.h
@@ -10,6 +10,8 @@
 
 #include "SChaosVDMainTab.generated.h"
 
+class IDetailsView;
+class IStructureDetailsView;
 class FComponentVisualizer;
 class FChaosVDEditorModeTools;
 class FChaosVDTabSpawnerBase;
@@ -20,9 +22,13 @@ class FChaosVDOutputLogTab;
 class FChaosVDPlaybackViewportTab;
 class FChaosVDObjectDetailsTab;
 class FChaosVDWorldOutlinerTab;
+class FStructOnScope;
 class SButton;
 class SDockTab;
 
+struct FDetailsViewArgs;
+struct FStructureDetailsViewArgs;
+
 UCLASS()
 class CHAOSVD_API UChaosVDMainToolbarMenuContext : public UObject
 {
@@ -74,8 +80,14 @@ public:
 
 	void LoadCVDFile(const FString& InFilename);
 
+	TSharedRef<IDetailsView> CreateDetailsView(const FDetailsViewArgs& InDetailsViewArgs);
+
+	TSharedRef<IStructureDetailsView> CreateStructureDetailsView(const FDetailsViewArgs& InDetailsViewArgs, const FStructureDetailsViewArgs& InStructureDetailsViewArgs, const TSharedPtr<FStructOnScope>& InStructData = nullptr, const FText& CustomName = FText::GetEmpty());
+
 private:
-	
+
+	void SetCustomPropertyLayouts(IDetailsView* DetailsView);
+
 	void SetUpDisableCPUThrottlingDelegate();
 	void CleanUpDisableCPUThrottlingDelegate() const;
 
diff --git a/Engine/Plugins/ChaosVD/Source/ChaosVD/Private/Widgets/SChaosVDSceneQueryDataInspector.cpp b/Engine/Plugins/ChaosVD/Source/ChaosVD/Private/Widgets/SChaosVDSceneQueryDataInspector.cpp
--- a/Engine/Plugins/ChaosVD/Source/ChaosVD/Private/Widgets/SChaosVDSceneQueryDataInspector.cpp
+++ b/Engine/Plugins/ChaosVD/Source/ChaosVD/Private/Widgets/SChaosVDSceneQueryDataInspector.cpp
@@ -9,6 +9,7 @@
 #include "EditorViewportClient.h"
 #include "IStructureDetailsView.h"
 #include "PropertyEditorModule.h"
+#include "SChaosVDMainTab.h"
 #include "Actors/ChaosVDSolverInfoActor.h"
 #include "Components/ChaosVDSceneQueryDataComponent.h"
 #include "Modules/ModuleManager.h"
@@ -61,10 +62,11 @@ SChaosVDSceneQueryDataInspector::~SChaosVDSceneQueryDataInspector()
 	UnregisterSceneEvents();
 }
 
-void SChaosVDSceneQueryDataInspector::Construct(const FArguments& InArgs, const TWeakPtr<FChaosVDScene>& InScenePtr, const TWeakPtr<FEditorModeTools>& InEditorModeTools)
+void SChaosVDSceneQueryDataInspector::Construct(const FArguments& InArgs, const TWeakPtr<FChaosVDScene>& InScenePtr, const TSharedRef<SChaosVDMainTab>& InMainTab)
 {
 	SceneWeakPtr = InScenePtr;
-	EditorModeToolsWeakPtr = InEditorModeTools;
+	EditorModeToolsWeakPtr = InMainTab->GetEditorModeManager().AsWeak();
+	MainTabWeakPtr = InMainTab;
 
 	RegisterSceneEvents();
 
@@ -511,9 +513,13 @@ FReply SChaosVDSceneQueryDataInspector::SelectParentQuery()
 	return FReply::Handled();
 }
 
-TSharedPtr<IStructureDetailsView> SChaosVDSceneQueryDataInspector::CreateDataDetailsView()
+TSharedPtr<IStructureDetailsView> SChaosVDSceneQueryDataInspector::CreateDataDetailsView() const
 {
-	FPropertyEditorModule& PropertyEditorModule = FModuleManager::GetModuleChecked<FPropertyEditorModule>("PropertyEditor");
+	TSharedPtr<SChaosVDMainTab> MainTabPtr = MainTabWeakPtr.Pin();
+	if (!MainTabPtr)
+	{
+		return nullptr;
+	}
 
 	const FStructureDetailsViewArgs StructDetailsViewArgs;
 	FDetailsViewArgs DetailsViewArgs;
@@ -522,7 +528,7 @@ TSharedPtr<IStructureDetailsView> SChaosVDSceneQueryDataInspector::CreateDataDet
 	DetailsViewArgs.bAllowSearch = false;
 	DetailsViewArgs.bShowScrollBar = false;
 
-	return PropertyEditorModule.CreateStructureDetailView(DetailsViewArgs,StructDetailsViewArgs, nullptr);
+	return MainTabPtr->CreateStructureDetailsView(DetailsViewArgs,StructDetailsViewArgs, nullptr);
 }
 
 void SChaosVDSceneQueryDataInspector::HandleSceneUpdated()
diff --git a/Engine/Plugins/ChaosVD/Source/ChaosVD/Private/Widgets/SChaosVDSceneQueryDataInspector.h b/Engine/Plugins/ChaosVD/Source/ChaosVD/Private/Widgets/SChaosVDSceneQueryDataInspector.h
--- a/Engine/Plugins/ChaosVD/Source/ChaosVD/Private/Widgets/SChaosVDSceneQueryDataInspector.h
+++ b/Engine/Plugins/ChaosVD/Source/ChaosVD/Private/Widgets/SChaosVDSceneQueryDataInspector.h
@@ -5,6 +5,7 @@
 #include "Components/ChaosVDSceneQueryDataComponent.h"
 #include "Widgets/SCompoundWidget.h"
 
+class SChaosVDMainTab;
 enum class EChaosVDPlaybackButtonsID : uint8;
 class SChaosVDTimelineWidget;
 struct FChaosVDSceneQuerySelectionHandle;
@@ -40,7 +41,7 @@ public:
 	virtual ~SChaosVDSceneQueryDataInspector() override;
 
 	/** Constructs this widget with InArgs */
-	void Construct(const FArguments& InArgs, const TWeakPtr<FChaosVDScene>& InScenePtr, const TWeakPtr<FEditorModeTools>& InEditorModeTools);
+	void Construct(const FArguments& InArgs, const TWeakPtr<FChaosVDScene>& InScenePtr, const TSharedRef<SChaosVDMainTab>& InMainTab);
 
 	/** Sets a new query data to be inspected */
 	void SetQueryDataToInspect(const TSharedPtr<FChaosVDSolverDataSelectionHandle>& InDataSelectionHandle);
@@ -62,7 +63,7 @@ protected:
 	FReply SelectQueryToInspectByID(int32 QueryID, int32 SolverID);
 	FReply SelectParentQuery();
 
-	static TSharedPtr<IStructureDetailsView> CreateDataDetailsView();
+	TSharedPtr<IStructureDetailsView> CreateDataDetailsView() const;
 	
 	void HandleSceneUpdated();
 	void HandleSubQueryNameSelected(TSharedPtr<FName> Name);
@@ -102,6 +103,8 @@ protected:
 
 	bool bListenToSelectionEvents = true;
 
+	TWeakPtr<SChaosVDMainTab> MainTabWeakPtr;
+
 	int32 GetCurrentMinSQVisitIndex() const;
 	int32 GetCurrentMaxSQVisitIndex() const;
 	int32 GetCurrentSQVisitIndex() const;
diff --git a/Engine/Plugins/ChaosVD/Source/ChaosVD/Public/ChaosVDModule.h b/Engine/Plugins/ChaosVD/Source/ChaosVD/Public/ChaosVDModule.h
--- a/Engine/Plugins/ChaosVD/Source/ChaosVD/Public/ChaosVDModule.h
+++ b/Engine/Plugins/ChaosVD/Source/ChaosVD/Public/ChaosVDModule.h
@@ -6,8 +6,10 @@
 #include "Modules/ModuleManager.h"
 #include "Misc/Guid.h"
 
+class SWidget;
 class FChaosVDTraceManager;
 class FChaosVDEngine;
+class SChaosVDMainTab;
 class SDockTab;
 class FSpawnTabArgs;
 struct FGuid;
@@ -19,6 +21,8 @@ class FChaosVDModule : public IModuleInterface
 {
 public:
 
+	using ToolkitHostType = SChaosVDMainTab;
+
 	static FChaosVDModule& Get();
 
 	/** IModuleInterface implementation */
@@ -35,8 +39,6 @@ public:
 
 private:
 
-	void RegisterClassesCustomDetails() const;
-
 	TSharedRef<SDockTab> SpawnMainTab(const FSpawnTabArgs& Args);
 
 	void HandleTabClosed(TSharedRef<SDockTab> ClosedTab, FGuid InstanceGUID);
@@ -60,5 +62,4 @@ private:
 	static FString ChaosVisualDebuggerProgramName;
 
 	bool bIsShuttingDown = false;
-
 };
diff --git a/Engine/Source/Runtime/Engine/Private/Collision/ChaosVDSQTraceHelper.cpp b/Engine/Source/Runtime/Engine/Private/Collision/ChaosVDSQTraceHelper.cpp
--- a/Engine/Source/Runtime/Engine/Private/Collision/ChaosVDSQTraceHelper.cpp
+++ b/Engine/Source/Runtime/Engine/Private/Collision/ChaosVDSQTraceHelper.cpp
@@ -9,6 +9,7 @@
 
 
 #include "ChaosVisualDebugger/ChaosVDContextProvider.h"
+#include "ChaosVisualDebugger/ChaosVDDataWrapperUtils.h"
 #include "ChaosVisualDebugger/ChaosVisualDebuggerTrace.h"
 #include "DataWrappers/ChaosVDQueryDataWrappers.h"
 #include "PBDRigidsSolver.h"
diff --git a/Engine/Source/Runtime/Engine/Public/Physics/PhysicsFiltering.h b/Engine/Source/Runtime/Engine/Public/Physics/PhysicsFiltering.h
--- a/Engine/Source/Runtime/Engine/Public/Physics/PhysicsFiltering.h
+++ b/Engine/Source/Runtime/Engine/Public/Physics/PhysicsFiltering.h
@@ -142,3 +142,45 @@ inline void UpdateMaskFilter(uint32& Word3, FMaskFilter NewMaskFilter)
 	Word3 &= (0xFFFFFFFFu >> NumExtraFilterBits);	//we drop the top NumExtraFilterBits bits because that's where the new mask filter is going
 	Word3 |= uint32(NewMaskFilter) << (32 - NumExtraFilterBits);
 }
+
+inline FCollisionResponseContainer ExtractSimCollisionResponseContainer(const FCollisionFilterData& InSimFilterData)
+{
+	FCollisionResponseContainer CollisionResponseContainer;
+	
+	for (int32 ChannelIndex = 0; ChannelIndex < UE_ARRAY_COUNT(CollisionResponseContainer.EnumArray); ++ChannelIndex)
+	{
+		if ((InSimFilterData.Word1) & (1 << ChannelIndex))
+		{
+			CollisionResponseContainer.EnumArray[ChannelIndex] = ECR_Block;
+		}
+		else
+		{
+			CollisionResponseContainer.EnumArray[ChannelIndex] = ECR_Ignore;
+		}
+	}
+
+	return CollisionResponseContainer;
+}
+
+inline FCollisionResponseContainer ExtractQueryCollisionResponseContainer(const FCollisionFilterData& InQueryFilterData)
+{
+	FCollisionResponseContainer CollisionResponseContainer;
+	
+	for (int32 ChannelIndex = 0; ChannelIndex < UE_ARRAY_COUNT(CollisionResponseContainer.EnumArray); ++ChannelIndex)
+	{
+		if ((InQueryFilterData.Word1) & (1 << ChannelIndex))
+		{
+			CollisionResponseContainer.EnumArray[ChannelIndex] = ECR_Block;
+		}
+		else if ((InQueryFilterData.Word2) & (1 << ChannelIndex))
+		{
+			CollisionResponseContainer.EnumArray[ChannelIndex] = ECR_Overlap;
+		}
+		else
+		{
+			CollisionResponseContainer.EnumArray[ChannelIndex] = ECR_Ignore;
+		}
+	}
+
+	return CollisionResponseContainer;
+}
diff --git a/Engine/Source/Runtime/Experimental/Chaos/Private/ChaosVisualDebugger/ChaosVisualDebuggerTrace.cpp b/Engine/Source/Runtime/Experimental/Chaos/Private/ChaosVisualDebugger/ChaosVisualDebuggerTrace.cpp
--- a/Engine/Source/Runtime/Experimental/Chaos/Private/ChaosVisualDebugger/ChaosVisualDebuggerTrace.cpp
+++ b/Engine/Source/Runtime/Experimental/Chaos/Private/ChaosVisualDebugger/ChaosVisualDebuggerTrace.cpp
@@ -1051,6 +1051,11 @@ void FChaosVisualDebuggerTrace::UnregisterEventHandlers()
 	bIsTracing = false;
 }
 
+TSharedRef<FChaosVDSerializableNameTable>& FChaosVisualDebuggerTrace::GetNameTableInstance()
+{
+	return CVDNameTable;
+}
+
 void FChaosVisualDebuggerTrace::Reset()
 {
 	CVDNameTable->ResetTable();
diff --git a/Engine/Source/Runtime/Experimental/Chaos/Public/Chaos/CollisionFilterData.h b/Engine/Source/Runtime/Experimental/Chaos/Public/Chaos/CollisionFilterData.h
--- a/Engine/Source/Runtime/Experimental/Chaos/Public/Chaos/CollisionFilterData.h
+++ b/Engine/Source/Runtime/Experimental/Chaos/Public/Chaos/CollisionFilterData.h
@@ -20,6 +20,21 @@ namespace Chaos
 		ModifyContacts			= 0b00100000,	// Unused - present for compatibility, whether to allow contact modification, handled in Chaos callbacks now
 		KinematicKinematicPairs	= 0b01000000,	// Unused - present for compatibility, whether to generate KK pairs, Chaos never generates KK pairs
 	};
+
+	inline const TCHAR* LexToString(EFilterFlags FilterFlag)
+	{
+		switch (FilterFlag)
+		{
+			case EFilterFlags::SimpleCollision: return TEXT("SimpleCollision");
+			case EFilterFlags::ComplexCollision: return TEXT("ComplexCollision");
+			case EFilterFlags::CCD: return TEXT("CCD");
+			case EFilterFlags::ContactNotify: return TEXT("ContactNotify");
+			case EFilterFlags::StaticShape: return TEXT("StaticShape");
+			case EFilterFlags::ModifyContacts: return TEXT("ModifyContacts");
+			case EFilterFlags::KinematicKinematicPairs: return TEXT("KinematicKinematicPairs");
+			default: return TEXT("Invalid");
+		}
+	}
 }
 
 struct FCollisionFilterData
diff --git a/Engine/Source/Runtime/Experimental/Chaos/Public/ChaosVisualDebugger/ChaosVDTraceMacros.h b/Engine/Source/Runtime/Experimental/Chaos/Public/ChaosVisualDebugger/ChaosVDTraceMacros.h
--- a/Engine/Source/Runtime/Experimental/Chaos/Public/ChaosVisualDebugger/ChaosVDTraceMacros.h
+++ b/Engine/Source/Runtime/Experimental/Chaos/Public/ChaosVisualDebugger/ChaosVDTraceMacros.h
@@ -90,8 +90,8 @@
 	#endif
 
 	#ifndef CVD_TRACE_BINARY_DATA
-		#define CVD_TRACE_BINARY_DATA(InData, TypeName) \
-		FChaosVisualDebuggerTrace::TraceBinaryData(InData, TypeName);
+		#define CVD_TRACE_BINARY_DATA(InData, TypeName, ...) \
+		FChaosVisualDebuggerTrace::TraceBinaryData(InData, TypeName, ##__VA_ARGS__);
 	#endif
 
 	#ifndef CVD_TRACE_SOLVER_SIMULATION_SPACE
@@ -249,7 +249,7 @@
 	#endif
 
 	#ifndef CVD_TRACE_BINARY_DATA
-		#define CVD_TRACE_BINARY_DATA(InData, TypeName)
+		#define CVD_TRACE_BINARY_DATA(InData, TypeName, ...)
 	#endif
 
 	#ifndef CVD_TRACE_SOLVER_SIMULATION_SPACE
diff --git a/Engine/Source/Runtime/Experimental/Chaos/Public/ChaosVisualDebugger/ChaosVisualDebuggerTrace.h b/Engine/Source/Runtime/Experimental/Chaos/Public/ChaosVisualDebugger/ChaosVisualDebuggerTrace.h
--- a/Engine/Source/Runtime/Experimental/Chaos/Public/ChaosVisualDebugger/ChaosVisualDebuggerTrace.h
+++ b/Engine/Source/Runtime/Experimental/Chaos/Public/ChaosVisualDebugger/ChaosVisualDebuggerTrace.h
@@ -373,7 +373,7 @@ public:
 	/** Unbinds to the static events triggered by the ChaosVD Runtime module */
 	static void UnregisterEventHandlers();
 
-	static TSharedRef<FChaosVDSerializableNameTable>& GetNameTableInstance() { return CVDNameTable; }
+	static CHAOS_API TSharedRef<FChaosVDSerializableNameTable>& GetNameTableInstance();
 
 private:
 
diff --git a/Engine/Source/Runtime/Experimental/ChaosSolverEngine/Private/Chaos/ChaosVDRecordingStateScreenMessageHandler.cpp b/Engine/Source/Runtime/Experimental/ChaosSolverEngine/Private/Chaos/ChaosVDRecordingStateScreenMessageHandler.cpp
--- a/Engine/Source/Runtime/Experimental/ChaosSolverEngine/Private/Chaos/ChaosVDRecordingStateScreenMessageHandler.cpp
+++ b/Engine/Source/Runtime/Experimental/ChaosSolverEngine/Private/Chaos/ChaosVDRecordingStateScreenMessageHandler.cpp
@@ -3,29 +3,34 @@
 #include "ChaosVDRecordingStateScreenMessageHandler.h"
 
 #if WITH_CHAOS_VISUAL_DEBUGGER
+
 #include "ChaosVDRuntimeModule.h"
+#include "ChaosVisualDebugger/ChaosVDTraceMacros.h"
+#include "ChaosVisualDebugger/ChaosVisualDebuggerTrace.h"
+#include "DataWrappers/ChaosVDCollisionDataWrappers.h"
 #include "Engine/Engine.h"
+#include "Serialization/MemoryWriter.h"
 
-void FChaosVDRecordingStateScreenMessageHandler::HandleCVDRecordingStarted()
+void FChaosVDRecordingStateScreenMessageHandler::AddOnScreenRecordingMessage()
 {
 	if (!GEngine)
 	{
 		return;
 	}
-	
+
 	static FText ChaosVDRecordingStartedMessage = NSLOCTEXT("ChaosVisualDebugger", "OnScreenChaosVDRecordingStartedMessage", "Chaos Visual Debugger recording in progress...");
 
 	if (CVDRecordingMessageKey == 0)
 	{
 		CVDRecordingMessageKey = GetTypeHash(ChaosVDRecordingStartedMessage.ToString());
 	}
-
+	
 	// Add a long duration value, we will remove the message manually when the recording stops
 	constexpr float MessageDurationSeconds = 3600.0f;
 	GEngine->AddOnScreenDebugMessage(CVDRecordingMessageKey, MessageDurationSeconds, FColor::Red,ChaosVDRecordingStartedMessage.ToString());
 }
 
-void FChaosVDRecordingStateScreenMessageHandler::HandleCVDRecordingStopped() const
+void FChaosVDRecordingStateScreenMessageHandler::RemoveOnScreenRecordingMessage()
 {
 	if (!GEngine)
 	{
@@ -38,6 +43,18 @@ void FChaosVDRecordingStateScreenMessageHandler::HandleCVDRecordingStopped() con
 	}
 }
 
+void FChaosVDRecordingStateScreenMessageHandler::HandleCVDRecordingStarted()
+{
+	SerializeCollisionChannelsNames();
+
+	AddOnScreenRecordingMessage();
+}
+
+void FChaosVDRecordingStateScreenMessageHandler::HandleCVDRecordingStopped()
+{
+	RemoveOnScreenRecordingMessage();
+}
+
 void FChaosVDRecordingStateScreenMessageHandler::HandleCVDRecordingStartFailed(const FText& InFailureReason) const
 {
 #if !WITH_EDITOR
@@ -58,6 +75,31 @@ void FChaosVDRecordingStateScreenMessageHandler::HandlePIEStarted(UGameInstance*
 	}
 }
 
+void FChaosVDRecordingStateScreenMessageHandler::SerializeCollisionChannelsNames()
+{
+	TArray<uint8> CollisionChannelsDataBuffer;
+	FMemoryWriter MemWriterAr(CollisionChannelsDataBuffer);
+
+	FChaosVDCollisionChannelsInfoContainer CollisionChannelInfoContainer;
+
+	if (UCollisionProfile* CollisionProfileData = UCollisionProfile::Get())
+	{
+		constexpr int32 MaxSupportedChannels = 32;
+		for (int32 ChannelIndex = 0; ChannelIndex < MaxSupportedChannels; ++ChannelIndex)
+		{
+			FChaosVDCollisionChannelInfo Info;
+			Info.DisplayName = CollisionProfileData->ReturnChannelNameFromContainerIndex(ChannelIndex).ToString();
+			Info.CollisionChannel = ChannelIndex;
+			Info.bIsTraceType = CollisionProfileData->ConvertToTraceType(static_cast<ECollisionChannel>(ChannelIndex)) != TraceTypeQuery_MAX;
+			CollisionChannelInfoContainer.CustomChannelsNames[ChannelIndex] = Info;
+		}
+	}
+
+	Chaos::VisualDebugger::WriteDataToBuffer(CollisionChannelsDataBuffer, CollisionChannelInfoContainer);
+
+	CVD_TRACE_BINARY_DATA(CollisionChannelsDataBuffer, FChaosVDCollisionChannelsInfoContainer::WrapperTypeName, EChaosVDTraceBinaryDataOptions::ForceTrace)
+}
+
 FChaosVDRecordingStateScreenMessageHandler& FChaosVDRecordingStateScreenMessageHandler::Get()
 {
 	static FChaosVDRecordingStateScreenMessageHandler MessageHandler;
diff --git a/Engine/Source/Runtime/Experimental/ChaosSolverEngine/Private/Chaos/ChaosVDRecordingStateScreenMessageHandler.h b/Engine/Source/Runtime/Experimental/ChaosSolverEngine/Private/Chaos/ChaosVDRecordingStateScreenMessageHandler.h
--- a/Engine/Source/Runtime/Experimental/ChaosSolverEngine/Private/Chaos/ChaosVDRecordingStateScreenMessageHandler.h
+++ b/Engine/Source/Runtime/Experimental/ChaosSolverEngine/Private/Chaos/ChaosVDRecordingStateScreenMessageHandler.h
@@ -20,11 +20,15 @@ public:
 	void TearDown();
 
 private:
+	void AddOnScreenRecordingMessage();
+	void RemoveOnScreenRecordingMessage();
 	void HandleCVDRecordingStarted();
-	void HandleCVDRecordingStopped() const;
+	void HandleCVDRecordingStopped();
 	void HandleCVDRecordingStartFailed(const FText& InFailureReason) const;
 	void HandlePIEStarted(UGameInstance* GameInstance);
 
+	void SerializeCollisionChannelsNames();
+
 	FDelegateHandle RecordingStartedHandle;
 	FDelegateHandle RecordingStoppedHandle;
 	FDelegateHandle RecordingStartFailedHandle;
diff --git a/Engine/Source/Runtime/Experimental/ChaosVDData/Private/ChaosVDRecording.cpp b/Engine/Source/Runtime/Experimental/ChaosVDData/Private/ChaosVDRecording.cpp
--- a/Engine/Source/Runtime/Experimental/ChaosVDData/Private/ChaosVDRecording.cpp
+++ b/Engine/Source/Runtime/Experimental/ChaosVDData/Private/ChaosVDRecording.cpp
@@ -438,6 +438,11 @@ bool FChaosVDRecording::IsEmpty() const
 	return GetAvailableSolversNumber_AssumesLocked() == 0 && GetAvailableGameFrames_AssumesLocked().Num() == 0 && ImplicitObjects.Num() == 0;
 }
 
+void FChaosVDRecording::SetCollisionChannelsInfoContainer(const TSharedPtr<FChaosVDCollisionChannelsInfoContainer>& InCollisionChannelsInfo)
+{
+	CollisionChannelsInfoContainer = InCollisionChannelsInfo;
+}
+
 void FChaosVDRecording::AddImplicitObject(const uint32 ID, const Chaos::FImplicitObject* InImplicitObject)
 {
 	FWriteScopeLock WriteLock(RecordingDataLock);
diff --git a/Engine/Source/Runtime/Experimental/ChaosVDData/Public/ChaosVDRecording.h b/Engine/Source/Runtime/Experimental/ChaosVDData/Public/ChaosVDRecording.h
--- a/Engine/Source/Runtime/Experimental/ChaosVDData/Public/ChaosVDRecording.h
+++ b/Engine/Source/Runtime/Experimental/ChaosVDData/Public/ChaosVDRecording.h
@@ -388,6 +388,10 @@ struct CHAOSVDDATA_API FChaosVDRecording
 	/** Returns the last Platform Cycle on which this recording was updated (A new frame was added) */
 	uint64 GetLastUpdatedTimeAsCycle() { return LastUpdatedTimeAsCycle; }
 
+	TSharedPtr<FChaosVDCollisionChannelsInfoContainer> GetCollisionChannelsInfoContainer(){ return CollisionChannelsInfoContainer; }
+	
+	void SetCollisionChannelsInfoContainer(const TSharedPtr<FChaosVDCollisionChannelsInfoContainer>& InCollisionChannelsInfo);
+
 protected:
 
 	/** Adds an Implicit Object to the recording and takes ownership of it */
@@ -427,6 +431,8 @@ protected:
 
 	Chaos::VisualDebugger::FChaosVDArchiveHeader HeaderData;
 
+	TSharedPtr<FChaosVDCollisionChannelsInfoContainer> CollisionChannelsInfoContainer;
+
 	friend class FChaosVDTraceProvider;
 	friend class FChaosVDTraceImplicitObjectProcessor;
 };
diff --git a/Engine/Source/Runtime/Experimental/ChaosVisualDebugger/Private/DataWrappers/ChaosVDCollisionDataWrappers.cpp b/Engine/Source/Runtime/Experimental/ChaosVisualDebugger/Private/DataWrappers/ChaosVDCollisionDataWrappers.cpp
--- a/Engine/Source/Runtime/Experimental/ChaosVisualDebugger/Private/DataWrappers/ChaosVDCollisionDataWrappers.cpp
+++ b/Engine/Source/Runtime/Experimental/ChaosVisualDebugger/Private/DataWrappers/ChaosVDCollisionDataWrappers.cpp
@@ -13,7 +13,7 @@ bool FChaosVDContactPoint::Serialize(FArchive& Ar)
 	Ar << FaceIndex;
 	Ar << ContactType;
 
-	return true;
+	return !Ar.IsError();
 }
 
 bool FChaosVDManifoldPoint::Serialize(FArchive& Ar)
@@ -142,7 +142,7 @@ bool FChaosVDConstraint::Serialize(FArchive& Ar)
 	Ar << Particle0Index;
 	Ar << Particle1Index;
 
-	return true;
+	return !Ar.IsError();
 }
 
 bool FChaosVDParticlePairMidPhase::Serialize(FArchive& Ar)
@@ -183,7 +183,7 @@ bool FChaosVDParticlePairMidPhase::Serialize(FArchive& Ar)
 
 	Ar << Constraints;
 
-	return true;
+	return !Ar.IsError();
 }
 
 bool FChaosVDCollisionFilterData::Serialize(FArchive& Ar)
@@ -228,7 +228,7 @@ bool FChaosVDShapeCollisionData::Serialize(FArchive& Ar)
 		Ar << QueryData;
 	}
 
-	return true;
+	return !Ar.IsError();
 }
 
 bool FChaosVDShapeCollisionData::operator==(const FChaosVDShapeCollisionData& Other) const
@@ -236,5 +236,25 @@ bool FChaosVDShapeCollisionData::operator==(const FChaosVDShapeCollisionData& Ot
 	return CollisionTraceType == Other.CollisionTraceType
 			&& bSimCollision == Other.bSimCollision
 			&& bQueryCollision == Other.bQueryCollision
-			&& bIsProbe == Other.bIsProbe;
+			&& bIsProbe == Other.bIsProbe
+			&& QueryData == Other.QueryData
+			&& SimData == Other.SimData;
+}
+
+bool FChaosVDCollisionChannelInfo::Serialize(FArchive& Ar)
+{
+	Ar << DisplayName;
+
+	Ar << CollisionChannel;
+
+	Ar << bIsTraceType;
+
+	return !Ar.IsError();
+}
+
+bool FChaosVDCollisionChannelsInfoContainer::Serialize(FArchive& Ar)
+{
+	CVD_SERIALIZE_STATIC_ARRAY(Ar, CustomChannelsNames);
+
+	return !Ar.IsError();
 }
diff --git a/Engine/Source/Runtime/Experimental/ChaosVisualDebugger/Public/DataWrappers/ChaosVDCollisionDataWrappers.h b/Engine/Source/Runtime/Experimental/ChaosVisualDebugger/Public/DataWrappers/ChaosVDCollisionDataWrappers.h
--- a/Engine/Source/Runtime/Experimental/ChaosVisualDebugger/Public/DataWrappers/ChaosVDCollisionDataWrappers.h
+++ b/Engine/Source/Runtime/Experimental/ChaosVisualDebugger/Public/DataWrappers/ChaosVDCollisionDataWrappers.h
@@ -1,6 +1,7 @@
 ﻿// Copyright Epic Games, Inc. All Rights Reserved.
 
 #pragma once
+#include "ChaosVDDataSerializationMacros.h"
 #include "HAL/Platform.h"
 
 #include "ChaosVDCollisionDataWrappers.generated.h"
@@ -79,20 +80,7 @@ struct CHAOSVDRUNTIME_API FChaosVDContactPoint
 	bool Serialize(FArchive& Ar);
 };
 
-template<>
-struct TStructOpsTypeTraits<FChaosVDContactPoint> : public TStructOpsTypeTraitsBase2<FChaosVDContactPoint>
-{
-	enum
-	{
-		WithSerializer = true,
-	};
-};
-
-inline FArchive& operator<<(FArchive& Ar, FChaosVDContactPoint& Data)
-{
-	Data.Serialize(Ar);
-	return Ar;
-}
+CVD_IMPLEMENT_SERIALIZER(FChaosVDContactPoint)
 
 UENUM()
 enum class EChaosVDManifoldPointFlags : uint8
@@ -150,20 +138,7 @@ struct CHAOSVDRUNTIME_API FChaosVDManifoldPoint
 	bool Serialize(FArchive& Ar);
 };
 
-template<>
-struct TStructOpsTypeTraits<FChaosVDManifoldPoint> : public TStructOpsTypeTraitsBase2<FChaosVDManifoldPoint>
-{
-	enum
-	{
-		WithSerializer = true,
-	};
-};
-
-inline FArchive& operator<<(FArchive& Ar, FChaosVDManifoldPoint& Data)
-{
-	Data.Serialize(Ar);
-	return Ar;
-}
+CVD_IMPLEMENT_SERIALIZER(FChaosVDManifoldPoint)
 
 USTRUCT()
 struct FChaosVDCollisionMaterial
@@ -209,20 +184,7 @@ struct FChaosVDCollisionMaterial
 	bool Serialize(FArchive& Ar);
 };
 
-template<>
-struct TStructOpsTypeTraits<FChaosVDCollisionMaterial> : public TStructOpsTypeTraitsBase2<FChaosVDCollisionMaterial>
-{
-	enum
-	{
-		WithSerializer = true,
-	};
-};
-
-inline FArchive& operator<<(FArchive& Ar, FChaosVDCollisionMaterial& Data)
-{
-	Data.Serialize(Ar);
-	return Ar;
-}
+CVD_IMPLEMENT_SERIALIZER(FChaosVDCollisionMaterial)
 
 UENUM()
 enum class EChaosVDConstraintFlags : uint16
@@ -345,20 +307,7 @@ struct CHAOSVDRUNTIME_API FChaosVDConstraint
 	bool Serialize(FArchive& Ar);
 };
 
-template<>
-struct TStructOpsTypeTraits<FChaosVDConstraint> : public TStructOpsTypeTraitsBase2<FChaosVDConstraint>
-{
-	enum
-	{
-		WithSerializer = true,
-	};
-};
-
-inline FArchive& operator<<(FArchive& Ar, FChaosVDConstraint& Data)
-{
-	Data.Serialize(Ar);
-	return Ar;
-}
+CVD_IMPLEMENT_SERIALIZER(FChaosVDConstraint)
 
 UENUM()
 enum class EChaosVDMidPhaseFlags : uint8
@@ -427,21 +376,7 @@ struct CHAOSVDRUNTIME_API FChaosVDParticlePairMidPhase
 	bool Serialize(FArchive& Ar);
 };
 
-template<>
-struct TStructOpsTypeTraits<FChaosVDParticlePairMidPhase> : public TStructOpsTypeTraitsBase2<FChaosVDParticlePairMidPhase>
-{
-	enum
-	{
-		WithSerializer = true,
-	};
-};
-
-
-inline FArchive& operator<<(FArchive& Ar, FChaosVDParticlePairMidPhase& Data)
-{
-	Data.Serialize(Ar);
-	return Ar;
-}
+CVD_IMPLEMENT_SERIALIZER(FChaosVDParticlePairMidPhase)
 
 UENUM()
 enum class EChaosVDCollisionTraceFlag
@@ -473,6 +408,8 @@ struct CHAOSVDRUNTIME_API FChaosVDCollisionFilterData
 	uint32 Word3 = 0;
 
 	bool Serialize(FArchive& Ar);
+
+	bool operator==(const FChaosVDCollisionFilterData& Other) const = default;
 };
 
 template <>
@@ -481,20 +418,8 @@ struct TTypeTraits<FChaosVDCollisionFilterData> : public TTypeTraitsBase <FChaos
 	enum { IsBytewiseComparable = true };
 };
 
-template<>
-struct TStructOpsTypeTraits<FChaosVDCollisionFilterData> : public TStructOpsTypeTraitsBase2<FChaosVDCollisionFilterData>
-{
-	enum
-	{
-		WithSerializer = true,
-	};
-};
+CVD_IMPLEMENT_SERIALIZER(FChaosVDCollisionFilterData)
 
-inline FArchive& operator<<(FArchive& Ar, FChaosVDCollisionFilterData& Data)
-{
-	Data.Serialize(Ar);
-	return Ar;
-}
 UENUM()
 enum class EChaosVDCollisionShapeDataFlags : uint8
 {
@@ -512,11 +437,11 @@ struct CHAOSVDRUNTIME_API FChaosVDShapeCollisionData
 	UPROPERTY(VisibleAnywhere, Category=CollisionData)
 	EChaosVDCollisionTraceFlag CollisionTraceType = EChaosVDCollisionTraceFlag::UseDefault;
 
-	UPROPERTY(VisibleAnywhere, Category=CollisionData)
+	UPROPERTY()
 	uint8 bSimCollision : 1 = false;
-	UPROPERTY(VisibleAnywhere, Category=CollisionData)
+	UPROPERTY()
 	uint8 bQueryCollision : 1 = false;
-	UPROPERTY(VisibleAnywhere, Category=CollisionData)
+	UPROPERTY()
 	uint8 bIsProbe : 1 = false;
 
 	UPROPERTY(VisibleAnywhere, Category=FilterData)
@@ -542,17 +467,42 @@ struct TTypeTraits<FChaosVDShapeCollisionData> : public TTypeTraitsBase <FChaosV
 	enum { IsBytewiseComparable = true };
 };
 
-template<>
-struct TStructOpsTypeTraits<FChaosVDShapeCollisionData> : public TStructOpsTypeTraitsBase2<FChaosVDShapeCollisionData>
+CVD_IMPLEMENT_SERIALIZER(FChaosVDShapeCollisionData)
+
+/** Minimum amount of data needed to reconstruct Collision names in CVD
+ * based on already serialized flags
+ */
+USTRUCT()
+struct CHAOSVDRUNTIME_API FChaosVDCollisionChannelInfo
 {
-	enum
-	{
-		WithSerializer = true,
-	};
+	GENERATED_BODY()
+
+	UPROPERTY()
+	FString DisplayName;
+
+	UPROPERTY()
+	int32 CollisionChannel = INDEX_NONE;
+
+	UPROPERTY()
+	bool bIsTraceType = false;
+	
+	bool Serialize(FArchive& Ar);
 };
 
-inline FArchive& operator<<(FArchive& Ar, FChaosVDShapeCollisionData& Data)
+CVD_IMPLEMENT_SERIALIZER(FChaosVDCollisionChannelInfo)
+
+/** Container for recorded custom collision profile data */
+USTRUCT()
+struct CHAOSVDRUNTIME_API FChaosVDCollisionChannelsInfoContainer
 {
-	Data.Serialize(Ar);
-	return Ar;
-}
+	GENERATED_BODY()
+
+	inline static FStringView WrapperTypeName = TEXT("FChaosVDCollisionChannelsInfoContainer");
+
+	UPROPERTY(VisibleAnywhere, Category=CollisionData)
+	FChaosVDCollisionChannelInfo CustomChannelsNames[32] = {};
+
+	bool Serialize(FArchive& Ar);
+};
+
+CVD_IMPLEMENT_SERIALIZER(FChaosVDCollisionChannelsInfoContainer)
diff --git a/Engine/Source/Runtime/Experimental/ChaosVisualDebugger/Public/DataWrappers/ChaosVDQueryDataWrappers.h b/Engine/Source/Runtime/Experimental/ChaosVisualDebugger/Public/DataWrappers/ChaosVDQueryDataWrappers.h
--- a/Engine/Source/Runtime/Experimental/ChaosVisualDebugger/Public/DataWrappers/ChaosVDQueryDataWrappers.h
+++ b/Engine/Source/Runtime/Experimental/ChaosVisualDebugger/Public/DataWrappers/ChaosVDQueryDataWrappers.h
@@ -18,10 +18,10 @@ enum class EChaosVDSceneQueryType
 UENUM()
 enum class EChaosVDSceneQueryMode
 {
-	Invalid,
 	Single,
 	Multi,
-	Test
+	Test,
+	Invalid
 };
 
 USTRUCT()
@@ -34,12 +34,16 @@ struct FChaosVDCollisionResponseParams : public FChaosVDWrapperDataBase
 	template <typename TOther>
 	void CopyFrom(const TOther& Other)
 	{
-		FlagsPerChannel = TArray(Other.CollisionResponse.EnumArray, UE_ARRAY_COUNT(Other.CollisionResponse.EnumArray));
+		const int32 ArrayNum = UE_ARRAY_COUNT(Other.CollisionResponse.EnumArray);
+		for (int32 ChannelIndex = 0; ChannelIndex < ArrayNum; ++ChannelIndex)
+		{
+			FlagsPerChannel.Add(Other.CollisionResponse.EnumArray[ChannelIndex]);
+		}
+
 		bHasValidData = true;
 	}
-	
-	UPROPERTY(VisibleAnywhere, Category=QueryData)
-	TArray<uint8> FlagsPerChannel;
+
+	TArray<uint8, TInlineAllocator<32>> FlagsPerChannel;
 };
 
 inline FArchive& operator<<(FArchive& Ar, FChaosVDCollisionResponseParams& Data)
@@ -323,6 +327,17 @@ enum class EChaosVDCollisionQueryHitType
 	Block = 2
 };
 
+UENUM()
+enum class EChaosVDSQVisitRejectReason
+{
+	None = 0,
+	NoHit = 1,
+	PreFilter = 2,
+	PostFilter = 3,
+	ColocatedHitHasWorseNormal = 4,
+	FailedFastBoundTest = 5
+};
+
 UENUM()
 enum class EChaosVDSceneQueryVisitorType
 {
@@ -362,6 +377,9 @@ struct FChaosVDQueryVisitStep : public FChaosVDWrapperDataBase
 
 	UPROPERTY(VisibleAnywhere, Category="SQ Visit Data")
 	FChaosVDQueryHitData HitData;
+
+	UPROPERTY()
+	EChaosVDSQVisitRejectReason RejectReason = EChaosVDSQVisitRejectReason::None;
 };
 
 inline FArchive& operator<<(FArchive& Ar, FChaosVDQueryVisitStep& Data)
diff --git a/Engine/Source/Runtime/PhysicsCore/Public/ChaosVDSQVisitorHelpers.h b/Engine/Source/Runtime/PhysicsCore/Public/ChaosVDSQVisitorHelpers.h
--- a/Engine/Source/Runtime/PhysicsCore/Public/ChaosVDSQVisitorHelpers.h
+++ b/Engine/Source/Runtime/PhysicsCore/Public/ChaosVDSQVisitorHelpers.h
@@ -25,6 +25,12 @@
 	Chaos::VisualDebugger::TraceHelpers::FillCVDHitDataHelper(Hit, HitType, CVDSQVisitStepData);
 #endif
 
+/** Sets a reject reason to the current scene query being recorded - Compiled out if we don't have CVD support */
+#ifndef CVD_SET_SQ_SHAPE_REJECT_REASON
+	#define CVD_SET_SQ_SHAPE_REJECT_REASON(Reason) \
+			CVDSQVisitStepData.RejectReason = Reason;
+#endif
+
 namespace Chaos::VisualDebugger::TraceHelpers
 {
 	template <class THitType>
@@ -91,4 +97,7 @@ namespace Chaos::VisualDebugger::TraceHelpers
 	#ifndef CVD_FILL_HIT_DATA_HELPER
 		#define CVD_FILL_HIT_DATA_HELPER(Hit, HitType)
 	#endif
+	#ifndef CVD_SET_SQ_SHAPE_REJECT_REASON
+		#define CVD_SET_SQ_SHAPE_REJECT_REASON(RejectReason)
+	#endif
 #endif
diff --git a/Engine/Source/Runtime/PhysicsCore/Public/SQVisitor.h b/Engine/Source/Runtime/PhysicsCore/Public/SQVisitor.h
--- a/Engine/Source/Runtime/PhysicsCore/Public/SQVisitor.h
+++ b/Engine/Source/Runtime/PhysicsCore/Public/SQVisitor.h
@@ -240,6 +240,7 @@ private:
 					const FVec3 InflatedBoundsTraceStart = SQ == ESQType::Raycast ? StartPoint : QueryGeomWorldBounds.Center();
 					if (!InflatedWorldBounds.RaycastFast(InflatedBoundsTraceStart, CurData->Dir, CurData->InvDir, CurData->bParallel, CurData->CurrentLength, CurData->InvCurrentLength, TmpTime, TmpExitTime))
 					{
+						CVD_SET_SQ_SHAPE_REJECT_REASON(EChaosVDSQVisitRejectReason::FailedFastBoundTest);
 						continue;
 					}
 				}
@@ -248,6 +249,7 @@ private:
 					const FVec3 QueryCenter = QueryGeom ? QueryGeomWorldBounds.Center() : StartTM.GetLocation();
 					if (!InflatedWorldBounds.Contains(QueryCenter))
 					{
+						CVD_SET_SQ_SHAPE_REJECT_REASON(EChaosVDSQVisitRejectReason::FailedFastBoundTest);
 						continue;
 					}
 				}
@@ -340,6 +342,7 @@ private:
 									}
 									else
 									{
+										CVD_SET_SQ_SHAPE_REJECT_REASON(EChaosVDSQVisitRejectReason::ColocatedHitHasWorseNormal);
 										// This hit is co-located but has a worse normal
 										bAcceptHit = false;
 									}
@@ -393,8 +396,20 @@ private:
 								break;
 							}
 						}
+						else
+						{
+							CVD_SET_SQ_SHAPE_REJECT_REASON(EChaosVDSQVisitRejectReason::PostFilter);
+						}
 					}
 				}
+				else
+				{
+					CVD_SET_SQ_SHAPE_REJECT_REASON(EChaosVDSQVisitRejectReason::NoHit);
+				}
+			}
+			else
+			{
+				CVD_SET_SQ_SHAPE_REJECT_REASON(EChaosVDSQVisitRejectReason::PreFilter);
 			}
 		}
 
@@ -582,6 +597,10 @@ private:
 				}
 				break;
 			}
+			else
+			{
+				CVD_SET_SQ_SHAPE_REJECT_REASON(EChaosVDSQVisitRejectReason::PreFilter);
+			}
 		}
 		return bContinue;
 	}
