[ChaosVD] Fix for an issue where debug draw text is shown in the incorrect possition if DPI scaling is greater than 1

Also fixed an issue where auto focusing on Overlap scene queries when you double click on them in the Scene Query Broweser was not working properly

diff --git a/Engine/Plugins/ChaosVD/Source/ChaosVD/Private/Visualizers/ChaosVDDebugDrawUtils.cpp b/Engine/Plugins/ChaosVD/Source/ChaosVD/Private/Visualizers/ChaosVDDebugDrawUtils.cpp
--- a/Engine/Plugins/ChaosVD/Source/ChaosVD/Private/Visualizers/ChaosVDDebugDrawUtils.cpp
+++ b/Engine/Plugins/ChaosVD/Source/ChaosVD/Private/Visualizers/ChaosVDDebugDrawUtils.cpp
@@ -376,13 +376,16 @@ void FChaosVDDebugDrawUtils::DrawCanvas(FViewport& InViewport, FSceneView& View,
 			switch (TextToDraw.LocationMode)
 			{
 				case EChaosVDDebugDrawTextLocationMode::World:
+				{
 					bHasValidLocation = View.ViewFrustum.IntersectPoint(TextToDraw.Location) && View.WorldToPixel(TextToDraw.Location, LocationToDraw);
+					LocationToDraw /= View.Family->DebugDPIScale;
 					break;
+				}
 				case EChaosVDDebugDrawTextLocationMode::Screen:
-					{
-						LocationToDraw = FVector2D(TextToDraw.Location.X,TextToDraw.Location.Y);
-						bHasValidLocation = true;
-					}
+				{
+					LocationToDraw = FVector2D(TextToDraw.Location.X, TextToDraw.Location.Y);
+					bHasValidLocation = true;
+				}
 				default:
 					break;
 			}
@@ -436,10 +439,13 @@ FBox Chaos::VisualDebugger::Utils::CalculateSceneQueryShapeBounds(const TSharedR
 	if (InputShapePtr && InputShapePtr->HasBoundingBox())
 	{
 		FAABB3 StartBounds = InputShapePtr->CalculateTransformedBounds(FRigidTransform3(InSceneQueryData->StartLocation, InSceneQueryData->GeometryOrientation));
-		FAABB3 EndBounds = InputShapePtr->CalculateTransformedBounds(FRigidTransform3(InSceneQueryData->EndLocation, InSceneQueryData->GeometryOrientation));
-
 		BoundsBuilder += FBox(StartBounds.Min(), StartBounds.Max());
-		BoundsBuilder += FBox(EndBounds.Min(), EndBounds.Max());
+
+		if (InSceneQueryData->Type != EChaosVDSceneQueryType::Overlap)
+		{
+			FAABB3 EndBounds = InputShapePtr->CalculateTransformedBounds(FRigidTransform3(InSceneQueryData->EndLocation, InSceneQueryData->GeometryOrientation));
+			BoundsBuilder += FBox(EndBounds.Min(), EndBounds.Max());
+		}
 	}
 	else
 	{
