[ChaosVD] Fix issue where the Main tab and therefore all the widgets are leaked when CVD is closed

This is fallout from a change made last week. It happens because the CVD main tab passes itself to the custom detail instances. Although the instances itself keep a weak ptr to the main tab, the Delegate instance used to create the customization object instances, keeps a shared reference to the main tab.

The problem aside form the leak itself is that all CVD playback logic is destroyed at that point, which means any late delegate triggering ends up in a editor crash.

The fix is to pass the toolkit host (CVD Main Tab) as a weak ptr

#tests I tested that after the change, when CVD is closed the the destructors of the widgets for each tab are called as expected, and the editor crash does not happen anymore

diff --git a/Engine/Plugins/ChaosVD/Source/ChaosVD/Private/ChaosVDParticleActorCustomization.cpp b/Engine/Plugins/ChaosVD/Source/ChaosVD/Private/ChaosVDParticleActorCustomization.cpp
--- a/Engine/Plugins/ChaosVD/Source/ChaosVD/Private/ChaosVDParticleActorCustomization.cpp
+++ b/Engine/Plugins/ChaosVD/Source/ChaosVD/Private/ChaosVDParticleActorCustomization.cpp
@@ -16,7 +16,7 @@
 
 #define LOCTEXT_NAMESPACE "ChaosVisualDebugger"
 
-FChaosVDParticleActorCustomization::FChaosVDParticleActorCustomization(const TSharedPtr<SChaosVDMainTab>& InMainTab)
+FChaosVDParticleActorCustomization::FChaosVDParticleActorCustomization(const TWeakPtr<SChaosVDMainTab>& InMainTab)
 {
 	AllowedCategories.Add(FChaosVDParticleActorCustomization::ParticleDataCategoryName);
 	AllowedCategories.Add(FChaosVDParticleActorCustomization::GeometryCategoryName);
@@ -31,7 +31,7 @@ FChaosVDParticleActorCustomization::~FChaosVDParticleActorCustomization()
 	RegisterCVDScene(nullptr);
 }
 
-TSharedRef<IDetailCustomization> FChaosVDParticleActorCustomization::MakeInstance(TSharedRef<SChaosVDMainTab> InMainTab)
+TSharedRef<IDetailCustomization> FChaosVDParticleActorCustomization::MakeInstance(TWeakPtr<SChaosVDMainTab> InMainTab)
 {
 	return MakeShared<FChaosVDParticleActorCustomization>(InMainTab);
 }
diff --git a/Engine/Plugins/ChaosVD/Source/ChaosVD/Private/ChaosVDParticleActorCustomization.h b/Engine/Plugins/ChaosVD/Source/ChaosVD/Private/ChaosVDParticleActorCustomization.h
--- a/Engine/Plugins/ChaosVD/Source/ChaosVD/Private/ChaosVDParticleActorCustomization.h
+++ b/Engine/Plugins/ChaosVD/Source/ChaosVD/Private/ChaosVDParticleActorCustomization.h
@@ -18,13 +18,13 @@ class AChaosVDParticleActor;
 class FChaosVDParticleActorCustomization : public IDetailCustomization
 {
 public:
-	FChaosVDParticleActorCustomization(const TSharedPtr<SChaosVDMainTab>& InMainTab);
+	FChaosVDParticleActorCustomization(const TWeakPtr<SChaosVDMainTab>& InMainTab);
 	virtual ~FChaosVDParticleActorCustomization() override;
 
 	inline static FName ParticleDataCategoryName = FName("Particle Data");
 	inline static FName GeometryCategoryName = FName("Geometry Shape Data");
 
-	static TSharedRef<IDetailCustomization> MakeInstance(TSharedRef<SChaosVDMainTab> InMainTab);
+	static TSharedRef<IDetailCustomization> MakeInstance(TWeakPtr<SChaosVDMainTab> InMainTab);
 
 	virtual void CustomizeDetails(IDetailLayoutBuilder& DetailBuilder) override;
 
diff --git a/Engine/Plugins/ChaosVD/Source/ChaosVD/Private/DetailsCustomizations/ChaosVDQueryDataWrappersCustomizationDetails.cpp b/Engine/Plugins/ChaosVD/Source/ChaosVD/Private/DetailsCustomizations/ChaosVDQueryDataWrappersCustomizationDetails.cpp
--- a/Engine/Plugins/ChaosVD/Source/ChaosVD/Private/DetailsCustomizations/ChaosVDQueryDataWrappersCustomizationDetails.cpp
+++ b/Engine/Plugins/ChaosVD/Source/ChaosVD/Private/DetailsCustomizations/ChaosVDQueryDataWrappersCustomizationDetails.cpp
@@ -86,7 +86,7 @@ void FChaosVDQueryDataWrapperCustomization::CustomizeDetails(IDetailLayoutBuilde
 	FChaosVDDetailsCustomizationUtils::HideInvalidCVDDataWrapperProperties(PotentialPropertiesToHide, DetailBuilder);
 }
 
-FChaosVDCollisionChannelsCustomizationBase::FChaosVDCollisionChannelsCustomizationBase(const TSharedPtr<SChaosVDMainTab>& InMainTab)
+FChaosVDCollisionChannelsCustomizationBase::FChaosVDCollisionChannelsCustomizationBase(const TWeakPtr<SChaosVDMainTab>& InMainTab)
 {
 	MainTabWeakPtr = InMainTab;
 	// Fill with the Engine defaults. When a CVD file is loaded we will update it with any new data is available
@@ -140,7 +140,7 @@ ECollisionResponse FChaosVDCollisionResponseParamsCustomization::GetCurrentColli
 	return ECollisionResponse::ECR_MAX;
 }
 
-TSharedRef<IPropertyTypeCustomization> FChaosVDCollisionResponseParamsCustomization::MakeInstance(TSharedRef<SChaosVDMainTab> MainTab)
+TSharedRef<IPropertyTypeCustomization> FChaosVDCollisionResponseParamsCustomization::MakeInstance(TWeakPtr<SChaosVDMainTab> MainTab)
 {
 	return MakeShared<FChaosVDCollisionResponseParamsCustomization>(MainTab);
 }
@@ -196,7 +196,7 @@ void FChaosVDCollisionResponseParamsCustomization::CustomizeChildren(TSharedRef<
 																	CollisionGroup);
 }
 
-TSharedRef<IPropertyTypeCustomization> FChaosVDCollisionObjectParamsCustomization::MakeInstance(TSharedRef<SChaosVDMainTab> MainTab)
+TSharedRef<IPropertyTypeCustomization> FChaosVDCollisionObjectParamsCustomization::MakeInstance(TWeakPtr<SChaosVDMainTab> MainTab)
 {
 	return MakeShared<FChaosVDCollisionObjectParamsCustomization>(MainTab);
 }
diff --git a/Engine/Plugins/ChaosVD/Source/ChaosVD/Private/DetailsCustomizations/ChaosVDQueryDataWrappersCustomizationDetails.h b/Engine/Plugins/ChaosVD/Source/ChaosVD/Private/DetailsCustomizations/ChaosVDQueryDataWrappersCustomizationDetails.h
--- a/Engine/Plugins/ChaosVD/Source/ChaosVD/Private/DetailsCustomizations/ChaosVDQueryDataWrappersCustomizationDetails.h
+++ b/Engine/Plugins/ChaosVD/Source/ChaosVD/Private/DetailsCustomizations/ChaosVDQueryDataWrappersCustomizationDetails.h
@@ -49,7 +49,7 @@ class FChaosVDCollisionChannelsCustomizationBase : public IPropertyTypeCustomiza
 {
 public:
 
-	FChaosVDCollisionChannelsCustomizationBase(const TSharedPtr<SChaosVDMainTab>& InMainTab);
+	FChaosVDCollisionChannelsCustomizationBase(const TWeakPtr<SChaosVDMainTab>& InMainTab);
 
 	virtual void CustomizeHeader(TSharedRef<IPropertyHandle> StructPropertyHandle, FDetailWidgetRow& HeaderRow, IPropertyTypeCustomizationUtils& StructCustomizationUtils) override {}
 	virtual void CustomizeChildren(TSharedRef<IPropertyHandle> StructPropertyHandle, IDetailChildrenBuilder& StructBuilder, IPropertyTypeCustomizationUtils& StructCustomizationUtils) override;
@@ -67,11 +67,11 @@ protected:
 class FChaosVDCollisionResponseParamsCustomization : public FChaosVDCollisionChannelsCustomizationBase
 {
 public:
-	FChaosVDCollisionResponseParamsCustomization(const TSharedPtr<SChaosVDMainTab>& InMainTab) : FChaosVDCollisionChannelsCustomizationBase(InMainTab)
+	FChaosVDCollisionResponseParamsCustomization(const TWeakPtr<SChaosVDMainTab>& InMainTab) : FChaosVDCollisionChannelsCustomizationBase(InMainTab)
 	{
 	}
 
-	static TSharedRef<IPropertyTypeCustomization> MakeInstance(TSharedRef<SChaosVDMainTab> MainTab);
+	static TSharedRef<IPropertyTypeCustomization> MakeInstance(TWeakPtr<SChaosVDMainTab> MainTab);
 
 	virtual void CustomizeHeader(TSharedRef<IPropertyHandle> StructPropertyHandle, FDetailWidgetRow& HeaderRow, IPropertyTypeCustomizationUtils& StructCustomizationUtils) override {}
 	virtual void CustomizeChildren(TSharedRef<IPropertyHandle> StructPropertyHandle, IDetailChildrenBuilder& StructBuilder, IPropertyTypeCustomizationUtils& StructCustomizationUtils) override;
@@ -87,11 +87,11 @@ private:
 class FChaosVDCollisionObjectParamsCustomization : public FChaosVDCollisionChannelsCustomizationBase
 {
 public:
-	FChaosVDCollisionObjectParamsCustomization(const TSharedPtr<SChaosVDMainTab>& InMainTab) : FChaosVDCollisionChannelsCustomizationBase(InMainTab)
+	FChaosVDCollisionObjectParamsCustomization(const TWeakPtr<SChaosVDMainTab>& InMainTab) : FChaosVDCollisionChannelsCustomizationBase(InMainTab)
 	{
 	}
 
-	static TSharedRef<IPropertyTypeCustomization> MakeInstance(TSharedRef<SChaosVDMainTab> MainTab);
+	static TSharedRef<IPropertyTypeCustomization> MakeInstance(TWeakPtr<SChaosVDMainTab> MainTab);
 
 	virtual void CustomizeHeader(TSharedRef<IPropertyHandle> StructPropertyHandle, FDetailWidgetRow& HeaderRow, IPropertyTypeCustomizationUtils& StructCustomizationUtils) override {}
 	virtual void CustomizeChildren(TSharedRef<IPropertyHandle> StructPropertyHandle, IDetailChildrenBuilder& StructBuilder, IPropertyTypeCustomizationUtils& StructCustomizationUtils) override;
diff --git a/Engine/Plugins/ChaosVD/Source/ChaosVD/Private/DetailsCustomizations/ChaosVDShapeDataCustomization.cpp b/Engine/Plugins/ChaosVD/Source/ChaosVD/Private/DetailsCustomizations/ChaosVDShapeDataCustomization.cpp
--- a/Engine/Plugins/ChaosVD/Source/ChaosVD/Private/DetailsCustomizations/ChaosVDShapeDataCustomization.cpp
+++ b/Engine/Plugins/ChaosVD/Source/ChaosVD/Private/DetailsCustomizations/ChaosVDShapeDataCustomization.cpp
@@ -24,12 +24,12 @@
 
 #define LOCTEXT_NAMESPACE "ChaosVisualDebugger"
 
-TSharedRef<IPropertyTypeCustomization> FChaosVDShapeDataCustomization::MakeInstance(TSharedRef<SChaosVDMainTab> MainTab)
+TSharedRef<IPropertyTypeCustomization> FChaosVDShapeDataCustomization::MakeInstance(TWeakPtr<SChaosVDMainTab> MainTab)
 {
 	return MakeShared<FChaosVDShapeDataCustomization>(MainTab);
 }
 
-FChaosVDShapeDataCustomization::FChaosVDShapeDataCustomization(const TSharedPtr<SChaosVDMainTab>& InMainTab)
+FChaosVDShapeDataCustomization::FChaosVDShapeDataCustomization(const TWeakPtr<SChaosVDMainTab>& InMainTab)
 {
 	MainTabWeakPtr = InMainTab;
 }
diff --git a/Engine/Plugins/ChaosVD/Source/ChaosVD/Private/DetailsCustomizations/ChaosVDShapeDataCustomization.h b/Engine/Plugins/ChaosVD/Source/ChaosVD/Private/DetailsCustomizations/ChaosVDShapeDataCustomization.h
--- a/Engine/Plugins/ChaosVD/Source/ChaosVD/Private/DetailsCustomizations/ChaosVDShapeDataCustomization.h
+++ b/Engine/Plugins/ChaosVD/Source/ChaosVD/Private/DetailsCustomizations/ChaosVDShapeDataCustomization.h
@@ -18,9 +18,9 @@ struct FChaosVDCollisionChannelsInfoContainer;
 class FChaosVDShapeDataCustomization : public IPropertyTypeCustomization
 {
 public:
-	static TSharedRef<IPropertyTypeCustomization> MakeInstance(TSharedRef<SChaosVDMainTab> MainTab);
+	static TSharedRef<IPropertyTypeCustomization> MakeInstance(TWeakPtr<SChaosVDMainTab> MainTab);
 	
-	FChaosVDShapeDataCustomization(const TSharedPtr<SChaosVDMainTab>& InMainTab);
+	FChaosVDShapeDataCustomization(const TWeakPtr<SChaosVDMainTab>& InMainTab);
 	virtual ~FChaosVDShapeDataCustomization() override;
 
 	virtual void CustomizeHeader(TSharedRef<IPropertyHandle> StructPropertyHandle, FDetailWidgetRow& HeaderRow, IPropertyTypeCustomizationUtils& StructCustomizationUtils) override;
diff --git a/Engine/Plugins/ChaosVD/Source/ChaosVD/Private/Widgets/SChaosVDMainTab.cpp b/Engine/Plugins/ChaosVD/Source/ChaosVD/Private/Widgets/SChaosVDMainTab.cpp
--- a/Engine/Plugins/ChaosVD/Source/ChaosVD/Private/Widgets/SChaosVDMainTab.cpp
+++ b/Engine/Plugins/ChaosVD/Source/ChaosVD/Private/Widgets/SChaosVDMainTab.cpp
@@ -489,7 +489,7 @@ void SChaosVDMainTab::SetCustomPropertyLayouts(IDetailsView* DetailsView)
 	}
 
 	FPropertyEditorModule& PropertyModule = FModuleManager::GetModuleChecked<FPropertyEditorModule>("PropertyEditor");
-	DetailsView->RegisterInstancedCustomPropertyLayout(AChaosVDParticleActor::StaticClass(), FOnGetDetailCustomizationInstance::CreateStatic(&FChaosVDParticleActorCustomization::MakeInstance, SharedThis(this)));
+	DetailsView->RegisterInstancedCustomPropertyLayout(AChaosVDParticleActor::StaticClass(), FOnGetDetailCustomizationInstance::CreateStatic(&FChaosVDParticleActorCustomization::MakeInstance, StaticCastWeakPtr<SChaosVDMainTab>(AsWeak())));
 	DetailsView->RegisterInstancedCustomPropertyLayout(UChaosVDInstancedStaticMeshComponent::StaticClass(), FOnGetDetailCustomizationInstance::CreateStatic(&FChaosVDGeometryComponentCustomization::MakeInstance));
 	DetailsView->RegisterInstancedCustomPropertyLayout(UChaosVDStaticMeshComponent::StaticClass(), FOnGetDetailCustomizationInstance::CreateStatic(&FChaosVDGeometryComponentCustomization::MakeInstance));
 	DetailsView->RegisterInstancedCustomPropertyLayout(FChaosVDQueryVisitStep::StaticStruct(), FOnGetDetailCustomizationInstance::CreateStatic(&FChaosVDQueryVisitDataCustomization::MakeInstance));
@@ -501,9 +501,9 @@ void SChaosVDMainTab::SetCustomPropertyLayouts(IDetailsView* DetailsView)
 	DetailsView->RegisterInstancedCustomPropertyTypeLayout(TEXT("ChaosVDQueryDataWrapper"), FOnGetPropertyTypeCustomizationInstance::CreateStatic(&FChaosVDParticleDataWrapperCustomization::MakeInstance));
 	DetailsView->RegisterInstancedCustomPropertyTypeLayout(TEXT("ChaosVDQueryVisitStep"), FOnGetPropertyTypeCustomizationInstance::CreateStatic(&FChaosVDParticleDataWrapperCustomization::MakeInstance));
 
-	DetailsView->RegisterInstancedCustomPropertyTypeLayout(TEXT("ChaosVDCollisionResponseParams"), FOnGetPropertyTypeCustomizationInstance::CreateStatic(&FChaosVDCollisionResponseParamsCustomization::MakeInstance, SharedThis(this)));
-	DetailsView->RegisterInstancedCustomPropertyTypeLayout(TEXT("ChaosVDCollisionObjectQueryParams"), FOnGetPropertyTypeCustomizationInstance::CreateStatic(&FChaosVDCollisionObjectParamsCustomization::MakeInstance, SharedThis(this)));
-	DetailsView->RegisterInstancedCustomPropertyTypeLayout(TEXT("ChaosVDShapeCollisionData"), FOnGetPropertyTypeCustomizationInstance::CreateStatic(&FChaosVDShapeDataCustomization::MakeInstance, SharedThis(this)));
+	DetailsView->RegisterInstancedCustomPropertyTypeLayout(TEXT("ChaosVDCollisionResponseParams"), FOnGetPropertyTypeCustomizationInstance::CreateStatic(&FChaosVDCollisionResponseParamsCustomization::MakeInstance, StaticCastWeakPtr<SChaosVDMainTab>(AsWeak())));
+	DetailsView->RegisterInstancedCustomPropertyTypeLayout(TEXT("ChaosVDCollisionObjectQueryParams"), FOnGetPropertyTypeCustomizationInstance::CreateStatic(&FChaosVDCollisionObjectParamsCustomization::MakeInstance,StaticCastWeakPtr<SChaosVDMainTab>(AsWeak())));
+	DetailsView->RegisterInstancedCustomPropertyTypeLayout(TEXT("ChaosVDShapeCollisionData"), FOnGetPropertyTypeCustomizationInstance::CreateStatic(&FChaosVDShapeDataCustomization::MakeInstance, StaticCastWeakPtr<SChaosVDMainTab>(AsWeak())));
 }
 
 void SChaosVDMainTab::BrowseChaosVDRecordingFromFolder(FStringView FolderPath)
