[ChaosVD] Speculative fix for a crash that can happen while playing back a recording with particles being individually destroyed before the entire solver is destroyed.

There is a null check missing on the scene outliner code. I already started a thread with the team owning it to see if it is something that is recoverable or should be an assert.

Meanwhile, this is a patch that will remove any invalid outliner items before we broadcast the update. As a workaround for UE-193877, CVD needs to enqueue all ADD/REMOVE events, de-duplicate them and process them later within a determined time budget. What seems to be happening is that we enqueue an add event, but then the underlying actor is destroyed before we process that event. To that to happen the solver that owns the particle needs to be unloaded right after, and the amount of updates to process in that frame has to be large enough to not fit within the 2ms budget, meaning the add event is processed a few frames later.

diff --git a/Engine/Plugins/ChaosVD/Source/ChaosVD/Private/ChaosVDWorldOutlinerMode.cpp b/Engine/Plugins/ChaosVD/Source/ChaosVD/Private/ChaosVDWorldOutlinerMode.cpp
--- a/Engine/Plugins/ChaosVD/Source/ChaosVD/Private/ChaosVDWorldOutlinerMode.cpp
+++ b/Engine/Plugins/ChaosVD/Source/ChaosVD/Private/ChaosVDWorldOutlinerMode.cpp
@@ -9,11 +9,21 @@
 #include "ChaosVDScene.h"
 #include "Elements/Framework/TypedElementSelectionSet.h"
 
-static FAutoConsoleVariable CVarChaosVDQueueAndCombineSceneOutlinerEvents(
+namespace Chaos::VisualDebugger::Cvars
+{
+static bool bQueueAndCombineSceneOutlinerEvent = true;
+static FAutoConsoleVariableRef CVarChaosVDQueueAndCombineSceneOutlinerEvents(
 	TEXT("p.Chaos.VD.Tool.QueueAndCombineSceneOutlinerEvents"),
-	true,
+	bQueueAndCombineSceneOutlinerEvent,
 	TEXT("If set to true, scene outliner events will be queued and sent once per frame. If there was a unprocessed event for an item, the las queued event will replace it"));
 
+static bool bPurgeInvalidOutlinerItemsBeforeBroadcast = true;
+static FAutoConsoleVariableRef CVarChaosVDPurgeInvalidOutlinerItemsBeforeBroadcast(
+	TEXT("p.Chaos.VD.Tool.PurgeInvalidOutlinerItemsBeforeBroadcast"),
+	bPurgeInvalidOutlinerItemsBeforeBroadcast,
+	TEXT("If set to true, scene outliner events will evaluated and any invalid outliner event in them will be removed before broadcasting the hierarchy change."));
+}
+
 const FSceneOutlinerTreeItemType FChaosVDActorTreeItem::Type(&FActorTreeItem::Type);
 
 #define LOCTEXT_NAMESPACE "ChaosVisualDebugger"
@@ -212,7 +222,17 @@ void FChaosVDWorldOutlinerMode::ProcessPendingHierarchyEvents()
 			CurrentTimeSpentSeconds += FPlatformTime::Seconds() - StartTimeSeconds;
 		}
 
-		Hierarchy->OnHierarchyChanged().Broadcast(RemoveIterator.Value());
+		FSceneOutlinerHierarchyChangedData& HierarchyChangedData = RemoveIterator->Value;
+
+		if (Chaos::VisualDebugger::Cvars::bPurgeInvalidOutlinerItemsBeforeBroadcast && HierarchyChangedData.Type == FSceneOutlinerHierarchyChangedData::Added)
+		{ 
+			HierarchyChangedData.Items.RemoveAllSwap([](const FSceneOutlinerTreeItemPtr& SceneOutlinerItemPtr)
+			{
+				return !SceneOutlinerItemPtr || !SceneOutlinerItemPtr->IsValid();
+			});
+		}
+
+		Hierarchy->OnHierarchyChanged().Broadcast(HierarchyChangedData);
 		RemoveIterator.RemoveCurrent();
 		CurrentEventProcessedNum++;
 	}
@@ -310,7 +330,7 @@ void FChaosVDWorldOutlinerMode::HandleActorActiveStateChanged(AChaosVDParticleAc
 		// (UE-193877). As our current use case is fairly simple, as a workaround we can just queue the events and process them once per frame
 		// Only taking into account only the last requested event for each item.
 		// Keeping this behind a cvar enabled by default so when the Scene Outliner bug is fixed, we can test it easily.
-		if (CVarChaosVDQueueAndCombineSceneOutlinerEvents->GetBool())
+		if (Chaos::VisualDebugger::Cvars::bQueueAndCombineSceneOutlinerEvent)
 		{
 			EnqueueAndCombineHierarchyEvent(ChangedActor, EventData);
 		}
