[ChaosVD] Added support to visualize the log output stream

- Implemented widget to render all recorded logs
- Enabled the log trace channel in CVD
- Implemented a search bar and category & verbosity filter
- Search and other filter application is async
- Implemented two-way selection system. If you scrub to any given frame, all the log lines for that game frame or solver frame will be selected and the log window will scroll automatically to the line. If you click on any log line, CVD will jump to the closes recorded game frame for the selected log line.

diff --git a/Engine/Plugins/ChaosVD/Source/ChaosVD/Private/ChaosVDRecordedLogTab.cpp b/Engine/Plugins/ChaosVD/Source/ChaosVD/Private/ChaosVDRecordedLogTab.cpp
new file mode 100644
--- /dev/null
+++ b/Engine/Plugins/ChaosVD/Source/ChaosVD/Private/ChaosVDRecordedLogTab.cpp
@@ -0,0 +1,45 @@
+﻿// Copyright Epic Games, Inc. All Rights Reserved.
+
+#include "ChaosVDRecordedLogTab.h"
+
+#include "ChaosLog.h"
+#include "ChaosVDModule.h"
+#include "ChaosVDStyle.h"
+#include "Framework/Docking/TabManager.h"
+#include "Widgets/Docking/SDockTab.h"
+#include "Widgets/SChaosVDMainTab.h"
+#include "Widgets/SChaosVDRecordedLogBrowser.h"
+
+#define LOCTEXT_NAMESPACE "ChaosVisualDebugger"
+
+TSharedRef<SDockTab> FChaosVDRecordedLogTab::HandleTabSpawnRequest(const FSpawnTabArgs& Args)
+{
+	TSharedRef<SDockTab> RecordedLogTab =
+		SNew(SDockTab)
+		.TabRole(ETabRole::PanelTab)
+		.Label(LOCTEXT("RecordedOutputLogTabLabel", "Recorded Output Log"));
+	
+	if (const TSharedPtr<SChaosVDMainTab> MainTabPtr = OwningTabWidget.Pin())
+	{
+		RecordedLogTab->SetContent
+		(
+			SNew(SChaosVDRecordedLogBrowser, MainTabPtr->GetChaosVDEngineInstance())
+		);
+	}
+	else
+	{
+		RecordedLogTab->SetContent
+		(
+			GenerateErrorWidget()
+		);
+	}
+
+	RecordedLogTab->SetTabIcon(FChaosVDStyle::Get().GetBrush("TabIconOutputLog"));
+
+	HandleTabSpawned(RecordedLogTab);
+
+	return RecordedLogTab;
+}
+
+#undef LOCTEXT_NAMESPACE
+
diff --git a/Engine/Plugins/ChaosVD/Source/ChaosVD/Private/ChaosVDRecordedLogTab.h b/Engine/Plugins/ChaosVD/Source/ChaosVD/Private/ChaosVDRecordedLogTab.h
new file mode 100644
--- /dev/null
+++ b/Engine/Plugins/ChaosVD/Source/ChaosVD/Private/ChaosVDRecordedLogTab.h
@@ -0,0 +1,23 @@
+﻿// Copyright Epic Games, Inc. All Rights Reserved.
+
+#pragma once
+
+#include "Templates/SharedPointer.h"
+#include "ChaosVDTabSpawnerBase.h"
+
+class SOutputLog;
+class FOutputLogHistory;
+class FName;
+
+/** Spawns and handles and instance for the visual debugger Recorded Output Panel */
+class FChaosVDRecordedLogTab : public FChaosVDTabSpawnerBase, public TSharedFromThis<FChaosVDRecordedLogTab>
+{
+public:
+
+	FChaosVDRecordedLogTab(const FName& InTabID, TSharedPtr<FTabManager> InTabManager, TWeakPtr<SChaosVDMainTab> InOwningTabWidget) : FChaosVDTabSpawnerBase(InTabID, InTabManager, InOwningTabWidget)
+	{
+	}
+
+	virtual TSharedRef<SDockTab> HandleTabSpawnRequest(const FSpawnTabArgs& Args) override;
+};
+
diff --git a/Engine/Plugins/ChaosVD/Source/ChaosVD/Private/ChaosVDTabsIDs.cpp b/Engine/Plugins/ChaosVD/Source/ChaosVD/Private/ChaosVDTabsIDs.cpp
--- a/Engine/Plugins/ChaosVD/Source/ChaosVD/Private/ChaosVDTabsIDs.cpp
+++ b/Engine/Plugins/ChaosVD/Source/ChaosVD/Private/ChaosVDTabsIDs.cpp
@@ -13,4 +13,5 @@ const FName FChaosVDTabID::StatusBar = TEXT("StatusBar");
 const FName FChaosVDTabID::CollisionDataDetails = TEXT("CollisionDataDetails");
 const FName FChaosVDTabID::SceneQueryDataDetails = TEXT("SceneQueryDataDetails");
 const FName FChaosVDTabID::ConstraintsInspector = TEXT("ConstraintsDetails");
-const FName FChaosVDTabID::SceneQueryBrowser = TEXT("SceneQueryBrowser");
\ No newline at end of file
+const FName FChaosVDTabID::SceneQueryBrowser = TEXT("SceneQueryBrowser");
+const FName FChaosVDTabID::RecordedOutputLog = TEXT("RecordedOutputLog");
\ No newline at end of file
diff --git a/Engine/Plugins/ChaosVD/Source/ChaosVD/Private/ChaosVDTabsIDs.h b/Engine/Plugins/ChaosVD/Source/ChaosVD/Private/ChaosVDTabsIDs.h
--- a/Engine/Plugins/ChaosVD/Source/ChaosVD/Private/ChaosVDTabsIDs.h
+++ b/Engine/Plugins/ChaosVD/Source/ChaosVD/Private/ChaosVDTabsIDs.h
@@ -19,4 +19,5 @@ public:
 	static const FName SceneQueryDataDetails;
 	static const FName ConstraintsInspector;
 	static const FName SceneQueryBrowser;
+	static const FName RecordedOutputLog;
 };
diff --git a/Engine/Plugins/ChaosVD/Source/ChaosVD/Private/Widgets/SChaosVDLogBrowserToolbar.cpp b/Engine/Plugins/ChaosVD/Source/ChaosVD/Private/Widgets/SChaosVDLogBrowserToolbar.cpp
new file mode 100644
--- /dev/null
+++ b/Engine/Plugins/ChaosVD/Source/ChaosVD/Private/Widgets/SChaosVDLogBrowserToolbar.cpp
@@ -0,0 +1,189 @@
+﻿// Copyright Epic Games, Inc. All Rights Reserved.
+
+#include "SChaosVDLogBrowserToolbar.h"
+
+#include "SChaosVDRecordedLogBrowser.h"
+#include "ToolMenuContext.h"
+#include "ToolMenus.h"
+#include "ToolMenuSection.h"
+#include "Widgets/Input/SComboButton.h"
+#include "Widgets/Input/SSearchBox.h"
+#include "Widgets/SChaosVDEnumFlagsMenu.h"
+
+#define LOCTEXT_NAMESPACE "ChaosVisualDebugger"
+
+void SChaosVDLogBrowserToolbar::Construct(const FArguments& InArgs, const TWeakPtr<SChaosVDRecordedLogBrowser>& InLogBrowserWeakPtr)
+{
+	LogBrowserInstanceWeakPtr = InLogBrowserWeakPtr;
+
+	ChildSlot
+	[
+		GenerateMainToolbarWidget()
+	];
+}
+
+void SChaosVDLogBrowserToolbar::RegisterMainToolbarMenu()
+{
+	const UToolMenus* ToolMenus = UToolMenus::Get();
+	if (ToolMenus->IsMenuRegistered(SChaosVDRecordedLogBrowser::ToolBarName))
+	{
+		return;
+	}
+
+	UToolMenu* ToolBar = UToolMenus::Get()->RegisterMenu(SChaosVDRecordedLogBrowser::ToolBarName, NAME_None, EMultiBoxType::SlimHorizontalToolBar);
+
+	FToolMenuSection& MainSection = ToolBar->AddSection(TEXT("LogBrowser.Toolbar.FiltersSection"));
+
+	MainSection.AddDynamicEntry(TEXT("MainSectionEntry"), FNewToolMenuSectionDelegate::CreateLambda([](FToolMenuSection& MenuSection)
+	{
+		const UChaosVDLogBrowserToolbarMenuContext* Context = MenuSection.FindContext<UChaosVDLogBrowserToolbarMenuContext>();
+		TSharedPtr<SChaosVDLogBrowserToolbar> ToolBarInstance = Context ? Context->ToolbarInstanceWeak.Pin() : nullptr;
+
+		MenuSection.AddEntry(FToolMenuEntry::InitWidget("SearchBar", ToolBarInstance->GenerateSearchBarWidget(), FText::GetEmpty()));
+
+		const FText FiltersMenuLabel = LOCTEXT("LogBrowserFiltersMenuLabel", "Filters");
+		constexpr bool bInOpenSubMenuOnClick = false;
+		MenuSection.AddSubMenu(TEXT("Filters"), FiltersMenuLabel, LOCTEXT("LogBrowserFiltersSubMenuTooltip", "Hide logs based on their category"), FNewToolMenuDelegate::CreateSP(ToolBarInstance.ToSharedRef(), &SChaosVDLogBrowserToolbar::GenerateFiltersSubMenu, FiltersMenuLabel), bInOpenSubMenuOnClick, FSlateIcon(FAppStyle::GetAppStyleSetName(), "Icons.Filter"));
+	}));
+}
+
+TSharedRef<SWidget> SChaosVDLogBrowserToolbar::GenerateMainToolbarWidget()
+{
+	RegisterMainToolbarMenu();
+
+	FToolMenuContext MenuContext;
+
+	UChaosVDLogBrowserToolbarMenuContext* CommonContextObject = NewObject<UChaosVDLogBrowserToolbarMenuContext>();
+	CommonContextObject->ToolbarInstanceWeak = SharedThis(this);
+
+	MenuContext.AddObject(CommonContextObject);
+
+	return UToolMenus::Get()->GenerateWidget(SChaosVDRecordedLogBrowser::ToolBarName, MenuContext);
+}
+
+TSharedRef<SWidget> SChaosVDLogBrowserToolbar::GenerateSearchBarWidget()
+{
+	TSharedPtr<SWidget> SearchBox =
+		SNew(SBox)
+		.MinDesiredWidth(400.0f)
+		[
+			SNew(SSearchBox)
+			.HintText(FText::FromString(TEXT("Search by message, category or verbosity...")))
+			.OnTextChanged(this, &SChaosVDLogBrowserToolbar::HandleSearchTextChanged)
+			.DelayChangeNotificationsWhileTyping(true)
+		];
+
+	return SearchBox.ToSharedRef();
+}
+
+void SChaosVDLogBrowserToolbar::GenerateCategoriesSubMenu(UToolMenu* Menu)
+{
+	const UChaosVDLogBrowserToolbarMenuContext* Context = Menu ? Menu->FindContext<UChaosVDLogBrowserToolbarMenuContext>() : nullptr;
+	TSharedPtr<SChaosVDLogBrowserToolbar> ToolBarInstance = Context ? Context->ToolbarInstanceWeak.Pin() : nullptr;
+
+	if (TSharedPtr<SChaosVDRecordedLogBrowser> LogBrowserInstance = ToolBarInstance ? ToolBarInstance->LogBrowserInstanceWeakPtr.Pin() : nullptr)
+	{
+		bool bHasCategories = false;
+		LogBrowserInstance->EnumerateNonEmptyCategories([ToolBarInstance, Menu, &bHasCategories](const SChaosVDRecordedLogBrowser::FCategorizedItemsContainer& InCategoryContainer)
+		{
+			FText CategoryNameAsText = FText::FromName(InCategoryContainer.CategoryName);
+			FToolMenuEntry Entry = FToolMenuEntry::InitMenuEntry(
+			NAME_None,
+			CategoryNameAsText,
+			FText::Format(LOCTEXT("LogCategoryTooltip", "Enable/disable the {0} category"), CategoryNameAsText),
+			FSlateIcon(),
+			FUIAction(FExecuteAction::CreateSP(ToolBarInstance.ToSharedRef(), &SChaosVDLogBrowserToolbar::ToggleCategoryEnabledState, InCategoryContainer.CategoryName),
+			FCanExecuteAction::CreateLambda([](){ return true; }), FIsActionChecked::CreateSP(ToolBarInstance.ToSharedRef(), &SChaosVDLogBrowserToolbar::IsCategoryEnabled, InCategoryContainer.CategoryName)), EUserInterfaceActionType::ToggleButton);
+
+			Menu->AddMenuEntry(NAME_None, Entry);
+
+			bHasCategories = true;
+		});
+
+		if (!bHasCategories)
+		{
+			FText ErrorLabel = LOCTEXT("EmptyCategoryListLabel", "No Categories Available");
+			TSharedRef<SWidget> EmptyCategoriesWidget = SNew(STextBlock)
+														.Margin(FMargin(4.0f,0.0f))
+														.Text(ErrorLabel);
+			Menu->AddMenuEntry(NAME_Name, FToolMenuEntry::InitWidget(NAME_None, EmptyCategoriesWidget, FText::GetEmpty()));
+		}
+	}
+}
+
+void SChaosVDLogBrowserToolbar::GenerateFiltersSubMenu(UToolMenu* Menu, FText FiltersMenuLabel)
+{
+	const UChaosVDLogBrowserToolbarMenuContext* Context = Menu ? Menu->FindContext<UChaosVDLogBrowserToolbarMenuContext>() : nullptr;
+	TSharedPtr<SChaosVDLogBrowserToolbar> ToolBarInstance = Context ? Context->ToolbarInstanceWeak.Pin() : nullptr;
+
+	if (TSharedPtr<SChaosVDRecordedLogBrowser> LogBrowserInstance = ToolBarInstance ? ToolBarInstance->LogBrowserInstanceWeakPtr.Pin() : nullptr)
+	{
+		FToolMenuSection& Section = Menu->AddSection(TEXT("LogBrowser.Toolbar.Filters"), FiltersMenuLabel);
+
+		Section.AddMenuEntry(
+			NAME_None,
+			LOCTEXT("LogBrowserShowAllCategories", "Show All"),
+			LOCTEXT("LogBrowserShowAllCategories_Tooltip", "Filter the Recorded Output Log to show all categories"),
+			FSlateIcon(),
+			FUIAction(FExecuteAction::CreateSP(LogBrowserInstance.ToSharedRef(), &SChaosVDRecordedLogBrowser::ToggleShowAllCategories),
+				FCanExecuteAction::CreateLambda([] { return true; }),
+				FIsActionChecked::CreateSP(LogBrowserInstance.ToSharedRef(), &SChaosVDRecordedLogBrowser::GetShowAllCategories)),
+			EUserInterfaceActionType::ToggleButton
+		);
+
+		Section.AddSubMenu(TEXT("LogBrowser.Toolbar.Filters.Categories"), LOCTEXT("LogBrowserFiltersSubMenuLabel", "Categories"), FText::GetEmpty(), FNewToolMenuDelegate::CreateSP(ToolBarInstance.ToSharedRef(), &SChaosVDLogBrowserToolbar::GenerateCategoriesSubMenu));
+
+		FToolMenuSection& VerbositySection = Menu->AddSection(TEXT("LogBrowser.Toolbar.Verbosity"), LOCTEXT("LogBrowserMenuVerbosityLabel", "Verbosity"));
+		TSharedRef<SWidget> VerbosityFlagsWidget = SNew(SChaosVDEnumFlagsMenu<EChaosVDLogVerbosityFlags>)
+								   .CurrentValue_Raw(ToolBarInstance.Get(), &SChaosVDLogBrowserToolbar::GetVerbosityFlags)
+								   .OnEnumSelectionChanged_Raw(ToolBarInstance.Get(), &SChaosVDLogBrowserToolbar::SetVerbosityFlags);
+
+		VerbositySection.AddEntry(FToolMenuEntry::InitWidget("VerbosityFlags", VerbosityFlagsWidget, FText::GetEmpty()));
+	}
+}
+
+void SChaosVDLogBrowserToolbar::ToggleCategoryEnabledState(FName CategoryName)
+{
+	if (const TSharedPtr<SChaosVDRecordedLogBrowser>& LogBrowserInstance = LogBrowserInstanceWeakPtr.Pin())
+    {
+		LogBrowserInstance->ToggleCategoryEnabled(CategoryName);
+    }
+}
+
+bool SChaosVDLogBrowserToolbar::IsCategoryEnabled(FName CategoryName)
+{
+	if (const TSharedPtr<SChaosVDRecordedLogBrowser>& LogBrowserInstance = LogBrowserInstanceWeakPtr.Pin())
+	{
+		return LogBrowserInstance->IsCategoryEnabled(CategoryName);
+	}
+
+	return false;
+}
+
+void SChaosVDLogBrowserToolbar::SetVerbosityFlags(EChaosVDLogVerbosityFlags NewFlags)
+{
+	if (const TSharedPtr<SChaosVDRecordedLogBrowser>& LogBrowserInstance = LogBrowserInstanceWeakPtr.Pin())
+	{
+		LogBrowserInstance->SetVerbosityFlags(NewFlags);
+	}
+}
+
+EChaosVDLogVerbosityFlags SChaosVDLogBrowserToolbar::GetVerbosityFlags() const
+{
+	if (const TSharedPtr<SChaosVDRecordedLogBrowser>& LogBrowserInstance = LogBrowserInstanceWeakPtr.Pin())
+	{
+		return LogBrowserInstance->GetVerbosityFlags();
+	}
+
+	return EChaosVDLogVerbosityFlags::None;
+}
+
+void SChaosVDLogBrowserToolbar::HandleSearchTextChanged(const FText& Text)
+{
+	if (const TSharedPtr<SChaosVDRecordedLogBrowser>& LogBrowserInstance = LogBrowserInstanceWeakPtr.Pin())
+	{
+		LogBrowserInstance->HandleSearchTextChanged(Text);
+	}
+}
+
+#undef LOCTEXT_NAMESPACE
diff --git a/Engine/Plugins/ChaosVD/Source/ChaosVD/Private/Widgets/SChaosVDLogBrowserToolbar.h b/Engine/Plugins/ChaosVD/Source/ChaosVD/Private/Widgets/SChaosVDLogBrowserToolbar.h
new file mode 100644
--- /dev/null
+++ b/Engine/Plugins/ChaosVD/Source/ChaosVD/Private/Widgets/SChaosVDLogBrowserToolbar.h
@@ -0,0 +1,53 @@
+﻿// Copyright Epic Games, Inc. All Rights Reserved.
+
+#pragma once
+
+#include "Widgets/SCompoundWidget.h"
+
+#include "SChaosVDLogBrowserToolbar.generated.h"
+
+enum class EChaosVDLogVerbosityFlags : uint8;
+class SChaosVDRecordedLogBrowser;
+class UToolMenu;
+
+UCLASS()
+class UChaosVDLogBrowserToolbarMenuContext : public UObject
+{
+	GENERATED_BODY()
+
+public:
+	TWeakPtr<class SChaosVDLogBrowserToolbar> ToolbarInstanceWeak;
+};
+
+/**
+ * Toolbar Widget for the Recorded Log Browser tab
+ */
+class SChaosVDLogBrowserToolbar : public SCompoundWidget
+{
+public:
+	SLATE_BEGIN_ARGS(SChaosVDLogBrowserToolbar)
+		{
+		}
+	SLATE_END_ARGS()
+
+	void Construct(const FArguments& InArgs, const TWeakPtr<SChaosVDRecordedLogBrowser>& InLogBrowserWeakPtr);
+
+private:
+	void RegisterMainToolbarMenu();
+
+	TSharedRef<SWidget> GenerateMainToolbarWidget();
+	TSharedRef<SWidget> GenerateSearchBarWidget();
+
+	void GenerateCategoriesSubMenu(UToolMenu* Menu);
+	void GenerateFiltersSubMenu(UToolMenu* Menu, FText FiltersMenuLabel);
+	
+	void ToggleCategoryEnabledState(FName CategoryName);
+	bool IsCategoryEnabled(FName CategoryName);
+
+	void SetVerbosityFlags(EChaosVDLogVerbosityFlags NewFlags);
+	EChaosVDLogVerbosityFlags GetVerbosityFlags() const;
+
+	TWeakPtr<SChaosVDRecordedLogBrowser> LogBrowserInstanceWeakPtr;
+
+	void HandleSearchTextChanged(const FText& Text);
+};
diff --git a/Engine/Plugins/ChaosVD/Source/ChaosVD/Private/Widgets/SChaosVDMainTab.cpp b/Engine/Plugins/ChaosVD/Source/ChaosVD/Private/Widgets/SChaosVDMainTab.cpp
--- a/Engine/Plugins/ChaosVD/Source/ChaosVD/Private/Widgets/SChaosVDMainTab.cpp
+++ b/Engine/Plugins/ChaosVD/Source/ChaosVD/Private/Widgets/SChaosVDMainTab.cpp
@@ -15,6 +15,7 @@
 #include "ChaosVDConstraintDataInspectorTab.h"
 #include "ChaosVDParticleActor.h"
 #include "ChaosVDParticleActorCustomization.h"
+#include "ChaosVDRecordedLogTab.h"
 #include "ChaosVDSceneQueryDataInspectorTab.h"
 #include "ChaosVDStyle.h"
 #include "ChaosVDTabsIDs.h"
@@ -96,6 +97,7 @@ void SChaosVDMainTab::Construct(const FArguments& InArgs, TSharedPtr<FChaosVDEng
 	RegisterTabSpawner<FChaosVDSceneQueryDataInspectorTab>(FChaosVDTabID::SceneQueryDataDetails);
 	RegisterTabSpawner<FChaosVDConstraintDataInspectorTab>(FChaosVDTabID::ConstraintsInspector);
 	RegisterTabSpawner<FChaosVDSceneQueryBrowserTab>(FChaosVDTabID::SceneQueryBrowser);
+	RegisterTabSpawner<FChaosVDRecordedLogTab>(FChaosVDTabID::RecordedOutputLog);
 
 	StatusBarID = FName(FChaosVDTabID::StatusBar.ToString() + InChaosVDEngine->GetInstanceGuid().ToString());
 	
@@ -284,6 +286,7 @@ TSharedRef<FTabManager::FLayout> SChaosVDMainTab::GenerateMainLayout()
 					FTabManager::NewStack()
 					->SetSizeCoefficient(0.3f)
 					->AddTab(FChaosVDTabID::SolversTrack, ETabState::OpenedTab)
+					->AddTab(FChaosVDTabID::RecordedOutputLog, ETabState::OpenedTab)
 					->AddTab(FChaosVDTabID::OutputLog, ETabState::OpenedTab)
 				)
 			)
diff --git a/Engine/Plugins/ChaosVD/Source/ChaosVD/Private/Widgets/SChaosVDRecordedLogBrowser.cpp b/Engine/Plugins/ChaosVD/Source/ChaosVD/Private/Widgets/SChaosVDRecordedLogBrowser.cpp
new file mode 100644
--- /dev/null
+++ b/Engine/Plugins/ChaosVD/Source/ChaosVD/Private/Widgets/SChaosVDRecordedLogBrowser.cpp
@@ -0,0 +1,621 @@
+﻿// Copyright Epic Games, Inc. All Rights Reserved.
+
+#include "SChaosVDRecordedLogBrowser.h"
+
+#include "ChaosVDEngine.h"
+#include "ChaosVDPlaybackController.h"
+#include "ChaosVDScene.h"
+#include "SChaosVDLogBrowserToolbar.h"
+#include "SChaosVDRecordedLogView.h"
+#include "SlateOptMacros.h"
+#include "ToolMenus.h"
+#include "Components/VerticalBox.h"
+#include "Misc/TextFilterExpressionEvaluator.h"
+#include "TraceServices/Model/AnalysisSession.h"
+#include "TraceServices/Model/Log.h"
+#include "Widgets/Input/SSearchBox.h"
+
+#define LOCTEXT_NAMESPACE "ChaosVisualDebugger"
+
+const FName SChaosVDRecordedLogBrowser::ToolBarName = FName(TEXT("ChaosVD.RecordedLogBrowser.ToolBar"));
+
+bool FChaosVDBasicLogFilterExpressionContext::TestBasicStringExpression(const FTextFilterString& InValue, const ETextFilterTextComparisonMode InTextComparisonMode) const
+{
+	return TextFilterUtils::TestBasicStringExpression(LogEntry->Message, InValue, InTextComparisonMode) ||
+			TextFilterUtils::TestBasicStringExpression(LogEntry->Category, InValue, InTextComparisonMode) ||
+			TextFilterUtils::TestBasicStringExpression(::ToString(LogEntry->Verbosity), InValue, InTextComparisonMode);
+}
+
+void FChaosVDLogFilteringAsyncTask::DoWork()
+{
+	if (TSharedPtr<SChaosVDRecordedLogBrowser> LogBrowserPtr = LogBrowser.Pin())
+	{
+		FWriteScopeLock WriteScopeLock(LogBrowserPtr->SourceDataCacheLock);
+
+		LogBrowserPtr->FilteredCachedLogItems->Reset();
+
+		for (const TPair<FName, SChaosVDRecordedLogBrowser::FCategorizedItemsContainer>& CategorisedItemsContainer : LogBrowserPtr->CachedLogItemsByCategory)
+		{
+			if (CategorisedItemsContainer.Value.bIsEnabled)
+			{
+				LogBrowserPtr->ApplyFilterToData_AssumesLocked(CategorisedItemsContainer.Value.Items, LogBrowserPtr->FilteredCachedLogItems.ToSharedRef());
+			}
+		}
+		
+		LogBrowserPtr->FilteredCachedLogItems->Sort([](const TSharedPtr<FChaosVDLogViewListItem>& ItemA, const TSharedPtr<FChaosVDLogViewListItem>& ItemB)
+		{
+			return ItemA && ItemB ? ItemA->EntryIndex < ItemB->EntryIndex : false;
+		});
+
+		LogBrowserPtr->SetDirtyFlag(EChaosVDLogBrowserDirtyFlags::Filtering);
+	}
+}
+
+void SChaosVDRecordedLogBrowser::Construct(const FArguments& InArgs, const TSharedRef<FChaosVDEngine>& InEngineInstance)
+{
+	EngineInstanceWeakPtr = InEngineInstance;
+	
+	FilteredCachedLogItems = MakeShared<TArray<TSharedPtr<FChaosVDLogViewListItem>>>();
+	
+	FilterEvaluator = MakeShared<FTextFilterExpressionEvaluator>(ETextFilterExpressionEvaluatorMode::BasicString);
+
+	RegisterNewController(InEngineInstance->GetPlaybackController());
+	
+	constexpr float NoPadding = 0.0f;
+	constexpr float MainContentBoxHorizontalPadding = 2.0f;
+	constexpr float MainContentBoxVerticalPadding = 5.0f;
+	constexpr float StatusBarSlotVerticalPadding = 1.0f;
+	constexpr float StatusBarInnerVerticalPadding = 9.0f;
+	constexpr float StatusBarInnerHorizontalPadding = 14.0f;
+
+	ChildSlot
+	[
+		SNew(SVerticalBox)
+		+SVerticalBox::Slot()
+		.AutoHeight()
+		.Padding(NoPadding)
+		[
+			SAssignNew(ToolbarPtr, SChaosVDLogBrowserToolbar, StaticCastWeakPtr<SChaosVDRecordedLogBrowser>(AsWeak()))
+		]
+		+SVerticalBox::Slot()
+		.Padding(MainContentBoxHorizontalPadding, MainContentBoxVerticalPadding, MainContentBoxHorizontalPadding, NoPadding)
+		.FillHeight(1.0f)
+		[
+			SAssignNew(LogViewWidget, SChaosVDRecordedLogView)
+			.OnItemSelected(this, &SChaosVDRecordedLogBrowser::HandleItemSelected)
+			.OnItemFocused(this, &SChaosVDRecordedLogBrowser::HandleItemFocused)
+		]
+		+SVerticalBox::Slot()
+		.Padding(NoPadding, StatusBarSlotVerticalPadding, NoPadding, StatusBarSlotVerticalPadding)
+		.AutoHeight()
+		[
+			SNew(SBorder)
+			.BorderImage(FAppStyle::Get().GetBrush("Brushes.Header") )
+			.VAlign(VAlign_Center)
+			.HAlign(HAlign_Left)
+			.Padding(FMargin(StatusBarInnerHorizontalPadding, StatusBarInnerVerticalPadding))
+			[
+				SNew(STextBlock)
+				.Text(this, &SChaosVDRecordedLogBrowser::GetFilterStatusText)
+				.ColorAndOpacity(this, &SChaosVDRecordedLogBrowser::GetFilterStatusTextColor)
+			]
+		]
+	];
+}
+
+SChaosVDRecordedLogBrowser::~SChaosVDRecordedLogBrowser()
+{
+	for (const TSharedRef<FChaosVDLogInFlightFilteringTaskWrapper>& FilterTask : FilteringTasksBeingCancelled)
+	{
+		if (!FilterTask->GetAsyncTaskRef().IsIdle())
+		{
+			FilterTask->GetAsyncTaskRef().EnsureCompletion();
+		}
+	}
+
+	if (CurrentFilteringTask)
+	{
+		CurrentFilteringTask->GetAsyncTaskRef().EnsureCompletion();
+	}
+}
+
+void SChaosVDRecordedLogBrowser::HandleSearchTextChanged(const FText& NewText)
+{
+	FilterEvaluator->SetFilterText(NewText);
+	ApplyFiltersAsync();
+}
+
+void SChaosVDRecordedLogBrowser::UpdateLogLineSelectionFromGameTrack()
+{
+	TSharedPtr<FChaosVDPlaybackController> PlaybackControllerPtr = PlaybackController.Pin();
+	if (TSharedPtr<FChaosVDRecording> RecordingData = PlaybackControllerPtr ? PlaybackControllerPtr->GetCurrentRecording().Pin() : nullptr)
+	{
+		if (TSharedPtr<const FChaosVDTrackInfo> GameTrackInfo = PlaybackControllerPtr->GetTrackInfo(EChaosVDTrackType::Game, FChaosVDPlaybackController::GameTrackID))
+		{
+			PendingTrackUpdatesToProcess.Enqueue(GameTrackInfo);
+		}
+	}
+}
+
+void SChaosVDRecordedLogBrowser::UpdateBrowserContents()
+{
+	TSharedPtr<FChaosVDEngine> CVDEngineInstance = EngineInstanceWeakPtr.Pin();
+	if (!CVDEngineInstance)
+	{
+		return;	
+	}
+
+	TSharedPtr<const TraceServices::IAnalysisSession> Session = FChaosVDModule::Get().GetTraceManager()->GetSession(CVDEngineInstance->GetCurrentSessionDescriptor().SessionName);
+	if (!Session)
+	{
+		return;
+	}
+	
+	TSharedPtr<SChaosVDRecordedLogView> LogViewList = LogViewWidget.Pin();
+	if(!LogViewList)
+	{
+		return;
+	}
+
+	TraceServices::FAnalysisSessionReadScope SessionReadScope(*Session.Get());
+	const TraceServices::ILogProvider& LogProvider = TraceServices::ReadLogProvider(*Session.Get());
+
+	uint64 LogProviderNumMessages = LogProvider.GetMessageCount();
+
+	bool bHasNewLogData = false;
+	bool bNeedsReset = false;
+	bool bFirstUpdate = false;
+
+	{
+		FReadScopeLock ReadScopeLock(SourceDataCacheLock);
+		int32 CurrentCachedItemsNum = UnfilteredCachedLogItems.Num();
+		bHasNewLogData =  CurrentCachedItemsNum != LogProviderNumMessages;
+		bNeedsReset = CurrentCachedItemsNum > LogProviderNumMessages;
+		bFirstUpdate = CurrentCachedItemsNum == 0;
+	}
+
+	if (bNeedsReset)
+	{
+		Reset();
+	}
+
+	if (bHasNewLogData)
+	{
+		FWriteScopeLock WriteScopeLock(SourceDataCacheLock);
+
+		for (uint64 Index = UnfilteredCachedLogItems.Num(); Index < LogProviderNumMessages; Index++)
+		{
+			TSharedRef<FChaosVDCachedLogItemEntry> LogEntryData = MakeShared<FChaosVDCachedLogItemEntry>();
+			LogProvider.ReadMessage(Index, [LogEntryData](const TraceServices::FLogMessageInfo& MessageInfo)
+			{
+				LogEntryData->Category = MessageInfo.Category->Name;
+				LogEntryData->Verbosity = MessageInfo.Verbosity;
+				LogEntryData->Message = MessageInfo.Message;
+				LogEntryData->Index = MessageInfo.Index;
+				LogEntryData->Time = MessageInfo.Time;
+			});
+
+			CachedReadLogData.Add(LogEntryData);
+
+			UnfilteredCachedLogItems.Add(MakeLogItem_AssumesLocked(LogEntryData));
+		}
+	}
+
+	if (bFirstUpdate || bNeedsReset)
+	{
+		UpdateLogLineSelectionFromGameTrack();
+	}
+}
+
+FText SChaosVDRecordedLogBrowser::GetFilterStatusText() const
+{
+	FReadScopeLock ReadScopeLock(SourceDataCacheLock);
+
+	const int32 FilteredItemsNum = FilteredCachedLogItems->Num();
+	const int32 UnFilteredItemsNum = UnfilteredCachedLogItems.Num();
+
+	return FText::FormatOrdered(LOCTEXT("LogBrowserFilterStatusMessage","Showing {0} Log Entries | {1} entries are hidden by filter."), FilteredItemsNum, UnFilteredItemsNum - FilteredItemsNum);
+}
+
+FSlateColor SChaosVDRecordedLogBrowser::GetFilterStatusTextColor() const
+{
+	FReadScopeLock ReadScopeLock(SourceDataCacheLock);
+
+	const int32 FilteredItemsNum = FilteredCachedLogItems->Num();
+	const int32 UnFilteredItemsNum = UnfilteredCachedLogItems.Num();
+	
+	if (FilterEvaluator->GetFilterText().IsEmpty())
+	{
+		return FSlateColor::UseForeground();
+	}
+	else if (FilteredItemsNum == 0 && UnFilteredItemsNum > 0)
+	{
+		return FAppStyle::Get().GetSlateColor("Colors.AccentRed");
+	}
+	else
+	{
+		return FAppStyle::Get().GetSlateColor("Colors.AccentGreen");
+	}
+}
+
+void SChaosVDRecordedLogBrowser::CleanPendingCanceledTaskList()
+{
+	bool bAllPendingTasksFullyCanceled = true;
+	for (const TSharedRef<FChaosVDLogInFlightFilteringTaskWrapper>& FilterTask : FilteringTasksBeingCancelled)
+	{
+		bAllPendingTasksFullyCanceled &= FilterTask->GetAsyncTaskRef().IsIdle();
+	}
+
+	// We can only clean the pending tasks list if they are fully cancelled and Idle
+	if (bAllPendingTasksFullyCanceled)
+	{
+		FilteringTasksBeingCancelled.Reset();
+	}
+}
+
+void SChaosVDRecordedLogBrowser::Tick(const FGeometry& AllottedGeometry, const double InCurrentTime, const float InDeltaTime)
+{
+	SCompoundWidget::Tick(AllottedGeometry, InCurrentTime, InDeltaTime);
+
+	if (GetDirtyFlags() != EChaosVDLogBrowserDirtyFlags::None)
+	{
+		UpdateBrowserContents();
+		ApplyFiltersAsync();
+		RefreshLogListView();
+		ProcessPendingTrackUpdates();
+		DirtyFlags = EChaosVDLogBrowserDirtyFlags::None;
+	}
+
+	CleanPendingCanceledTaskList();
+}
+
+void SChaosVDRecordedLogBrowser::ToggleShowAllCategories()
+{
+    bShowAllCategories = !bShowAllCategories;
+
+	{
+		FWriteScopeLock WriteScopeLock(SourceDataCacheLock);
+    	// This is to mimic the filtering widget of the output log, where this setting stomps the enabled value of any category
+    	for (TPair<FName, FCategorizedItemsContainer>& Container : CachedLogItemsByCategory)
+    	{
+    		Container.Value.bIsEnabled = bShowAllCategories;
+    	}
+	}
+	
+	SetDirtyFlag(EChaosVDLogBrowserDirtyFlags::Categories);
+}
+
+bool SChaosVDRecordedLogBrowser::IsCategoryEnabled(FName CategoryName)
+{
+	FReadScopeLock ReadScopeLock(SourceDataCacheLock);
+
+	if (FCategorizedItemsContainer* Container = CachedLogItemsByCategory.Find(CategoryName))
+	{
+		return Container->bIsEnabled;
+	}
+
+	return false;
+}
+
+void SChaosVDRecordedLogBrowser::ToggleCategoryEnabled(FName CategoryName)
+{
+	FWriteScopeLock WriteScopeLock(SourceDataCacheLock);
+
+	if (FCategorizedItemsContainer* Container = CachedLogItemsByCategory.Find(CategoryName))
+	{
+		Container->bIsEnabled = !Container->bIsEnabled;
+
+		SetDirtyFlag(EChaosVDLogBrowserDirtyFlags::Categories);
+	}
+}
+
+void SChaosVDRecordedLogBrowser::SetVerbosityFlags(EChaosVDLogVerbosityFlags NewFlags)
+{
+	{
+		FWriteScopeLock WriteScopeLock(VerbosityFlagsLock);
+		VerbosityFlags = NewFlags;
+	}
+
+	SetDirtyFlag(EChaosVDLogBrowserDirtyFlags::Verbosity);
+}
+
+TSharedPtr<FChaosVDLogViewListItem> SChaosVDRecordedLogBrowser::MakeLogItem_AssumesLocked(const TSharedPtr<FChaosVDCachedLogItemEntry>& InLogData)
+{
+	if (!InLogData)
+	{
+		return nullptr;
+	}
+
+	if (TSharedPtr<FChaosVDLogViewListItem>* FoundItemPtr = CachedLogItemsByID.Find(InLogData->Index))
+	{
+		return *FoundItemPtr;
+	}
+
+	TSharedPtr<FChaosVDLogViewListItem> NewLogItem = MakeShared<FChaosVDLogViewListItem>();
+	NewLogItem->ItemWeakPtr = InLogData;
+	NewLogItem->EntryIndex = InLogData->Index;
+
+	CachedLogItemsByID.Add(InLogData->Index, NewLogItem);
+	FindOrAddCategorizedItemsContainer_AssumesLocked(InLogData->Category).Items.Add(NewLogItem);
+
+	return NewLogItem;
+}
+
+void SChaosVDRecordedLogBrowser::ApplyFiltersAsync()
+{
+	FReadScopeLock ReadScopeLock(SourceDataCacheLock);
+
+	EChaosVDLogBrowserDirtyFlags FlagsForFilteringContext = GetDirtyFlags();
+	// We can ignore playback because that is only related to selection events, no need to re-filter
+	EnumRemoveFlags(FlagsForFilteringContext, EChaosVDLogBrowserDirtyFlags::Playback);
+
+	// We can ignore filtering because this just indicates that a filter was applied, it does not mean it is out of date
+	EnumRemoveFlags(FlagsForFilteringContext, EChaosVDLogBrowserDirtyFlags::Filtering);
+
+	FChaosVDLogInFlightFilteringTaskWrapper::FContext FilteringContext(UnfilteredCachedLogItems.Num(), FilterEvaluator->GetFilterText(), FlagsForFilteringContext, VerbosityFlags);
+	if (CurrentFilteringTask)
+	{
+		FAsyncTask<FChaosVDLogFilteringAsyncTask>& AsyncTask = CurrentFilteringTask->GetAsyncTaskRef();
+
+		if (!AsyncTask.IsIdle())
+		{
+			// If we have a valid task that is in progress or is waiting to start, and the context (filter and source data) is the same as before,
+			// we can just early out
+			if (CurrentFilteringTask->HasSameContext(FilteringContext))
+			{
+				return;
+			}
+			else
+			{
+				AsyncTask.Cancel();
+	
+				// We need to keep the task alive until it is processed for cancellation
+				FilteringTasksBeingCancelled.Add(CurrentFilteringTask.ToSharedRef());
+			}
+		}
+	}
+	
+	CurrentFilteringTask = MakeShared<FChaosVDLogInFlightFilteringTaskWrapper>(FilteringContext, SharedThis(this));
+	CurrentFilteringTask->GetAsyncTaskRef().StartBackgroundTask();
+}
+
+void SChaosVDRecordedLogBrowser::ApplyFilterToData_AssumesLocked(TConstArrayView<TSharedPtr<FChaosVDLogViewListItem>> InDataSource, const TSharedRef<TArray<TSharedPtr<FChaosVDLogViewListItem>>>& OutFilteredData)
+{
+	for (const TSharedPtr<FChaosVDLogViewListItem>& LogListItem : InDataSource)
+	{
+		TSharedPtr<FChaosVDCachedLogItemEntry> LogEntry = LogListItem ? LogListItem->ItemWeakPtr.Pin() : nullptr;
+		if (!LogEntry)
+		{
+			continue;
+		}
+
+		if (IsVerbosityEnabled(LogEntry->Verbosity) && FilterEvaluator->TestTextFilter(FChaosVDBasicLogFilterExpressionContext(LogEntry.ToSharedRef())))
+		{
+			OutFilteredData->Emplace(LogListItem);
+		}
+	}
+}
+
+void SChaosVDRecordedLogBrowser::RefreshLogListView()
+{
+	TSharedPtr<SChaosVDRecordedLogView> LogViewList = LogViewWidget.Pin();
+	if (!LogViewList)
+	{
+		return;
+	}
+
+	FReadScopeLock ReadScopeLock(SourceDataCacheLock);
+	static TArray<TSharedPtr<FChaosVDLogViewListItem>> CurrentSelection;
+	CurrentSelection.Reset();
+	
+	LogViewList->GetSelectedItems(CurrentSelection);
+
+	LogViewList->SetSourceList(FilteredCachedLogItems);
+
+	LogViewList->SelectItems(CurrentSelection, ESelectInfo::Direct);
+}
+
+void SChaosVDRecordedLogBrowser::HandleItemSelected(const TSharedPtr<FChaosVDLogViewListItem>& InLogViewListItem, ESelectInfo::Type SelectInfo)
+{
+	if (SelectInfo != ESelectInfo::Direct)
+	{
+		HandleItemFocused(InLogViewListItem);
+	}
+}
+
+void SChaosVDRecordedLogBrowser::HandleItemFocused(const TSharedPtr<FChaosVDLogViewListItem>& InLogViewListItem)
+{
+	TSharedPtr<FChaosVDCachedLogItemEntry> LogEntryPtr = InLogViewListItem->ItemWeakPtr.Pin();
+
+	TSharedPtr<FChaosVDPlaybackController> PlaybackControllerPtr = GetObservedController().Pin();
+	TSharedPtr<FChaosVDRecording> RecordingData = PlaybackControllerPtr ? PlaybackControllerPtr->GetCurrentRecording().Pin() : nullptr;
+	if (!RecordingData)
+	{
+		return;
+	}
+
+	TSharedPtr<FChaosVDEngine> CVDEngineInstance = EngineInstanceWeakPtr.Pin();
+	TSharedPtr<const TraceServices::IAnalysisSession> Session = CVDEngineInstance ?
+																FChaosVDModule::Get().GetTraceManager()->GetSession(CVDEngineInstance->GetCurrentSessionDescriptor().SessionName) :
+																nullptr;
+
+	if (!Session)
+	{
+		return;
+	}
+
+	int32 GameFrameNumber = RecordingData->GetLowestGameFrameNumberAtTime(LogEntryPtr->Time);
+	if (GameFrameNumber != INDEX_NONE)
+	{
+		constexpr int32 StageNumber = 0;
+		PlaybackControllerPtr->GoToTrackFrameAndSync(GetInstigatorID(), EChaosVDTrackType::Game, FChaosVDPlaybackController::GameTrackID, GameFrameNumber, StageNumber);
+	}
+}
+
+void SChaosVDRecordedLogBrowser::HandlePlaybackControllerDataUpdated(TWeakPtr<FChaosVDPlaybackController> InController)
+{
+	FChaosVDPlaybackControllerObserver::HandlePlaybackControllerDataUpdated(InController);
+	SetDirtyFlag(EChaosVDLogBrowserDirtyFlags::Messages);
+}
+
+void SChaosVDRecordedLogBrowser::HandleControllerTrackFrameUpdated(TWeakPtr<FChaosVDPlaybackController> InController, TWeakPtr<const FChaosVDTrackInfo> UpdatedTrackInfo, FGuid InstigatorGuid)
+{
+	if (InstigatorGuid != GetInstigatorID())
+	{
+		PendingTrackUpdatesToProcess.Enqueue(UpdatedTrackInfo);
+		SetDirtyFlag(EChaosVDLogBrowserDirtyFlags::Playback);
+	}
+}
+
+void SChaosVDRecordedLogBrowser::SetDirtyFlag(EChaosVDLogBrowserDirtyFlags Flag)
+{
+	FWriteScopeLock WriteLock(DirtyFlagsLock);
+	EnumAddFlags(DirtyFlags, Flag);
+}
+
+void SChaosVDRecordedLogBrowser::RemoveDirtyFlag(EChaosVDLogBrowserDirtyFlags Flag)
+{
+	FWriteScopeLock WriteLock(DirtyFlagsLock);
+	EnumRemoveFlags(DirtyFlags, Flag);
+}
+
+void SChaosVDRecordedLogBrowser::ClearAllDirtyFlags()
+{
+	FWriteScopeLock WriteLock(DirtyFlagsLock);
+	DirtyFlags = EChaosVDLogBrowserDirtyFlags::None;
+}
+
+EChaosVDLogBrowserDirtyFlags SChaosVDRecordedLogBrowser::GetDirtyFlags() const
+{
+	FReadScopeLock ReadLock(DirtyFlagsLock);
+	return DirtyFlags;
+}
+
+SChaosVDRecordedLogBrowser::FCategorizedItemsContainer& SChaosVDRecordedLogBrowser::FindOrAddCategorizedItemsContainer_AssumesLocked(FName CategoryName)
+{
+	if (FCategorizedItemsContainer* Container = CachedLogItemsByCategory.Find(CategoryName))
+	{
+		return *Container;
+	}
+
+	FCategorizedItemsContainer NewContainer;
+	NewContainer.CategoryName = CategoryName;
+	NewContainer.bIsEnabled = bShowAllCategories;
+	return CachedLogItemsByCategory.Add(CategoryName, NewContainer);
+}
+
+bool SChaosVDRecordedLogBrowser::IsVerbosityEnabled(ELogVerbosity::Type VerbosityLevel)
+{
+	FReadScopeLock ReadScopeLock(VerbosityFlagsLock);
+
+	switch(VerbosityLevel)
+	{
+	case ELogVerbosity::Type::Error:
+		{
+			return EnumHasAnyFlags(VerbosityFlags, EChaosVDLogVerbosityFlags::Errors);
+		}
+	case ELogVerbosity::Type::Warning:
+		{
+			return EnumHasAnyFlags(VerbosityFlags, EChaosVDLogVerbosityFlags::Warnings);
+		}
+	default:
+		{
+			return EnumHasAnyFlags(VerbosityFlags, EChaosVDLogVerbosityFlags::Messages);
+		}
+	}	
+}
+
+void SChaosVDRecordedLogBrowser::ProcessTrackUpdate(const TSharedRef<FChaosVDRecording>& InRecordingData, const TSharedRef<SChaosVDRecordedLogView>& InLogViewList, const TraceServices::ILogProvider& InLogProvider, const TWeakPtr<const FChaosVDTrackInfo>& InPendingTrackInfo)
+{
+	TSharedPtr<const FChaosVDTrackInfo> TrackInfoPtr = InPendingTrackInfo.Pin();
+	if (!TrackInfoPtr)
+	{
+		return;
+	}
+
+	FMemMark StackMarker(FMemStack::Get());
+	TArray<uint64, TInlineAllocator<64, TMemStackAllocator<>>> MessageIndexesBuffer;
+
+	switch (TrackInfoPtr->TrackType)
+	{
+	case EChaosVDTrackType::Game:
+		{
+			if (FChaosVDGameFrameData* FrameData = InRecordingData->GetGameFrameData_AssumesLocked(TrackInfoPtr->CurrentFrame))
+			{
+				GetLogMessageIndexesForFrame(InLogProvider, *FrameData, MessageIndexesBuffer);
+			}
+			break;
+		}
+	case EChaosVDTrackType::Solver:
+		{
+			if (FChaosVDSolverFrameData* FrameData = InRecordingData->GetSolverFrameData_AssumesLocked(TrackInfoPtr->TrackID, TrackInfoPtr->CurrentFrame))
+			{
+				GetLogMessageIndexesForFrame(InLogProvider, *FrameData, MessageIndexesBuffer);
+			}
+			break;
+		}
+	default: break;
+	}
+
+	for (uint64 MessageIndex : MessageIndexesBuffer)
+	{
+		if (TSharedPtr<FChaosVDLogViewListItem>* FoundItem = CachedLogItemsByID.Find(MessageIndex))
+		{
+			InLogViewList->SelectItem(*FoundItem, ESelectInfo::Type::Direct);
+		}
+	}
+}
+
+void SChaosVDRecordedLogBrowser::ProcessPendingTrackUpdates()
+{
+	if (PendingTrackUpdatesToProcess.IsEmpty())
+	{
+		return;
+	}
+
+	TSharedPtr<SChaosVDRecordedLogView> LogViewList = LogViewWidget.Pin();
+	if (!LogViewList)
+	{
+		return;
+	}
+
+	TSharedPtr<FChaosVDPlaybackController> PlaybackControllerPtr = GetObservedController().Pin();
+	TSharedPtr<FChaosVDRecording> RecordingData = PlaybackControllerPtr ? PlaybackControllerPtr->GetCurrentRecording().Pin() : nullptr;
+	if (!RecordingData)
+	{
+		return;
+	}
+
+	TSharedPtr<FChaosVDEngine> CVDEngineInstance = EngineInstanceWeakPtr.Pin();
+	TSharedPtr<FChaosVDTraceManager> TraceManager = FChaosVDModule::Get().GetTraceManager();
+	TSharedPtr<const TraceServices::IAnalysisSession> Session = CVDEngineInstance && TraceManager ?
+																TraceManager->GetSession(CVDEngineInstance->GetCurrentSessionDescriptor().SessionName) :
+																nullptr;
+	if (!Session)
+	{
+		return;
+	}
+	
+	TraceServices::FAnalysisSessionReadScope SessionReadScope(*Session.Get());
+	const TraceServices::ILogProvider& LogProvider = TraceServices::ReadLogProvider(*Session.Get());
+
+	LogViewList->ClearSelection();
+
+	TWeakPtr<const FChaosVDTrackInfo> PendingTrackInfo;
+	while (PendingTrackUpdatesToProcess.Dequeue(PendingTrackInfo))
+	{
+		ProcessTrackUpdate(RecordingData.ToSharedRef(), LogViewList.ToSharedRef(), LogProvider, PendingTrackInfo);
+	}
+}
+
+void SChaosVDRecordedLogBrowser::Reset()
+{
+	FWriteScopeLock WriteScopeLock(SourceDataCacheLock);
+
+	CachedReadLogData.Reset();
+	CachedLogItemsByID.Reset();
+	CachedLogItemsByCategory.Reset();
+	UnfilteredCachedLogItems.Reset();
+}
+
+#undef LOCTEXT_NAMESPACE
diff --git a/Engine/Plugins/ChaosVD/Source/ChaosVD/Private/Widgets/SChaosVDRecordedLogBrowser.h b/Engine/Plugins/ChaosVD/Source/ChaosVD/Private/Widgets/SChaosVDRecordedLogBrowser.h
new file mode 100644
--- /dev/null
+++ b/Engine/Plugins/ChaosVD/Source/ChaosVD/Private/Widgets/SChaosVDRecordedLogBrowser.h
@@ -0,0 +1,290 @@
+﻿// Copyright Epic Games, Inc. All Rights Reserved.
+
+#pragma once
+
+#include "CoreMinimal.h"
+#include "ChaosVDPlaybackControllerInstigator.h"
+#include "ChaosVDPlaybackControllerObserver.h"
+#include "ContentBrowserDataMenuContexts.h"
+#include "Async/AsyncWork.h"
+#include "Misc/TextFilterExpressionEvaluator.h"
+#include "TraceServices/Model/Log.h"
+#include "Widgets/SCompoundWidget.h"
+
+#include "SChaosVDRecordedLogBrowser.generated.h"
+
+class FChaosVDEngine;
+class SChaosVDLogBrowserToolbar;
+class SChaosVDRecordedLogBrowser;
+class SChaosVDRecordedLogView;
+
+struct FChaosVDCachedLogItemEntry;
+struct FChaosVDLogViewListItem;
+struct FChaosVDRecording;
+
+namespace TraceServices
+{
+	class ILogProvider;
+}
+
+class FChaosVDBasicLogFilterExpressionContext : public ITextFilterExpressionContext
+{
+public:
+	explicit FChaosVDBasicLogFilterExpressionContext(const TSharedRef<FChaosVDCachedLogItemEntry>& InTestLogEntry)
+		: LogEntry(InTestLogEntry)
+	{
+	}
+
+	virtual bool TestBasicStringExpression(const FTextFilterString& InValue, const ETextFilterTextComparisonMode InTextComparisonMode) const override;
+	virtual bool TestComplexExpression(const FName& InKey, const FTextFilterString& InValue, const ETextFilterComparisonOperation InComparisonOperation, const ETextFilterTextComparisonMode InTextComparisonMode) const override
+	{
+		return false;
+	}
+
+private:
+	const TSharedRef<FChaosVDCachedLogItemEntry> LogEntry;
+};
+
+class FChaosVDLogFilteringAsyncTask : public FNonAbandonableTask
+{
+public:
+
+	explicit FChaosVDLogFilteringAsyncTask(const TWeakPtr<SChaosVDRecordedLogBrowser>& InLogBrowser) : LogBrowser(InLogBrowser)
+	{
+	}
+	
+	void DoWork();
+
+	FORCEINLINE TStatId GetStatId() const
+	{
+		RETURN_QUICK_DECLARE_CYCLE_STAT(FChaosVDLogFilteringAsyncTask, STATGROUP_ThreadPoolAsyncTasks);
+	}
+
+private:
+
+	TWeakPtr<SChaosVDRecordedLogBrowser> LogBrowser;
+
+	bool bAbandonTask = false;
+};
+
+UENUM(meta = (Bitflags, UseEnumValuesAsMaskValuesInEditor = "true"))
+enum class EChaosVDLogVerbosityFlags : uint8
+{
+	None = 0 UMETA(Hidden),
+	Messages = 1 << 0,
+	Warnings = 1 << 1,
+	Errors = 1 << 2,
+
+	All = Messages | Warnings | Errors
+};
+ENUM_CLASS_FLAGS(EChaosVDLogVerbosityFlags)
+
+UENUM(meta = (Bitflags, UseEnumValuesAsMaskValuesInEditor = "true"))
+enum class EChaosVDLogBrowserDirtyFlags : uint8
+{
+	None = 0 UMETA(Hidden),
+	Categories = 1 << 0,
+	Messages = 1 << 1,
+	Verbosity = 1 << 2,
+	Filtering = 1 << 3,
+	Playback = 1 << 4
+};
+ENUM_CLASS_FLAGS(EChaosVDLogBrowserDirtyFlags)
+
+class FChaosVDLogInFlightFilteringTaskWrapper
+{
+public:
+	struct FContext
+	{
+		explicit FContext(uint64 CurrentLogEntriesNum, const FText& InTextFilter, EChaosVDLogBrowserDirtyFlags InDirtyFlags, EChaosVDLogVerbosityFlags InnVerbosityFlags)
+			: CurrentLogEntriesNum(CurrentLogEntriesNum),
+			DirtyFlags(InDirtyFlags),
+			VerbosityFlags(InnVerbosityFlags),
+			TextFilter(InTextFilter.ToString())
+		{
+		}
+
+		uint64 CurrentLogEntriesNum;
+		EChaosVDLogBrowserDirtyFlags DirtyFlags;
+		EChaosVDLogVerbosityFlags VerbosityFlags;
+
+		FName TextFilter;
+	
+		bool operator==(const FContext& OtherContext) const = default;
+	};
+	
+	explicit FChaosVDLogInFlightFilteringTaskWrapper(const FContext& InContext, const TWeakPtr<SChaosVDRecordedLogBrowser>& LogBrowser)
+	: Context(InContext)
+	{
+		AsyncTask = MakeUnique<FAsyncTask<FChaosVDLogFilteringAsyncTask>>(LogBrowser);
+	}
+
+	bool HasSameContext(const FContext& InContext) const
+	{
+		return InContext == Context;
+	}
+
+	FAsyncTask<FChaosVDLogFilteringAsyncTask>& GetAsyncTaskRef() const
+	{
+		return *AsyncTask.Get();
+	}
+	
+private:
+
+	TUniquePtr<FAsyncTask<FChaosVDLogFilteringAsyncTask>> AsyncTask;
+	FContext Context;
+};
+
+/**
+ * Widget used to render the recorded log stream in a CVD Recording
+ */
+class SChaosVDRecordedLogBrowser : public SCompoundWidget, public FChaosVDPlaybackControllerObserver, public IChaosVDPlaybackControllerInstigator
+{
+public:
+	SLATE_BEGIN_ARGS(SChaosVDRecordedLogBrowser)
+		{
+		}
+
+
+	SLATE_END_ARGS()
+
+	/** Constructs this widget with InArgs */
+	void Construct(const FArguments& InArgs, const TSharedRef<FChaosVDEngine>& InEngineInstance);
+
+	virtual ~SChaosVDRecordedLogBrowser() override;
+	
+	void HandleSearchTextChanged(const FText& NewText);
+	void UpdateLogLineSelectionFromGameTrack();
+
+
+	TSharedPtr<SChaosVDLogBrowserToolbar> GetToolBar()
+	{
+		return ToolbarPtr;
+	};
+
+	virtual void Tick(const FGeometry& AllottedGeometry, const double InCurrentTime, const float InDeltaTime) override;
+	
+	static const FName ToolBarName;
+
+	bool GetShowAllCategories() const
+	{
+		return bShowAllCategories;
+	}
+	
+	void ToggleShowAllCategories();
+
+	template<typename Visitor>
+	void EnumerateNonEmptyCategories(Visitor VisitorFunc)
+	{
+		FReadScopeLock ReadScopeLock(SourceDataCacheLock);
+
+		for (const TPair<FName, FCategorizedItemsContainer>& ItemsContainer : CachedLogItemsByCategory)
+		{
+			VisitorFunc(ItemsContainer.Value);
+		}
+	}
+
+	bool IsCategoryEnabled(FName CategoryName);
+	void ToggleCategoryEnabled(FName CategoryName);
+
+	void SetVerbosityFlags(EChaosVDLogVerbosityFlags NewFlags);
+
+	EChaosVDLogVerbosityFlags GetVerbosityFlags() const
+	{
+		return VerbosityFlags;
+	}
+
+	struct FCategorizedItemsContainer
+	{
+		TArray<TSharedPtr<FChaosVDLogViewListItem>> Items;
+		FName CategoryName = NAME_None;
+		bool bIsEnabled = true;
+	};
+
+protected:
+	virtual void HandlePlaybackControllerDataUpdated(TWeakPtr<FChaosVDPlaybackController> InController) override;
+	virtual void HandleControllerTrackFrameUpdated(TWeakPtr<FChaosVDPlaybackController> InController, TWeakPtr<const FChaosVDTrackInfo> UpdatedTrackInfo, FGuid InstigatorGuid) override;
+
+private:
+
+	void ProcessTrackUpdate(const TSharedRef<FChaosVDRecording>& InRecordingData, const TSharedRef<SChaosVDRecordedLogView>& InLogViewList, const TraceServices::ILogProvider& InLogProvider, const TWeakPtr<const FChaosVDTrackInfo>& InPendingTrackInfo);
+
+	void SetDirtyFlag(EChaosVDLogBrowserDirtyFlags Flag);
+	void RemoveDirtyFlag(EChaosVDLogBrowserDirtyFlags Flag);
+	void ClearAllDirtyFlags();
+	EChaosVDLogBrowserDirtyFlags GetDirtyFlags() const;
+
+	FCategorizedItemsContainer& FindOrAddCategorizedItemsContainer_AssumesLocked(FName CategoryName);
+
+	bool IsVerbosityEnabled(ELogVerbosity::Type VerbosityLevel);
+
+	void ProcessPendingTrackUpdates();
+
+	void Reset();
+
+	FText GetFilterStatusText() const;
+	FSlateColor GetFilterStatusTextColor() const;
+
+	void CleanPendingCanceledTaskList();
+	void UpdateBrowserContents();
+
+	TSharedPtr<FChaosVDLogViewListItem> MakeLogItem_AssumesLocked(const TSharedPtr<FChaosVDCachedLogItemEntry>& InLogData);
+
+	void ApplyFiltersAsync();
+
+	void ApplyFilterToData_AssumesLocked(TConstArrayView<TSharedPtr<FChaosVDLogViewListItem>> InDataSource, const TSharedRef<TArray<TSharedPtr<FChaosVDLogViewListItem>>>& OutFilteredData);
+	void RefreshLogListView();
+
+	void HandleItemSelected(const TSharedPtr<FChaosVDLogViewListItem>& InLogViewListItem, ESelectInfo::Type SelectInfo);
+	void HandleItemFocused(const TSharedPtr<FChaosVDLogViewListItem>& InLogViewListItem);
+
+	template<typename FrameType, typename IndexesContainerType>
+	void GetLogMessageIndexesForFrame(const TraceServices::ILogProvider& InLogProvider, const FrameType& InFrameData, IndexesContainerType& OutIndexesContainer);
+
+	TQueue<TWeakPtr<const FChaosVDTrackInfo>> PendingTrackUpdatesToProcess;
+	TSharedPtr<FTextFilterExpressionEvaluator> FilterEvaluator;
+
+	TWeakPtr<FChaosVDEngine> EngineInstanceWeakPtr;
+	
+	TWeakPtr<SChaosVDRecordedLogView> LogViewWidget;
+	
+	TSharedPtr<TArray<TSharedPtr<FChaosVDLogViewListItem>>> FilteredCachedLogItems;
+	TArray<TSharedPtr<FChaosVDLogViewListItem>> UnfilteredCachedLogItems;
+
+	TArray<TSharedPtr<FChaosVDLogViewListItem>> SelectedLogItems;
+
+	TArray<TSharedPtr<FChaosVDCachedLogItemEntry>> CachedReadLogData;
+	TMap<int32, TSharedPtr<FChaosVDLogViewListItem>> CachedLogItemsByID;
+	TMap<FName, FCategorizedItemsContainer> CachedLogItemsByCategory;
+
+	TSharedPtr<FChaosVDLogInFlightFilteringTaskWrapper> CurrentFilteringTask;
+	TArray<TSharedRef<FChaosVDLogInFlightFilteringTaskWrapper>> FilteringTasksBeingCancelled;
+
+	TSharedPtr<SChaosVDLogBrowserToolbar> ToolbarPtr;
+
+	bool bShowAllCategories = true;
+
+	EChaosVDLogVerbosityFlags VerbosityFlags = EChaosVDLogVerbosityFlags::All;
+
+	EChaosVDLogBrowserDirtyFlags DirtyFlags = EChaosVDLogBrowserDirtyFlags::None;
+
+	mutable FRWLock SourceDataCacheLock;
+	mutable FRWLock DirtyFlagsLock;
+	mutable FRWLock VerbosityFlagsLock;
+
+	friend FChaosVDLogFilteringAsyncTask;
+};
+
+template <typename FrameType, typename IndexesContainerType>
+void SChaosVDRecordedLogBrowser::GetLogMessageIndexesForFrame(const TraceServices::ILogProvider& InLogProvider, const FrameType& InFrameData, IndexesContainerType& OutIndexesContainer)
+{
+	InLogProvider.EnumerateMessages(InFrameData.StartTime, InFrameData.EndTime, [&OutIndexesContainer](const TraceServices::FLogMessageInfo& InMessageInfo)
+			{
+				OutIndexesContainer.Add(InMessageInfo.Index);
+			});
+
+	if (OutIndexesContainer.IsEmpty())
+	{
+		OutIndexesContainer.Add(InLogProvider.BinarySearchClosestByTime(InFrameData.StartTime));
+	}
+}
diff --git a/Engine/Plugins/ChaosVD/Source/ChaosVD/Private/Widgets/SChaosVDRecordedLogView.cpp b/Engine/Plugins/ChaosVD/Source/ChaosVD/Private/Widgets/SChaosVDRecordedLogView.cpp
new file mode 100644
--- /dev/null
+++ b/Engine/Plugins/ChaosVD/Source/ChaosVD/Private/Widgets/SChaosVDRecordedLogView.cpp
@@ -0,0 +1,202 @@
+﻿// Copyright Epic Games, Inc. All Rights Reserved.
+
+#include "SChaosVDRecordedLogView.h"
+
+#include "Widgets/Text/STextBlock.h"
+
+#define LOCTEXT_NAMESPACE "ChaosVisualDebugger"
+
+SChaosVDRecordedLogView::FColumNames SChaosVDRecordedLogView::ColumnNames = SChaosVDRecordedLogView::FColumNames();
+
+void SChaosVDLogViewRow::Construct(const FArguments& InArgs, const TSharedRef<STableViewBase>& InOwnerTableView)
+{
+	Item = InArgs._Item;
+	FSuperRowType::FArguments Args = FSuperRowType::FArguments()
+									.Style(&FAppStyle::Get().GetWidgetStyle<FTableRowStyle>("SceneOutliner.TableViewRow"));
+
+	SMultiColumnTableRow<TSharedPtr<const FChaosVDLogViewListItem>>::Construct(Args, InOwnerTableView);
+}
+
+TSharedRef<SWidget> SChaosVDLogViewRow::GenerateTextWidgetFromText(const FText& Text, const FLinearColor& InColor)
+{
+	constexpr float MarginLeft = 4.0f;
+	constexpr float NoMargin = 0.0f;
+
+	return SNew(STextBlock)
+			.Margin(FMargin(MarginLeft, NoMargin, NoMargin, NoMargin))
+			.ColorAndOpacity(InColor)
+			.Text(Text);
+}
+
+TSharedRef<SWidget> SChaosVDLogViewRow::GenerateWidgetForColumn(const FName& ColumnName)
+{
+	if (!Item)
+	{
+		return SNullWidget::NullWidget;
+	}
+
+	TSharedPtr<const FChaosVDCachedLogItemEntry> LogEntryDataPtr = Item->ItemWeakPtr.Pin();
+	if (!LogEntryDataPtr)
+	{
+		return SNullWidget::NullWidget;
+	}
+
+	FLinearColor LogLineColor = FLinearColor::White;
+	if (LogEntryDataPtr->Verbosity == ELogVerbosity::Error)
+	{
+		LogLineColor = FLinearColor::Red;
+	}
+	else if (LogEntryDataPtr->Verbosity == ELogVerbosity::Warning)
+	{
+		LogLineColor = FLinearColor::Yellow;
+	}
+	
+	if (ColumnName == SChaosVDRecordedLogView::ColumnNames.Time)
+	{
+		static const FNumberFormattingOptions RecordedTimeFormatOptions = FNumberFormattingOptions()
+		.SetMinimumFractionalDigits(3)
+		.SetMaximumFractionalDigits(3);
+
+		const FText RecordingTimeSecondsAsText = FText::AsNumber(LogEntryDataPtr->Time, &RecordedTimeFormatOptions);
+		return GenerateTextWidgetFromText(FText::FormatOrdered(FText::AsCultureInvariant("{0}s"), RecordingTimeSecondsAsText), LogLineColor);
+	}
+	
+	if (ColumnName == SChaosVDRecordedLogView::ColumnNames.Verbosity)
+	{
+		return GenerateTextWidgetFromText(FText::FromString(::ToString(LogEntryDataPtr->Verbosity)), LogLineColor);
+	}
+
+	if (ColumnName == SChaosVDRecordedLogView::ColumnNames.Category)
+	{
+		return GenerateTextWidgetFromText(FText::FromName(LogEntryDataPtr->Category), LogLineColor);
+	}
+
+	if (ColumnName == SChaosVDRecordedLogView::ColumnNames.Message)
+	{
+		return GenerateTextWidgetFromText(FText::FromString(LogEntryDataPtr->Message), LogLineColor);
+	}
+	
+	return SNullWidget::NullWidget;
+}
+
+void SChaosVDRecordedLogView::Construct(const FArguments& InArgs)
+{
+	ItemSelectedDelegate = InArgs._OnItemSelected;
+	ItemFocusedDelegate = InArgs._OnItemFocused;
+
+	constexpr float BottomPadding = 2.0f;
+	constexpr float NoPadding = 0.0f;
+
+	ChildSlot
+	[
+		SNew(SVerticalBox)
+		+SVerticalBox::Slot()
+		.Padding(NoPadding, NoPadding, NoPadding, BottomPadding)
+		[
+			SAssignNew(LogListWidget, SListView<TSharedPtr<FChaosVDLogViewListItem>>)
+			.OnGenerateRow(this, &SChaosVDRecordedLogView::GenerateLogEntryRow)
+			.OnSelectionChanged(this, &SChaosVDRecordedLogView::LogItemSelectionChanged)
+			.OnMouseButtonDoubleClick(this, &SChaosVDRecordedLogView::HandleFocusRequest)
+			.SelectionMode(ESelectionMode::Type::Multi)
+			.ListItemsSource(&InternalItemSourceData)
+			.HeaderRow
+			(
+				// TODO: Add sort support for the columns
+				SNew(SHeaderRow)
+				+SHeaderRow::Column(ColumnNames.Time)
+				.SortMode(EColumnSortMode::None)
+				.ManualWidth(80.0f)
+				.DefaultLabel(LOCTEXT("LogTimeHeader", "Time"))
+
+				+SHeaderRow::Column(ColumnNames.Category)
+				.SortMode(EColumnSortMode::None)
+				.ManualWidth(160.0f)
+				.DefaultLabel(LOCTEXT("LogCategoryHeader", "Category"))
+
+				+SHeaderRow::Column(ColumnNames.Verbosity)
+				.SortMode(EColumnSortMode::None)
+				.ManualWidth(80.0f)
+				.DefaultLabel(LOCTEXT("LogVerbosityHeader", "Verbosity"))
+				
+				+SHeaderRow::Column(ColumnNames.Message)
+				.SortMode(EColumnSortMode::None)
+				.FillWidth(1.0f)
+				.DefaultLabel(LOCTEXT("LogMessageHeader", "Message"))
+			)
+		]
+	];
+}
+
+TSharedRef<ITableRow> SChaosVDRecordedLogView::GenerateLogEntryRow(TSharedPtr<FChaosVDLogViewListItem> LogEntryData, const TSharedRef<STableViewBase>& OwnerTable)
+{
+	if (!LogEntryData)
+	{
+		return
+		SNew(STableRow< TSharedPtr<FString> >, OwnerTable)
+		[
+			SNew(SVerticalBox)
+			+SVerticalBox::Slot()
+			[
+				SNew(STextBlock)
+				.Text(LOCTEXT("SChaosVDSceneQueryListErrorMessage", "Failed to read data for solver."))
+			]
+		];
+	}
+
+	return SNew(SChaosVDLogViewRow, OwnerTable)
+			.Item(LogEntryData);
+}
+
+void SChaosVDRecordedLogView::LogItemSelectionChanged(TSharedPtr<FChaosVDLogViewListItem> SelectedLogItem, ESelectInfo::Type Type)
+{
+	ItemSelectedDelegate.ExecuteIfBound(SelectedLogItem, Type);
+}
+
+void SChaosVDRecordedLogView::SelectItem(const TSharedPtr<FChaosVDLogViewListItem>& ItemToSelect, ESelectInfo::Type Type)
+{
+	if (LogListWidget->IsItemSelected(ItemToSelect))
+	{
+		return;
+	}
+
+	constexpr bool bIsSelected = true;
+	LogListWidget->SetItemSelection(ItemToSelect, bIsSelected, Type);
+	LogListWidget->RequestScrollIntoView(ItemToSelect);
+}
+
+void SChaosVDRecordedLogView::SelectItems(TConstArrayView<TSharedPtr<FChaosVDLogViewListItem>> ItemsToSelect, ESelectInfo::Type Type)
+{
+	for (TSharedPtr<FChaosVDLogViewListItem> Item : ItemsToSelect)
+	{
+		SelectItem(Item, ESelectInfo::Direct);
+	}
+}
+
+int32 SChaosVDRecordedLogView::GetSelectedItems(TArray<TSharedPtr<FChaosVDLogViewListItem>>& OutItems)
+{
+	return LogListWidget->GetSelectedItems(OutItems);
+}
+
+void SChaosVDRecordedLogView::ClearSelection()
+{
+	LogListWidget->ClearSelection();
+}
+
+void SChaosVDRecordedLogView::SetSourceList(const TSharedPtr<TArray<TSharedPtr<FChaosVDLogViewListItem>>>& InSourceList)
+{
+	LogListWidget->SetItemsSource(InSourceList.Get());
+	LogListWidget->RebuildList();
+}
+
+void SChaosVDRecordedLogView::Tick(const FGeometry& AllottedGeometry, const double InCurrentTime, const float InDeltaTime)
+{
+	SCompoundWidget::Tick(AllottedGeometry, InCurrentTime, InDeltaTime);
+}
+
+void SChaosVDRecordedLogView::HandleFocusRequest(TSharedPtr<FChaosVDLogViewListItem> InFocusedItem)
+{
+	ItemFocusedDelegate.ExecuteIfBound(InFocusedItem);
+}
+
+#undef LOCTEXT_NAMESPACE
+
diff --git a/Engine/Plugins/ChaosVD/Source/ChaosVD/Private/Widgets/SChaosVDRecordedLogView.h b/Engine/Plugins/ChaosVD/Source/ChaosVD/Private/Widgets/SChaosVDRecordedLogView.h
new file mode 100644
--- /dev/null
+++ b/Engine/Plugins/ChaosVD/Source/ChaosVD/Private/Widgets/SChaosVDRecordedLogView.h
@@ -0,0 +1,105 @@
+﻿// Copyright Epic Games, Inc. All Rights Reserved.
+
+#pragma once
+
+#include "Containers/Array.h"
+#include "Templates/SharedPointer.h"
+#include "Widgets/SCompoundWidget.h"
+#include "Widgets/Views/STableViewBase.h"
+#include "Widgets/Views/SListView.h"
+#include "Widgets/Views/STableRow.h"
+
+class FChaosVDEngine;
+
+struct FChaosVDCachedLogItemEntry
+{
+	uint64 Index = 0;
+	FName Category = NAME_None;
+	ELogVerbosity::Type Verbosity = ELogVerbosity::Type::All;
+	FString Message;
+	double Time = 0.0;
+};
+
+struct FChaosVDLogViewListItem
+{
+	TWeakPtr<FChaosVDCachedLogItemEntry> ItemWeakPtr;
+	uint64 EntryIndex = 0;
+};
+
+class SChaosVDSceneQueryBrowser;
+struct FChaosVDQueryDataWrapper;
+struct FChaosVDSceneQueryTreeItem;
+
+/**
+ * Widget used to represent a row on the Scene Query Browser tree view
+ */
+class SChaosVDLogViewRow : public SMultiColumnTableRow<TSharedPtr<const FChaosVDLogViewListItem>>
+{
+public:
+	SLATE_BEGIN_ARGS(SChaosVDLogViewRow)
+	{
+	}
+	SLATE_ARGUMENT(TSharedPtr<FChaosVDLogViewListItem>, Item)
+
+	SLATE_END_ARGS()
+
+	void Construct(const FArguments& InArgs, const TSharedRef<STableViewBase>& InOwnerTableView);
+
+	virtual TSharedRef<SWidget> GenerateWidgetForColumn(const FName& ColumnName) override;
+
+	TSharedRef<SWidget> GenerateTextWidgetFromText(const FText& Text, const FLinearColor& InColor = FLinearColor::White);
+
+private:
+	TSharedPtr<FChaosVDLogViewListItem> Item;
+};
+
+DECLARE_DELEGATE_TwoParams(FChaosVDLogIetmItemSelected, const TSharedPtr<FChaosVDLogViewListItem>&, ESelectInfo::Type)
+DECLARE_DELEGATE_OneParam(FChaosVDLogItemFocused, const TSharedPtr<FChaosVDLogViewListItem>&)
+
+class SChaosVDRecordedLogView : public SCompoundWidget
+{
+public:
+	SLATE_BEGIN_ARGS(SChaosVDRecordedLogView) {}
+	SLATE_EVENT(FChaosVDLogIetmItemSelected, OnItemSelected)
+	SLATE_EVENT(FChaosVDLogItemFocused, OnItemFocused)
+SLATE_END_ARGS()
+	
+	void Construct(const FArguments& InArgs);
+
+	int32 GetSelectedItems(TArray<TSharedPtr<FChaosVDLogViewListItem>>& OutItems);
+	void SelectItem(const TSharedPtr<FChaosVDLogViewListItem>& ItemToSelect, ESelectInfo::Type Type);
+	void SelectItems(TConstArrayView<TSharedPtr<FChaosVDLogViewListItem>> ItemsToSelect, ESelectInfo::Type Type);
+	void ClearSelection();
+
+	void SetSourceList(const TSharedPtr<TArray<TSharedPtr<FChaosVDLogViewListItem>>>& InSourceList);
+
+	void Tick(const FGeometry& AllottedGeometry, const double InCurrentTime, const float InDeltaTime) override;
+
+protected:
+
+	void HandleFocusRequest(TSharedPtr<FChaosVDLogViewListItem> InFocusedItem);
+
+	TSharedRef<ITableRow> GenerateLogEntryRow(TSharedPtr<FChaosVDLogViewListItem> LogEntryData, const TSharedRef<STableViewBase>& OwnerTable);
+
+	void LogItemSelectionChanged(TSharedPtr<FChaosVDLogViewListItem> SelectedLogItem, ESelectInfo::Type Type);
+
+	TSharedPtr<SListView<TSharedPtr<FChaosVDLogViewListItem>>> LogListWidget;
+
+	TArray<TSharedPtr<FChaosVDLogViewListItem>> InternalItemSourceData;
+
+	FChaosVDLogIetmItemSelected ItemSelectedDelegate;
+	FChaosVDLogItemFocused ItemFocusedDelegate;
+
+public:
+
+	struct FColumNames
+	{
+		const FName Time = FName("Time");
+		const FName Category = FName("Category");
+		const FName Verbosity = FName("Verbosity");
+		const FName Message = FName("Message");
+	};
+
+	static FColumNames ColumnNames;
+	
+};
diff --git a/Engine/Source/Runtime/Experimental/ChaosVDData/Private/ChaosVDRecording.cpp b/Engine/Source/Runtime/Experimental/ChaosVDData/Private/ChaosVDRecording.cpp
--- a/Engine/Source/Runtime/Experimental/ChaosVDData/Private/ChaosVDRecording.cpp
+++ b/Engine/Source/Runtime/Experimental/ChaosVDData/Private/ChaosVDRecording.cpp
@@ -361,6 +361,12 @@ int32 FChaosVDRecording::GetLowestGameFrameNumberAtCycle(uint64 Cycle)
 	return Algo::LowerBoundBy(GameFrames, Cycle, &FChaosVDGameFrameData::FirstCycle);
 }
 
+int32 FChaosVDRecording::GetLowestGameFrameNumberAtTime(double Time)
+{
+	FReadScopeLock ReadLock(RecordingDataLock);
+	return Algo::LowerBoundBy(GameFrames, Time, &FChaosVDGameFrameData::StartTime);
+}
+
 void FChaosVDRecording::CollapseSolverFramesRange_AssumesLocked(int32 SolverID, int32 StartFrame, int32 EndFrame, FChaosVDSolverFrameData& OutCollapsedFrameData)
 {
 	// Make sure we start with a clear map
diff --git a/Engine/Source/Runtime/Experimental/ChaosVDData/Public/ChaosVDRecording.h b/Engine/Source/Runtime/Experimental/ChaosVDData/Public/ChaosVDRecording.h
--- a/Engine/Source/Runtime/Experimental/ChaosVDData/Public/ChaosVDRecording.h
+++ b/Engine/Source/Runtime/Experimental/ChaosVDData/Public/ChaosVDRecording.h
@@ -328,6 +328,13 @@ struct CHAOSVDDATA_API FChaosVDRecording
 	 */
 	int32 GetLowestGameFrameNumberAtCycle(uint64 Cycle);
 
+	/**
+	 * Searches and returns the lowest game frame number at the specified cycle
+	 * @param Time Platform Time to be used in the search as lower bound
+	 * @return Found Game frame number. INDEX_NONE if no frame is found for the specified cycle
+	 */
+	int32 GetLowestGameFrameNumberAtTime(double Time);
+
 	/**
      * Gathers all available solvers IDs at the given Game frame number
      * @param FrameNumber Game Frame number to evaluate
diff --git a/Engine/Source/Runtime/Experimental/ChaosVisualDebugger/Private/ChaosVDRuntimeModule.cpp b/Engine/Source/Runtime/Experimental/ChaosVisualDebugger/Private/ChaosVDRuntimeModule.cpp
--- a/Engine/Source/Runtime/Experimental/ChaosVisualDebugger/Private/ChaosVDRuntimeModule.cpp
+++ b/Engine/Source/Runtime/Experimental/ChaosVisualDebugger/Private/ChaosVDRuntimeModule.cpp
@@ -171,6 +171,7 @@ void FChaosVDRuntimeModule::StartRecording(TConstArrayView<FString> Args)
 
 	UE::Trace::ToggleChannel(TEXT("ChaosVDChannel"), true); 
 	UE::Trace::ToggleChannel(TEXT("Frame"), true);
+	UE::Trace::ToggleChannel(TEXT("Log"), true);
 
 	FTraceAuxiliary::FOptions TracingOptions;
 	TracingOptions.bExcludeTail = true;
