[Chaos VD] Fix for log spam and missing that happens when a recording a scene with dynamic collections.

CVD needs to map each inner implicit object to a valid shape instance index so it can properly colorize and show the correct per shape data in the details panel.

We were already handling the case for nested Implicit Object unions, where the objects in the inner union share the same shape instance index, but we need to do the same for union clustered objects that might be the root.

diff --git a/Engine/Plugins/ChaosVD/Source/ChaosVD/Private/ChaosVDGeometryBuilder.cpp b/Engine/Plugins/ChaosVD/Source/ChaosVD/Private/ChaosVDGeometryBuilder.cpp
--- a/Engine/Plugins/ChaosVD/Source/ChaosVD/Private/ChaosVDGeometryBuilder.cpp
+++ b/Engine/Plugins/ChaosVD/Source/ChaosVD/Private/ChaosVDGeometryBuilder.cpp
@@ -324,10 +324,10 @@ void FChaosVDGeometryBuilder::DeInitialize()
 	bInitialized = false;
 }
 
-void FChaosVDGeometryBuilder::CreateMeshesFromImplicitObject(const Chaos::FImplicitObject* InImplicitObject, AActor* Owner, TArray<TSharedPtr<FChaosVDExtractedGeometryDataHandle>>& OutMeshDataHandles, const int32 DesiredLODCount, const Chaos::FRigidTransform3& InTransform, const int32 MeshIndex)
+void FChaosVDGeometryBuilder::CreateMeshesFromImplicitObject(const Chaos::FImplicitObject* InImplicitObject, AActor* Owner, TArray<TSharedPtr<FChaosVDExtractedGeometryDataHandle>>& OutMeshDataHandles,int32 AvailableShapeDataNum , const int32 DesiredLODCount, const Chaos::FRigidTransform3& InTransform, const int32 MeshIndex)
 {
 	// To start set the leaf and the root to the same ptr. If the object is an union, in the subsequent recursive call the leaf will be set correctly
-	CreateMeshesFromImplicit_Internal(InImplicitObject, InImplicitObject, Owner, OutMeshDataHandles, DesiredLODCount, InTransform, MeshIndex);
+	CreateMeshesFromImplicit_Internal(InImplicitObject, InImplicitObject, Owner, OutMeshDataHandles, DesiredLODCount, InTransform, MeshIndex, AvailableShapeDataNum);
 }
 
 void FChaosVDGeometryBuilder::AddReferencedObjects(FReferenceCollector& Collector)
@@ -383,7 +383,7 @@ bool FChaosVDGeometryBuilder::HasNegativeScale(const Chaos::FRigidTransform3& In
 	return ScaleSignVector.X * ScaleSignVector.Y * ScaleSignVector.Z < 0;
 }
 
-void FChaosVDGeometryBuilder::CreateMeshesFromImplicit_Internal(const Chaos::FImplicitObject* InRootImplicitObject, const Chaos::FImplicitObject* InLeafImplicitObject, AActor* Owner, TArray<TSharedPtr<FChaosVDExtractedGeometryDataHandle>>& OutMeshDataHandles, const int32 DesiredLODCount, const Chaos::FRigidTransform3& InTransform, const int32 ShapeInstanceIndex)
+void FChaosVDGeometryBuilder::CreateMeshesFromImplicit_Internal(const Chaos::FImplicitObject* InRootImplicitObject, const Chaos::FImplicitObject* InLeafImplicitObject, AActor* Owner, TArray<TSharedPtr<FChaosVDExtractedGeometryDataHandle>>& OutMeshDataHandles, const int32 DesiredLODCount, const Chaos::FRigidTransform3& InTransform, const int32 ParentShapeInstanceIndex, int32 AvailableShapeDataNum)
 {
 	TRACE_CPUPROFILER_EVENT_SCOPE(FChaosVDGeometryBuilder::CreateMeshesFromImplicit_Internal);
 
@@ -397,14 +397,31 @@ void FChaosVDGeometryBuilder::CreateMeshesFromImplicit_Internal(const Chaos::FIm
 		{
 			const bool bIsRootUnion = InRootImplicitObject == InLeafImplicitObject;
 
+			const bool bIsCluster = InnerType == ImplicitObjectType::UnionClustered;
+
 			for (int32 ObjectIndex = 0; ObjectIndex < Union->GetObjects().Num(); ++ObjectIndex)
 			{
 				const FImplicitObjectPtr& UnionImplicit = Union->GetObjects()[ObjectIndex];
 
-				// If this union it is not the root implicit object, then all its objects will share the same Instance index
-				int32 CurrentShapeInstanceIndex = bIsRootUnion ? ObjectIndex : ShapeInstanceIndex;
+				int32 CurrentShapeInstanceIndex = ParentShapeInstanceIndex;
 
-				CreateMeshesFromImplicit_Internal(InRootImplicitObject, UnionImplicit.GetReference(), Owner, OutMeshDataHandles, DesiredLODCount, InTransform, CurrentShapeInstanceIndex);	
+				if (bIsRootUnion)
+				{
+					if (bIsCluster)
+					{
+						// Geometry Collections might break the usual rule of how may shape data instances we have per geometry
+						// Sometimes they can create clusters where all particles share a single instance
+						constexpr int32 SingleShapeInstanceDataIndex = 0;
+						CurrentShapeInstanceIndex = AvailableShapeDataNum == 1 ? SingleShapeInstanceDataIndex : ParentShapeInstanceIndex;
+					}
+					else
+					{
+						// If this union it is not the root implicit object, and it is not a cluster, then all its objects will share the same Instance index
+						CurrentShapeInstanceIndex = ObjectIndex;
+					}
+				}
+
+				CreateMeshesFromImplicit_Internal(InRootImplicitObject, UnionImplicit.GetReference(), Owner, OutMeshDataHandles, DesiredLODCount, InTransform, CurrentShapeInstanceIndex, AvailableShapeDataNum);	
 			}
 		}
 
@@ -416,7 +433,7 @@ void FChaosVDGeometryBuilder::CreateMeshesFromImplicit_Internal(const Chaos::FIm
 		if (const TImplicitObjectTransformed<FReal, 3>* Transformed = InLeafImplicitObject->template GetObject<TImplicitObjectTransformed<FReal, 3>>())
 		{
 			// For transformed objects, the Instance index is the same so we pass it in without changing it
-			CreateMeshesFromImplicit_Internal(InRootImplicitObject, Transformed->GetTransformedObject(), Owner, OutMeshDataHandles, DesiredLODCount, Transformed->GetTransform(), ShapeInstanceIndex);
+			CreateMeshesFromImplicit_Internal(InRootImplicitObject, Transformed->GetTransformedObject(), Owner, OutMeshDataHandles, DesiredLODCount, Transformed->GetTransform(), ParentShapeInstanceIndex, AvailableShapeDataNum);
 		}
 		
 		return;
@@ -425,13 +442,14 @@ void FChaosVDGeometryBuilder::CreateMeshesFromImplicit_Internal(const Chaos::FIm
 	if (const TSharedPtr<FChaosVDExtractedGeometryDataHandle> MeshDataHandle = ExtractGeometryDataForImplicit(InLeafImplicitObject, InTransform))
 	{
 		MeshDataHandle->SetImplicitObject(InLeafImplicitObject);
-		MeshDataHandle->SetShapeInstanceIndex(ShapeInstanceIndex);
+		MeshDataHandle->SetShapeInstanceIndex(ParentShapeInstanceIndex);
 		MeshDataHandle->SetRootImplicitObject(InRootImplicitObject);
 
 		OutMeshDataHandles.Add(MeshDataHandle);
 	}
 }
 
+
 bool FChaosVDGeometryBuilder::HasGeometryInCache(uint32 GeometryKey)
 {
 	FReadScopeLock ReadLock(GeometryCacheRWLock);
diff --git a/Engine/Plugins/ChaosVD/Source/ChaosVD/Private/ChaosVDGeometryBuilder.h b/Engine/Plugins/ChaosVD/Source/ChaosVD/Private/ChaosVDGeometryBuilder.h
--- a/Engine/Plugins/ChaosVD/Source/ChaosVD/Private/ChaosVDGeometryBuilder.h
+++ b/Engine/Plugins/ChaosVD/Source/ChaosVD/Private/ChaosVDGeometryBuilder.h
@@ -68,7 +68,7 @@ public:
 	 *	@param InTransform to apply to the generated components/geometry
 	 */
 
-	void CreateMeshesFromImplicitObject(const Chaos::FImplicitObject* InImplicitObject, AActor* Owner, TArray<TSharedPtr<FChaosVDExtractedGeometryDataHandle>>& OutMeshDataHandles, const int32 DesiredLODCount = 0, const Chaos::FRigidTransform3& InTransform = Chaos::FRigidTransform3(), const int32 MeshIndex = 0);
+	void CreateMeshesFromImplicitObject(const Chaos::FImplicitObject* InImplicitObject, AActor* Owner, TArray<TSharedPtr<FChaosVDExtractedGeometryDataHandle>>& OutMeshDataHandles, int32 AvailableShapeDataNum = 0, const int32 DesiredLODCount = 0, const Chaos::FRigidTransform3& InTransform = Chaos::FRigidTransform3(), const int32 MeshIndex = 0);
 
 	virtual void AddReferencedObjects(FReferenceCollector& Collector) override;
 	virtual FString GetReferencerName() const override
@@ -92,7 +92,7 @@ public:
 
 private:
 	
-	void CreateMeshesFromImplicit_Internal(const Chaos::FImplicitObject* InRootImplicitObject,const Chaos::FImplicitObject* InLeafImplicitObject, AActor* Owner, TArray<TSharedPtr<FChaosVDExtractedGeometryDataHandle>>& OutMeshDataHandles, const int32 DesiredLODCount = 0, const Chaos::FRigidTransform3& InTransform = Chaos::FRigidTransform3(), const int32 ShapeInstanceIndex = 0);
+	void CreateMeshesFromImplicit_Internal(const Chaos::FImplicitObject* InRootImplicitObject,const Chaos::FImplicitObject* InLeafImplicitObject, AActor* Owner, TArray<TSharedPtr<FChaosVDExtractedGeometryDataHandle>>& OutMeshDataHandles, const int32 DesiredLODCount = 0, const Chaos::FRigidTransform3& InTransform = Chaos::FRigidTransform3(), const int32 ParentShapeInstanceIndex = 0, int32 AvailableShapeDataNum = 0);
 	
 public:
 	/**
diff --git a/Engine/Plugins/ChaosVD/Source/ChaosVD/Private/ChaosVDGeometryDataComponent.cpp b/Engine/Plugins/ChaosVD/Source/ChaosVD/Private/ChaosVDGeometryDataComponent.cpp
--- a/Engine/Plugins/ChaosVD/Source/ChaosVD/Private/ChaosVDGeometryDataComponent.cpp
+++ b/Engine/Plugins/ChaosVD/Source/ChaosVD/Private/ChaosVDGeometryDataComponent.cpp
@@ -228,7 +228,12 @@ void FChaosVDGeometryComponentUtils::UpdateCollisionDataFromShapeArray(const TAr
 		const FName ImplicitObjectTypeName = InInstanceHandle->GetState().ImplicitObjectInfo.ImplicitObjectType;
 		const Chaos::FImplicitObject* RootImplicitObject = ExtractedGeometryHandle->GetRootImplicitObject();
 		const FName RootImplicitObjectTypeName = !InInstanceHandle->GetState().ImplicitObjectInfo.bIsRootObject && RootImplicitObject ? Chaos::GetImplicitObjectTypeName(Chaos::GetInnerType(RootImplicitObject->GetType())) : TEXT("None");
-		UE_LOG(LogChaosVDEditor, Warning, TEXT("[%s] Failed to find shape instance data at Index [%d] | Particle ID[%d] | Available Shape instance Data Num [%d] | Implicit Type [%s] - Root Implicit Type [%s] | This geometry will be hidden..."), ANSI_TO_TCHAR(__FUNCTION__), ShapeInstanceIndex, InInstanceHandle->GetOwningParticleID(), InShapeArray.Num(), *ImplicitObjectTypeName.ToString(), *RootImplicitObjectTypeName.ToString());
+		
+		FString ErrorMessage = FString::Printf(TEXT("[%s] Failed to find shape instance data at Index [%d] | Particle ID[%d] | Available Shape instance Data Num [%d] | Implicit Type [%s] - Root Implicit Type [%s] | This geometry will be hidden..."), ANSI_TO_TCHAR(__FUNCTION__), ShapeInstanceIndex, InInstanceHandle->GetOwningParticleID(), InShapeArray.Num(), *ImplicitObjectTypeName.ToString(), *RootImplicitObjectTypeName.ToString());
+		
+		UE_LOG(LogChaosVDEditor, Verbose, TEXT("[%s]"), *ErrorMessage);
+
+		ensureMsgf(false, TEXT("[%s]"), *ErrorMessage);
 
 		InInstanceHandle->bFailedToUpdateShapeInstanceData = true;
 		return;
@@ -236,7 +241,7 @@ void FChaosVDGeometryComponentUtils::UpdateCollisionDataFromShapeArray(const TAr
 	else if (InInstanceHandle->bFailedToUpdateShapeInstanceData)
 	{
 		InInstanceHandle->bFailedToUpdateShapeInstanceData = false;
-		UE_LOG(LogChaosVDEditor, Warning, TEXT("[%s] Recovered from failing to find shape instance data at Index [%d] | Particle ID[%d] | Available Shape instance Data Num [%d] | This geometry will be shown again..."), ANSI_TO_TCHAR(__FUNCTION__), ShapeInstanceIndex, InInstanceHandle->GetOwningParticleID(), InShapeArray.Num());
+		UE_LOG(LogChaosVDEditor, Verbose, TEXT("[%s] Recovered from failing to find shape instance data at Index [%d] | Particle ID[%d] | Available Shape instance Data Num [%d] | This geometry will be shown again..."), ANSI_TO_TCHAR(__FUNCTION__), ShapeInstanceIndex, InInstanceHandle->GetOwningParticleID(), InShapeArray.Num());
 	}
 
 	FChaosVDShapeCollisionData CollisionDataToUpdate = InShapeArray[ShapeInstanceIndex];
diff --git a/Engine/Plugins/ChaosVD/Source/ChaosVD/Private/ChaosVDParticleActor.cpp b/Engine/Plugins/ChaosVD/Source/ChaosVD/Private/ChaosVDParticleActor.cpp
--- a/Engine/Plugins/ChaosVD/Source/ChaosVD/Private/ChaosVDParticleActor.cpp
+++ b/Engine/Plugins/ChaosVD/Source/ChaosVD/Private/ChaosVDParticleActor.cpp
@@ -176,6 +176,10 @@ void AChaosVDParticleActor::ProcessUpdatedAndRemovedHandles(TArray<TSharedPtr<FC
 				{
 					bExists = true;
 
+					// Although the geometry is the same, we need to copy over all the new data on the updated handle
+					// Otherwise the ptr to the root implicit object or the Shape Instance Index will be outdated 
+					*ExistingMeshDataHandle->GetGeometryHandle() = *GeometryDataHandle;
+
 					// If we have a CVD Geometry Component for this handle, just remove it from the list as it means we don't need to re-create it
 					HandleRemoveIterator.RemoveCurrent();
 					break;
@@ -254,7 +258,7 @@ void AChaosVDParticleActor::UpdateGeometry(const Chaos::FConstImplicitObjectPtr&
 	constexpr int32 LODsToGenerateNum = 3;
 	constexpr int32 LODsToGenerateNumForInstancedStaticMesh = 0;
 
-	GeometryGenerator->CreateMeshesFromImplicitObject(InImplicitObject, this, OutExtractedGeometryDataHandles, bHasToUseStaticMeshComponent ? LODsToGenerateNum : LODsToGenerateNumForInstancedStaticMesh);
+	GeometryGenerator->CreateMeshesFromImplicitObject(InImplicitObject, this, OutExtractedGeometryDataHandles, ParticleDataPtr->CollisionDataPerShape.Num(), bHasToUseStaticMeshComponent ? LODsToGenerateNum : LODsToGenerateNumForInstancedStaticMesh);
 
 	// This should not happen in theory, but there might be some valid situations where it does. Adding an ensure to catch them and then evaluate if it is really an issue (if it is not I will remove the ensure later on). 
 	if (!ensure(ObjectsToGenerateNum == OutExtractedGeometryDataHandles.Num()))
diff --git a/Engine/Source/Runtime/Experimental/Chaos/Private/Chaos/ImplicitObject.cpp b/Engine/Source/Runtime/Experimental/Chaos/Private/Chaos/ImplicitObject.cpp
--- a/Engine/Source/Runtime/Experimental/Chaos/Private/Chaos/ImplicitObject.cpp
+++ b/Engine/Source/Runtime/Experimental/Chaos/Private/Chaos/ImplicitObject.cpp
@@ -37,6 +37,9 @@ namespace Chaos
 
 	FImplicitObject::~FImplicitObject()
 	{
+		// CVD caches geometry hashes using ptrs as key. There is a chance that this ptr could get re-used
+		CVD_TRACE_INVALIDATE_CACHED_GEOMETRY(this)
+
 #if TRACK_CHAOS_GEOMETRY
 		if(bIsTracked)
 		{
