[ChaosVD] Allow toggle translucent geometry selection

For 5.4 I locked translucent selection to be always enabled as CVD uses translucent geometry a lot to show query only shapes. The issue with that is that we also have cases of very small shapes surrounded by a box (usually an overlap box used as a trigger), which means we cannot select the small shapes/particles inside it.

This CL  adds the same shortcut used by the Level editor to toggle on and off translucent selection.

In addition, this CL also fixes an issue where the key shortcut to focus the viewport on selection was not working.

diff --git a/Engine/Plugins/ChaosVD/Source/ChaosVD/Private/ChaosVDCommands.cpp b/Engine/Plugins/ChaosVD/Source/ChaosVD/Private/ChaosVDCommands.cpp
--- a/Engine/Plugins/ChaosVD/Source/ChaosVD/Private/ChaosVDCommands.cpp
+++ b/Engine/Plugins/ChaosVD/Source/ChaosVD/Private/ChaosVDCommands.cpp
@@ -14,6 +14,7 @@ void FChaosVDCommands::RegisterCommands()
 {
 	UI_COMMAND(ToggleFollowSelectedObject, "Follow Selected Object", "Start or Stop following the selected object", EUserInterfaceActionType::ToggleButton, FInputChord(EKeys::F8));
 	UI_COMMAND(OverridePlaybackFrameRate, "Override Recorded Framerate", "When enabled, allows to playback the recording at a fixed framerate", EUserInterfaceActionType::ToggleButton, FInputChord(EModifierKey::Control, EKeys::R));
+	UI_COMMAND(AllowTranslucentSelection, "Allow Translucent Selection", "Allows translucent objects to be selected", EUserInterfaceActionType::ToggleButton, FInputChord(EKeys::T));
 }
 
 #undef LOCTEXT_NAMESPACE
diff --git a/Engine/Plugins/ChaosVD/Source/ChaosVD/Private/ChaosVDCommands.h b/Engine/Plugins/ChaosVD/Source/ChaosVD/Private/ChaosVDCommands.h
--- a/Engine/Plugins/ChaosVD/Source/ChaosVD/Private/ChaosVDCommands.h
+++ b/Engine/Plugins/ChaosVD/Source/ChaosVD/Private/ChaosVDCommands.h
@@ -17,4 +17,5 @@ public:
 public:
 	TSharedPtr<FUICommandInfo> ToggleFollowSelectedObject;
 	TSharedPtr<FUICommandInfo> OverridePlaybackFrameRate;
+	TSharedPtr<FUICommandInfo> AllowTranslucentSelection;
 };
\ No newline at end of file
diff --git a/Engine/Plugins/ChaosVD/Source/ChaosVD/Private/ChaosVDPlaybackViewportClient.cpp b/Engine/Plugins/ChaosVD/Source/ChaosVD/Private/ChaosVDPlaybackViewportClient.cpp
--- a/Engine/Plugins/ChaosVD/Source/ChaosVD/Private/ChaosVDPlaybackViewportClient.cpp
+++ b/Engine/Plugins/ChaosVD/Source/ChaosVD/Private/ChaosVDPlaybackViewportClient.cpp
@@ -145,6 +145,18 @@ void FChaosVDPlaybackViewportClient::SetScene(TWeakPtr<FChaosVDScene> InScene)
 	}
 }
 
+void FChaosVDPlaybackViewportClient::SetCanSelectTranslucentGeometry(bool bCanSelect)
+{
+	bAllowTranslucentHitProxies = bCanSelect;
+
+	Invalidate();
+}
+
+void FChaosVDPlaybackViewportClient::ToggleCanSelectTranslucentGeometry()
+{
+	SetCanSelectTranslucentGeometry(!bAllowTranslucentHitProxies);
+}
+
 void FChaosVDPlaybackViewportClient::HandleFocusRequest(FBox BoxToFocusOn)
 {
 	FocusViewportOnBox(BoxToFocusOn);
@@ -172,6 +184,11 @@ void FChaosVDPlaybackViewportClient::TrackSelectedObject()
 		return;
 	}
 
+	FocusOnSelectedObject();
+}
+
+void FChaosVDPlaybackViewportClient::FocusOnSelectedObject()
+{
 	if (const TSharedPtr<FChaosVDScene> CVDSceneSharedPtr = CVDScene.Pin())
 	{
 		USelection* CurrentSelection = ModeTools->GetSelectedActors();
@@ -221,10 +238,10 @@ void FChaosVDPlaybackViewportClient::Draw(const FSceneView* View, FPrimitiveDraw
 {
 	if (View)
 	{
-		// Hack to allow selection of translucent objects (for CVD is all geometry set a Query Only)
+		// Hack to allow CVD control the selection of translucent objects (for CVD is all geometry set as Query Only)
 		// The current setting to allow this behaviour is project wide or on custom hitproxies implementations which we can't use
-		// A proper fix would be have a way to override this per viewport, which could be done by adding a new method to FViewElementDrawer
-		const_cast<FSceneView*>(View)->bAllowTranslucentPrimitivesInHitProxy = true;
+		// A proper fix would be to have a way to override this per viewport, which could be done by adding a new method to FViewElementDrawer
+		const_cast<FSceneView*>(View)->bAllowTranslucentPrimitivesInHitProxy = bAllowTranslucentHitProxies;
 	}
 
 	const TSharedPtr<SChaosVDMainTab> MainTabToolkitHost = ModeTools.IsValid() ? StaticCastSharedPtr<SChaosVDMainTab>(ModeTools->GetToolkitHost()) : nullptr;
diff --git a/Engine/Plugins/ChaosVD/Source/ChaosVD/Private/ChaosVDPlaybackViewportClient.h b/Engine/Plugins/ChaosVD/Source/ChaosVD/Private/ChaosVDPlaybackViewportClient.h
--- a/Engine/Plugins/ChaosVD/Source/ChaosVD/Private/ChaosVDPlaybackViewportClient.h
+++ b/Engine/Plugins/ChaosVD/Source/ChaosVD/Private/ChaosVDPlaybackViewportClient.h
@@ -35,11 +35,17 @@ public:
 	void GoToLocation(const FVector& InLocation);
 
 	void TrackSelectedObject();
+
+	void FocusOnSelectedObject();
 	
 	virtual void UpdateMouseDelta() override;
 
 	void HandleCVDSceneUpdated();
 
+	bool GetCanSelectTranslucentGeometry() const { return bAllowTranslucentHitProxies; };
+	void SetCanSelectTranslucentGeometry(bool bCanSelect);
+	void ToggleCanSelectTranslucentGeometry();
+
 private:
 
 	void HandleFocusRequest(FBox BoxToFocusOn);
@@ -51,5 +57,5 @@ private:
 
 	bool bAutoTrackSelectedObject = false;
 	float TrackingViewDistance = 120.0f;
-
+	bool bAllowTranslucentHitProxies = true;
 };
diff --git a/Engine/Plugins/ChaosVD/Source/ChaosVD/Private/Widgets/SChaosVDPlaybackViewport.cpp b/Engine/Plugins/ChaosVD/Source/ChaosVD/Private/Widgets/SChaosVDPlaybackViewport.cpp
--- a/Engine/Plugins/ChaosVD/Source/ChaosVD/Private/Widgets/SChaosVDPlaybackViewport.cpp
+++ b/Engine/Plugins/ChaosVD/Source/ChaosVD/Private/Widgets/SChaosVDPlaybackViewport.cpp
@@ -128,15 +128,12 @@ void SChaosVDPlaybackViewport::BindCommands()
 		const TSharedRef<FChaosVDPlaybackViewportClient> ViewportClientRef = StaticCastSharedRef<FChaosVDPlaybackViewportClient>(Client.ToSharedRef());
 		FUIAction ToggleObjectTrackingAction;
 		ToggleObjectTrackingAction.ExecuteAction.BindSP(ViewportClientRef, &FChaosVDPlaybackViewportClient::ToggleObjectTrackingIfSelected);
-		ToggleObjectTrackingAction.GetActionCheckState.BindLambda([WeakThis = ViewportClientRef.ToWeakPtr()]()
+		ToggleObjectTrackingAction.GetActionCheckState.BindLambda([WeakViewportClient = ViewportClientRef.ToWeakPtr()]()
 		{
-			TSharedPtr<const FChaosVDPlaybackViewportClient> ViewportPtr = StaticCastSharedPtr<const FChaosVDPlaybackViewportClient>(WeakThis.Pin());
+			TSharedPtr<const FChaosVDPlaybackViewportClient> ViewportPtr = WeakViewportClient.Pin();
 			return ViewportPtr.IsValid() && ViewportPtr->IsAutoTrackingSelectedObject() ? ECheckBoxState::Checked : ECheckBoxState::Unchecked;
 		});
-	
-		CommandList->MapAction(
-			Commands.ToggleFollowSelectedObject,
-			ToggleObjectTrackingAction);
+		CommandList->MapAction(Commands.ToggleFollowSelectedObject, ToggleObjectTrackingAction);
 
 		FUIAction ToggleOverrideFrameRateAction;
 		ToggleOverrideFrameRateAction.ExecuteAction.BindSP(SharedThis(this), &SChaosVDPlaybackViewport::ToggleUseFrameRateOverride);
@@ -145,10 +142,16 @@ void SChaosVDPlaybackViewport::BindCommands()
 			TSharedPtr<const SChaosVDPlaybackViewport> ViewportPtr = StaticCastSharedPtr<const SChaosVDPlaybackViewport>(WeakThis.Pin());
 			return ViewportPtr.IsValid() && ViewportPtr->IsUsingFrameRateOverride() ? ECheckBoxState::Checked : ECheckBoxState::Unchecked;
 		});
+		CommandList->MapAction(Commands.OverridePlaybackFrameRate, ToggleOverrideFrameRateAction);
 
-		CommandList->MapAction(
-			Commands.OverridePlaybackFrameRate,
-			ToggleOverrideFrameRateAction);
+		FUIAction ToggleTranslucentGeometrySelectionAction;
+		ToggleTranslucentGeometrySelectionAction.ExecuteAction.BindSP(ViewportClientRef, &FChaosVDPlaybackViewportClient::ToggleCanSelectTranslucentGeometry);
+		ToggleTranslucentGeometrySelectionAction.GetActionCheckState.BindLambda([WeakViewportClient = ViewportClientRef.ToWeakPtr()]()
+		{
+			TSharedPtr<const FChaosVDPlaybackViewportClient> ViewportPtr = WeakViewportClient.Pin();
+			return ViewportPtr.IsValid() && ViewportPtr->GetCanSelectTranslucentGeometry() ? ECheckBoxState::Checked : ECheckBoxState::Unchecked;
+		});
+		CommandList->MapAction(Commands.AllowTranslucentSelection, ToggleTranslucentGeometrySelectionAction);
 	}
 }
 
@@ -204,6 +207,11 @@ void SChaosVDPlaybackViewport::ExecuteExternalViewportInvalidateRequest()
 	ExternalViewportInvalidationRequestHandler.Broadcast();
 }
 
+void SChaosVDPlaybackViewport::OnFocusViewportToSelection()
+{
+	PlaybackViewportClient->FocusOnSelectedObject();
+}
+
 TSharedRef<FEditorViewportClient> SChaosVDPlaybackViewport::MakeEditorViewportClient()
 {
 	TSharedPtr<FChaosVDPlaybackViewportClient> NewViewport = MakeShared<FChaosVDPlaybackViewportClient>(EditorModeTools, GetViewportWidget());
diff --git a/Engine/Plugins/ChaosVD/Source/ChaosVD/Private/Widgets/SChaosVDPlaybackViewport.h b/Engine/Plugins/ChaosVD/Source/ChaosVD/Private/Widgets/SChaosVDPlaybackViewport.h
--- a/Engine/Plugins/ChaosVD/Source/ChaosVD/Private/Widgets/SChaosVDPlaybackViewport.h
+++ b/Engine/Plugins/ChaosVD/Source/ChaosVD/Private/Widgets/SChaosVDPlaybackViewport.h
@@ -58,6 +58,8 @@ public:
 	
 	static void ExecuteExternalViewportInvalidateRequest();
 
+	virtual void OnFocusViewportToSelection() override;
+
 protected:
 
 	virtual TSharedRef<FEditorViewportClient> MakeEditorViewportClient() override;
diff --git a/Engine/Plugins/ChaosVD/Source/ChaosVD/Private/Widgets/SChaosVDViewportToolbar.cpp b/Engine/Plugins/ChaosVD/Source/ChaosVD/Private/Widgets/SChaosVDViewportToolbar.cpp
--- a/Engine/Plugins/ChaosVD/Source/ChaosVD/Private/Widgets/SChaosVDViewportToolbar.cpp
+++ b/Engine/Plugins/ChaosVD/Source/ChaosVD/Private/Widgets/SChaosVDViewportToolbar.cpp
@@ -93,6 +93,9 @@ void SChaosVDViewportToolbar::ExtendOptionsMenu(FMenuBuilder& OptionsMenuBuilder
 		CVDOptionsMenuBuilder.AddWidget(GenerateFarViewPlaneMenu(), LOCTEXT("FarViewPlane", "Far View Plane"));
 
 		CVDOptionsMenuBuilder.AddSeparator();
+
+		CVDOptionsMenuBuilder.AddMenuEntry(FChaosVDCommands::Get().AllowTranslucentSelection, NAME_None);
+		
 	}
 	CVDOptionsMenuBuilder.EndSection();
 
