[ChaosVD] Fix for an issue where sphere collision shapes inside of a static mesh are visualized at the incorrect location.

There issue is caused because when we add a sphere collision that has an offset in its position relative to its parent transform, we just change the center position. CVD was only handling such scenarios for implicit transformed objects. The fix is just take into account the center location when we calculate the final transform, the same way as we do it for capsules already.

diff --git a/Engine/Plugins/ChaosVD/Source/ChaosVD/Private/ChaosVDGeometryBuilder.cpp b/Engine/Plugins/ChaosVD/Source/ChaosVD/Private/ChaosVDGeometryBuilder.cpp
--- a/Engine/Plugins/ChaosVD/Source/ChaosVD/Private/ChaosVDGeometryBuilder.cpp
+++ b/Engine/Plugins/ChaosVD/Source/ChaosVD/Private/ChaosVDGeometryBuilder.cpp
@@ -800,7 +800,7 @@ void FChaosVDGeometryBuilder::AdjustedTransformForImplicit(const Chaos::FImplici
 	using namespace Chaos;
 	switch (GetInnerType(InImplicit->GetType()))
 	{
-		// Currently, only capsules transforms needs to be re-adjusted
+		// Currently, only capsules and spheres transforms needs to be re-adjusted to take into account non-zero center locations
 		case ImplicitObjectType::Capsule:
 		{
 			if (const FCapsule* Capsule = InImplicit->template GetObject<FCapsule>())
@@ -814,6 +814,15 @@ void FChaosVDGeometryBuilder::AdjustedTransformForImplicit(const Chaos::FImplici
 			}
 			break;
 		}
+		case ImplicitObjectType::Sphere:
+			{
+				if (const Chaos::TSphere<FReal, 3>* Sphere = InImplicit->template GetObject<Chaos::TSphere<FReal, 3>>())
+				{
+					const FVector FinalLocation = OutAdjustedTransform.TransformPosition(Sphere->GetCenter());
+					OutAdjustedTransform.SetLocation(FinalLocation);
+				}
+				break;
+			}
 		default:
 			break;
 	}
