[ChaosVD] Fix for an issue where if async physics is enabled and the GT ticks at a lower delta time than the PT, CVD will discard game thread data even if it was properly recorded.

This happened due to an issue with the "auto trim" feature I implemented very early on, which pre-dates proper async physics support. We re-use the Frame Start/End event already available in UE trace and as we don't have control over them, this means we can have frames with no valid CVD data, and that is why we want to remove them from the recording.

The problem is that since that was implemented, now an "empty frame" is not the same as it was before. Previously an empty frame was a frame with no solver frame data associated with it, but nowadays game frames can have GT specific data like scene queries.

Another issue is that with async physics, the end of a GT frame can happen while a PT frame is also being processed, so we need to evaluate the frames we didn't commit to the CVD recording yet before determining a GT frame is indeed empty.

diff --git a/Engine/Plugins/ChaosVD/Source/ChaosVD/Private/Trace/ChaosVDTraceProvider.cpp b/Engine/Plugins/ChaosVD/Source/ChaosVD/Private/Trace/ChaosVDTraceProvider.cpp
--- a/Engine/Plugins/ChaosVD/Source/ChaosVD/Private/Trace/ChaosVDTraceProvider.cpp
+++ b/Engine/Plugins/ChaosVD/Source/ChaosVD/Private/Trace/ChaosVDTraceProvider.cpp
@@ -27,6 +27,15 @@
 #include "Trace/DataProcessors/ChaosVDArchiveHeaderProcessor.h"
 #include "Trace/DataProcessors/ChaosVDDataProcessorBase.h"
 
+namespace Chaos::VisualDebugger::Cvars
+{
+	static int32 MaxFramesToQueue = 2;
+	static FAutoConsoleVariableRef CVarChaosVDMaxFramesToQueue(
+		TEXT("p.Chaos.VD.Tool.MaxFramesToQueue"),
+		MaxFramesToQueue,
+		TEXT("The max number of game frames we will queue before processing them while the trace files is being loaded."));
+}
+
 FName FChaosVDTraceProvider::ProviderName("ChaosVDProvider");
 
 FChaosVDTraceProvider::FChaosVDTraceProvider(TraceServices::IAnalysisSession& InSession): Session(InSession)
@@ -65,16 +74,29 @@ void FChaosVDTraceProvider::StartSolverFrame(const int32 InSolverID, FChaosVDSol
 	}
 }
 
+void FChaosVDTraceProvider::GetAvailablePendingSolverIDsAtGameFrame(const TSharedRef<FChaosVDGameFrameData>& InProcessedGameFrameData, TArray<int32, TInlineAllocator<16>>& OutSolverIDs)
+{
+	for (const TPair<int32, FChaosVDSolverFrameData>& FrameDataWithSolverID : CurrentSolverFramesByID)
+	{
+		if (FrameDataWithSolverID.Value.FrameCycle < InProcessedGameFrameData->FirstCycle)
+		{
+			OutSolverIDs.Add(FrameDataWithSolverID.Key);
+		}
+	}
+}
+
 void FChaosVDTraceProvider::CommitProcessedGameFramesToRecording()
 {
 	TArray<int32, TInlineAllocator<16>> SolverIDs;
 
-	// The Game Frame events are not generated by CVD trace code and we don't have control over them
+	// The Game Frame events are not generated by CVD trace code, and we don't have control over them.
 	// we use them as general timestamps.
 	// These are generated even when no solvers are available (specially in PIE), so we need to discard any game frame that will not resolve to a solver frame
-	// Solver data might be slightly behind, so to make sure we have all the solver data for that time range, we queue a handful of game frames before processing them.
-	constexpr int32 MaxQueueSize = 2;
-	if (CurrentGameFrameQueueSize > MaxQueueSize)
+	// Physics Frames and GT frames lifetimes might not align with async physics enabled, so to make sure we have all the solver data for that time range, we queue a handful of game frames before processing them.
+
+	// TODO: This empty frame logic trimming is no longer 100% correct since we added support for async physics and support for recording GT specific data.
+	// We need to find a better way to trim out empty data. Jira for Tracking UE-227365
+	if (CurrentGameFrameQueueSize > Chaos::VisualDebugger::Cvars::MaxFramesToQueue)
 	{
 		TSharedPtr<FChaosVDGameFrameData> ProcessedGameFrameData;
 		DeQueueGameFrameForProcessing(ProcessedGameFrameData);
@@ -88,7 +110,17 @@ void FChaosVDTraceProvider::CommitProcessedGameFramesToRecording()
 		{
 			InternalRecording->GetAvailableSolverIDsAtGameFrame(*ProcessedGameFrameData, SolverIDs);
 
-			if (!SolverIDs.IsEmpty())
+			// Is it possible that the solver data is not commited to the recording yet as it is being processed.
+			// Usually this happens on recordings with Async Physics
+			if (SolverIDs.IsEmpty())
+			{
+				GetAvailablePendingSolverIDsAtGameFrame(ProcessedGameFrameData.ToSharedRef(), SolverIDs);
+			}
+
+			// Any frame that is marked as dirty, or has any kind of solver data is considered relevant. Otherwise, we can trim it from the recording
+			// by just ignoring it
+			const bool bHasRelevantCVDData = !SolverIDs.IsEmpty() || ProcessedGameFrameData->IsDirty();
+			if (bHasRelevantCVDData)
 			{
 				InternalRecording->AddGameFrameData(*ProcessedGameFrameData);	
 			}
diff --git a/Engine/Plugins/ChaosVD/Source/ChaosVD/Private/Trace/ChaosVDTraceProvider.h b/Engine/Plugins/ChaosVD/Source/ChaosVD/Private/Trace/ChaosVDTraceProvider.h
--- a/Engine/Plugins/ChaosVD/Source/ChaosVD/Private/Trace/ChaosVDTraceProvider.h
+++ b/Engine/Plugins/ChaosVD/Source/ChaosVD/Private/Trace/ChaosVDTraceProvider.h
@@ -100,6 +100,9 @@ private:
 
 	void EnqueueGameFrameForProcessing(const TSharedPtr<FChaosVDGameFrameData>& FrameData);
 	void DeQueueGameFrameForProcessing(TSharedPtr<FChaosVDGameFrameData>& OutFrameData);
+
+	/** Gathers any solver id from solver data that is not fully processed yet but that will be valid for the provided game frame data later on */
+	void GetAvailablePendingSolverIDsAtGameFrame(const TSharedRef<FChaosVDGameFrameData>& InProcessedGameFrameData, TArray<int32, TInlineAllocator<16>>& OutSolverIDs);
 	
 	TraceServices::IAnalysisSession& Session;
 
diff --git a/Engine/Plugins/ChaosVD/Source/ChaosVD/Private/Trace/DataProcessors/ChaosVDAABBTreeDataProcessor.cpp b/Engine/Plugins/ChaosVD/Source/ChaosVD/Private/Trace/DataProcessors/ChaosVDAABBTreeDataProcessor.cpp
--- a/Engine/Plugins/ChaosVD/Source/ChaosVD/Private/Trace/DataProcessors/ChaosVDAABBTreeDataProcessor.cpp
+++ b/Engine/Plugins/ChaosVD/Source/ChaosVD/Private/Trace/DataProcessors/ChaosVDAABBTreeDataProcessor.cpp
@@ -57,6 +57,7 @@ bool FChaosVDAABBTreeDataProcessor::ProcessRawData(const TArray<uint8>& InData)
 
 			TArray<TSharedPtr<FChaosVDAABBTreeDataWrapper>>& ABBTrees = CurrentFrameData->RecordedAABBTreesBySolverID.FindOrAdd(AABBTreeData->SolverId);
 			ABBTrees.Add(AABBTreeData);
+			CurrentFrameData->MarkDirty();
 		}
 	}
 
diff --git a/Engine/Plugins/ChaosVD/Source/ChaosVD/Private/Trace/DataProcessors/ChaosVDDebugDrawBoxDataProcessor.cpp b/Engine/Plugins/ChaosVD/Source/ChaosVD/Private/Trace/DataProcessors/ChaosVDDebugDrawBoxDataProcessor.cpp
--- a/Engine/Plugins/ChaosVD/Source/ChaosVD/Private/Trace/DataProcessors/ChaosVDDebugDrawBoxDataProcessor.cpp
+++ b/Engine/Plugins/ChaosVD/Source/ChaosVD/Private/Trace/DataProcessors/ChaosVDDebugDrawBoxDataProcessor.cpp
@@ -25,6 +25,7 @@ bool FChaosVDDebugDrawBoxDataProcessor::ProcessRawData(const TArray<uint8>& InDa
 		if (TSharedPtr<FChaosVDGameFrameData> CurrentFrameData = ProviderSharedPtr->GetCurrentGameFrame().Pin())
 		{
 			CurrentFrameData->RecordedDebugDrawBoxesBySolverID.FindOrAdd(DebugDrawData->SolverID).Add(DebugDrawData);
+			CurrentFrameData->MarkDirty();
 		}
 	}
 
diff --git a/Engine/Plugins/ChaosVD/Source/ChaosVD/Private/Trace/DataProcessors/ChaosVDDebugDrawImplicitObjectDataProcessor.cpp b/Engine/Plugins/ChaosVD/Source/ChaosVD/Private/Trace/DataProcessors/ChaosVDDebugDrawImplicitObjectDataProcessor.cpp
--- a/Engine/Plugins/ChaosVD/Source/ChaosVD/Private/Trace/DataProcessors/ChaosVDDebugDrawImplicitObjectDataProcessor.cpp
+++ b/Engine/Plugins/ChaosVD/Source/ChaosVD/Private/Trace/DataProcessors/ChaosVDDebugDrawImplicitObjectDataProcessor.cpp
@@ -25,6 +25,7 @@ bool FChaosVDDebugDrawImplicitObjectDataProcessor::ProcessRawData(const TArray<u
 		if (TSharedPtr<FChaosVDGameFrameData> CurrentFrameData = ProviderSharedPtr->GetCurrentGameFrame().Pin())
 		{
 			CurrentFrameData->RecordedDebugDrawImplicitObjectsBySolverID.FindOrAdd(DebugDrawData->SolverID).Add(DebugDrawData);
+			CurrentFrameData->MarkDirty();
 		}
 	}
 
diff --git a/Engine/Plugins/ChaosVD/Source/ChaosVD/Private/Trace/DataProcessors/ChaosVDDebugDrawLineDataProcessor.cpp b/Engine/Plugins/ChaosVD/Source/ChaosVD/Private/Trace/DataProcessors/ChaosVDDebugDrawLineDataProcessor.cpp
--- a/Engine/Plugins/ChaosVD/Source/ChaosVD/Private/Trace/DataProcessors/ChaosVDDebugDrawLineDataProcessor.cpp
+++ b/Engine/Plugins/ChaosVD/Source/ChaosVD/Private/Trace/DataProcessors/ChaosVDDebugDrawLineDataProcessor.cpp
@@ -25,6 +25,7 @@ bool FChaosVDDebugDrawLineDataProcessor::ProcessRawData(const TArray<uint8>& InD
 		if (TSharedPtr<FChaosVDGameFrameData> CurrentFrameData = ProviderSharedPtr->GetCurrentGameFrame().Pin())
 		{
 			CurrentFrameData->RecordedDebugDrawLinesBySolverID.FindOrAdd(DebugDrawData->SolverID).Add(DebugDrawData);
+			CurrentFrameData->MarkDirty();
 		}
 	}
 
diff --git a/Engine/Plugins/ChaosVD/Source/ChaosVD/Private/Trace/DataProcessors/ChaosVDDebugDrawSphereDataProcessor.cpp b/Engine/Plugins/ChaosVD/Source/ChaosVD/Private/Trace/DataProcessors/ChaosVDDebugDrawSphereDataProcessor.cpp
--- a/Engine/Plugins/ChaosVD/Source/ChaosVD/Private/Trace/DataProcessors/ChaosVDDebugDrawSphereDataProcessor.cpp
+++ b/Engine/Plugins/ChaosVD/Source/ChaosVD/Private/Trace/DataProcessors/ChaosVDDebugDrawSphereDataProcessor.cpp
@@ -25,6 +25,7 @@ bool FChaosVDDebugDrawSphereDataProcessor::ProcessRawData(const TArray<uint8>& I
 		if (TSharedPtr<FChaosVDGameFrameData> CurrentFrameData = ProviderSharedPtr->GetCurrentGameFrame().Pin())
 		{
 			CurrentFrameData->RecordedDebugDrawSpheresBySolverID.FindOrAdd(DebugDrawData->SolverID).Add(DebugDrawData);
+			CurrentFrameData->MarkDirty();
 		}
 	}
 
diff --git a/Engine/Plugins/ChaosVD/Source/ChaosVD/Private/Trace/DataProcessors/ChaosVDSceneQueryDataProcessor.cpp b/Engine/Plugins/ChaosVD/Source/ChaosVD/Private/Trace/DataProcessors/ChaosVDSceneQueryDataProcessor.cpp
--- a/Engine/Plugins/ChaosVD/Source/ChaosVD/Private/Trace/DataProcessors/ChaosVDSceneQueryDataProcessor.cpp
+++ b/Engine/Plugins/ChaosVD/Source/ChaosVD/Private/Trace/DataProcessors/ChaosVDSceneQueryDataProcessor.cpp
@@ -30,7 +30,7 @@ bool FChaosVDSceneQueryDataProcessor::ProcessRawData(const TArray<uint8>& InData
 	{
 		if (const TSharedPtr<FChaosVDGameFrameData> CurrentFrameData = ProviderSharedPtr->GetCurrentGameFrame().Pin())
 		{
-			// If ParentQueryID was set, this is a sub query, so find the parent add it to the sub-queries list so we can navigate trough the query "hierarchy" later on
+			// If ParentQueryID was set, this is a sub query, so find the parent add it to the sub-queries list so we can navigate through the query "hierarchy" later on
 			if (QueryData->ParentQueryID != INDEX_NONE)
 			{
 				if (TMap<int32, TSharedPtr<FChaosVDQueryDataWrapper>>* ParentQueryDataByQueryIDPtr = CurrentFrameData->RecordedSceneQueriesBySolverID.Find(QueryData->WorldSolverID))
@@ -44,6 +44,8 @@ bool FChaosVDSceneQueryDataProcessor::ProcessRawData(const TArray<uint8>& InData
 
 			CurrentFrameData->RecordedSceneQueriesByQueryID.Add(QueryData->ID, QueryData);
 			CurrentFrameData->RecordedSceneQueriesBySolverID.FindOrAdd(QueryData->WorldSolverID).Add(QueryData->ID, QueryData);
+
+			CurrentFrameData->MarkDirty();
 		}
 	}
 
diff --git a/Engine/Source/Runtime/Experimental/ChaosVDData/Public/ChaosVDRecording.h b/Engine/Source/Runtime/Experimental/ChaosVDData/Public/ChaosVDRecording.h
--- a/Engine/Source/Runtime/Experimental/ChaosVDData/Public/ChaosVDRecording.h
+++ b/Engine/Source/Runtime/Experimental/ChaosVDData/Public/ChaosVDRecording.h
@@ -171,6 +171,10 @@ PRAGMA_ENABLE_DEPRECATION_WARNINGS
 		return EndTime - StartTime;
 	}
 
+	bool IsDirty() const { return bIsDirty; };
+
+	void MarkDirty() { bIsDirty = true; }
+
 	TMap<FName, FChaosVDTrackedLocation> RecordedNonSolverLocationsByID;
 	TMap<FName, FChaosVDTrackedTransform> RecordedNonSolverTransformsByID;
 	TMap<int32, TMap<int32, TSharedPtr<FChaosVDQueryDataWrapper>>> RecordedSceneQueriesBySolverID;
@@ -185,6 +189,9 @@ PRAGMA_ENABLE_DEPRECATION_WARNINGS
 	TMap<int32,TArray<TSharedPtr<FChaosVDDebugDrawLineDataWrapper>>> RecordedDebugDrawLinesBySolverID;
 	TMap<int32,TArray<TSharedPtr<FChaosVDDebugDrawSphereDataWrapper>>> RecordedDebugDrawSpheresBySolverID;
 	TMap<int32,TArray<TSharedPtr<FChaosVDDebugDrawImplicitObjectDataWrapper>>> RecordedDebugDrawImplicitObjectsBySolverID;
+
+private:
+	bool bIsDirty = false;
 };
 
 /**
