[ChaosVD] Fix for playback controls getting stuck while scrubbing in some situations

A few weeks ago we introduced support for "in-between" solver stages data, for any recorded data that happens outside a pre-defined solver stage. Examples of this are particle creation/destruction events.

That work required to auto-generate a solver stage structure on the fly in CVD when a CVD file is loaded. The problem is that the UI code assumed the number stages to be stable between frames, which is no longer true with this auto-generated solver state system.

This CL just make the UI pass INDEX_NONE as solver stage to the playback controller, so the controller can just play the last available stage in these situations. We are still checking and erroring out if we pass a specific stage number that is not valid.

diff --git a/Engine/Plugins/ChaosVD/Source/ChaosVD/Private/ChaosVDPlaybackController.cpp b/Engine/Plugins/ChaosVD/Source/ChaosVD/Private/ChaosVDPlaybackController.cpp
--- a/Engine/Plugins/ChaosVD/Source/ChaosVD/Private/ChaosVDPlaybackController.cpp
+++ b/Engine/Plugins/ChaosVD/Source/ChaosVD/Private/ChaosVDPlaybackController.cpp
@@ -213,10 +213,18 @@ void FChaosVDPlaybackController::GoToRecordedSolverStage_AssumesLocked(const int
 
 			if (FChaosVDSolverFrameData* SolverFrameData = LoadedRecording->GetSolverFrameData_AssumesLocked(InTrackID, FrameNumber))
 			{
+				int32 TargetStageIndex = StageNumber;
 				// All solver frames need to be played with a valid specified solver stage. If we don't have just early out
-				if (!SolverFrameData->SolverSteps.IsValidIndex(StageNumber))
+				if (!SolverFrameData->SolverSteps.IsValidIndex(TargetStageIndex))
 				{
-					return;
+					if (SolverFrameData->SolverSteps.IsEmpty() || TargetStageIndex != INDEX_NONE)
+					{
+						UE_LOG(LogChaosVDEditor, Error, TEXT("[%s] Invalid solver stage index [%d] at frame [%d] for Solver ID [%d]. We can't continue..."), ANSI_TO_TCHAR(__FUNCTION__), TargetStageIndex, FrameNumber, InTrackID);
+						return;
+					}
+
+					// If we got INDEX_NONE as stage number, we should play the last available stage
+					TargetStageIndex = SolverFrameData->SolverSteps.Num() -1;
 				}
 
 				const int32 FrameDiff = FrameNumber - CurrentTrackInfo->CurrentFrame;
@@ -235,19 +243,22 @@ void FChaosVDPlaybackController::GoToRecordedSolverStage_AssumesLocked(const int
 					PlayFromClosestKeyFrame_AssumesLocked(InTrackID, FrameNumber, *SceneToControlSharedPtr.Get());
 				}
 
-				const int32 StageNumberDiff = StageNumber - CurrentTrackInfo->CurrentStage;
-				const bool bNeedsPlayPreviousSteps = CurrentTrackInfo->CurrentFrame != FrameNumber || StageNumberDiff < 0 || FMath::Abs(StageNumberDiff) > FrameDriftTolerance;
+				
+				const int32 StageNumberDiff = TargetStageIndex - CurrentTrackInfo->CurrentStage;
+				const bool bIsPlayingNewSolverFrame = CurrentTrackInfo->CurrentFrame != FrameNumber;
+
+				const bool bNeedsPlayPreviousSteps = bIsPlayingNewSolverFrame || StageNumberDiff < 0 || FMath::Abs(StageNumberDiff) > FrameDriftTolerance;
 
 				if (Chaos::VisualDebugger::Cvars::bPlayAllPreviousFrameSteps && bNeedsPlayPreviousSteps)
 				{
-					for (int32 StageIndex = 0; StageIndex <= StageNumber; StageIndex++)
+					for (int32 StageIndex = 0; StageIndex <= TargetStageIndex; StageIndex++)
 					{
 						PlaySolverStepData(InTrackID, SceneToControlSharedPtr.ToSharedRef(), *SolverFrameData, StageIndex);
 					}
 				}
 				else
 				{
-					PlaySolverStepData(InTrackID, SceneToControlSharedPtr.ToSharedRef(), *SolverFrameData, StageNumber);
+					PlaySolverStepData(InTrackID, SceneToControlSharedPtr.ToSharedRef(), *SolverFrameData, TargetStageIndex);
 				}
 
 				if (CurrentTrackInfo->CurrentFrame != FrameNumber)
@@ -260,7 +271,7 @@ void FChaosVDPlaybackController::GoToRecordedSolverStage_AssumesLocked(const int
 					SceneToControlSharedPtr->HandleEnterNewSolverFrame(FrameNumber, *SolverFrameData);
 				}
 
-				CurrentTrackInfo->CurrentStage = StageNumber;
+				CurrentTrackInfo->CurrentStage = TargetStageIndex;
 				CurrentTrackInfo->bIsReSimulated = SolverFrameData->bIsResimulated;
 
 				CurrentTrackInfo->CurrentStageNames.Reset();
diff --git a/Engine/Plugins/ChaosVD/Source/ChaosVD/Private/Widgets/SChaosVDSolverPlaybackControls.cpp b/Engine/Plugins/ChaosVD/Source/ChaosVD/Private/Widgets/SChaosVDSolverPlaybackControls.cpp
--- a/Engine/Plugins/ChaosVD/Source/ChaosVD/Private/Widgets/SChaosVDSolverPlaybackControls.cpp
+++ b/Engine/Plugins/ChaosVD/Source/ChaosVD/Private/Widgets/SChaosVDSolverPlaybackControls.cpp
@@ -314,7 +314,7 @@ void SChaosVDSolverPlaybackControls::OnFrameSelectionUpdated(int32 NewFrameIndex
 	if (const TSharedPtr<FChaosVDPlaybackController> PlaybackControllerPtr = PlaybackController.Pin())
 	{
 		// By default we always playback frames at the last recorded stage as that represents the end of frame state
-		const int32 LastStepNumber = GetMaxSolverStage();
+		constexpr int32 LastStepNumber = INDEX_NONE;
 		PlaybackControllerPtr->GoToTrackFrameAndSync(GetInstigatorID(), EChaosVDTrackType::Solver, SolverTrackInfoRef->TrackID, NewFrameIndex, LastStepNumber);
 	}
 }
