[ChaosVD] Baking-in the colorization settings default values

The system CVD has to load defaults from an ini file does not work as intended. For UE 5.5 the colorization settings are just a few values, so I am just backing the new defaults into C++ and will work on a better fix for the next release.

I am also re-adding a visualization flag into CVD's show menu that was removed by mistake.

diff --git a/Engine/Plugins/ChaosVD/Source/ChaosVD/Private/Settings/ChaosVDParticleVisualizationSettings.h b/Engine/Plugins/ChaosVD/Source/ChaosVD/Private/Settings/ChaosVDParticleVisualizationSettings.h
--- a/Engine/Plugins/ChaosVD/Source/ChaosVD/Private/Settings/ChaosVDParticleVisualizationSettings.h
+++ b/Engine/Plugins/ChaosVD/Source/ChaosVD/Private/Settings/ChaosVDParticleVisualizationSettings.h
@@ -54,23 +54,23 @@ struct FChaosDebugDrawColorsByShapeType
 
 	/** Color used for Sphere, Plane, Cube, Capsule, Cylinder, tapered shapes */
 	UPROPERTY(EditAnywhere, Category=DebugDraw)
-	FColor SimpleTypeColor = FColor(0, 255, 0); 
+	FColor SimpleTypeColor = FColor(0, 158, 115); 
 
 	/** Color used for convex shapes */
 	UPROPERTY(EditAnywhere, Category=DebugDraw)
-	FColor ConvexColor = FColor(0, 255, 255);
+	FColor ConvexColor = FColor(240, 228, 66);
 
 	/** Color used for heightfield */
 	UPROPERTY(EditAnywhere, Category=DebugDraw)
-	FColor HeightFieldColor = FColor(0, 0, 255);
+	FColor HeightFieldColor = FColor(86, 180, 233);
 	
 	/** Color used for triangle meshes */
 	UPROPERTY(EditAnywhere, Category=DebugDraw)
-	FColor TriangleMeshColor = FColor(255, 0, 0);
+	FColor TriangleMeshColor = FColor(213, 94, 0);
 
 	/** Color used for triangle LevelSets */
 	UPROPERTY(EditAnywhere, Category=DebugDraw)
-	FColor LevelSetColor = FColor(255, 0, 128);
+	FColor LevelSetColor = FColor(204, 121, 167);
 
 	FColor GetColorFromShapeType(Chaos::EImplicitObjectType ShapeType) const;
 };
@@ -83,11 +83,11 @@ struct FChaosDebugDrawColorsByClientServer
 	
 	/** Color used for server shapes that are not awake or sleeping dynamic */
     UPROPERTY(config, EditAnywhere, Category=DebugDraw)
-    FColor ServerBaseColor = FColor(50, 0, 0);
+    FColor ServerBaseColor = FColor(231, 92, 80);
 
 	/** Color used for server shapes that are not awake or sleeping dynamic */
 	UPROPERTY(config, EditAnywhere, Category=DebugDraw)
-	FColor ClientBaseColor = FColor(0, 0, 50); 
+	FColor ClientBaseColor = FColor(0, 114, 178); 
 
 	FColor GetColorFromState(bool bIsServer, EChaosVDObjectStateType State) const;
 	FColor GetColorAtIntensity(const FColor& InColor, float Intensity) const;
@@ -114,19 +114,19 @@ struct FChaosDebugDrawColorsByState
 
 	/** Color used for dynamic particles */
 	UPROPERTY(config, EditAnywhere, Category=DebugDraw)
-	FColor DynamicColor = FColor(255, 255, 0);
+	FColor DynamicColor = FColor(253, 246, 98);
 	
 	/** Color used for sleeping particles */
 	UPROPERTY(config, EditAnywhere, Category=DebugDraw)
-	FColor SleepingColor = FColor(128, 128, 128);
+	FColor SleepingColor = FColor(231, 92, 80);
 
 	/** Color used for kinematic particles */
 	UPROPERTY(config, EditAnywhere, Category=DebugDraw)
-	FColor KinematicColor = FColor(0, 128, 255);
+	FColor KinematicColor = FColor(0, 114, 178);
 
 	/** Color used for static particles */
 	UPROPERTY(config, EditAnywhere, Category=DebugDraw)
-	FColor StaticColor = FColor(255, 0, 0);
+	FColor StaticColor = FColor(150, 159, 156);
 
 	FColor GetColorFromState(EChaosVDObjectStateType State) const;
 };
@@ -215,7 +215,7 @@ class UChaosVDParticleVisualizationColorSettings : public UChaosVDVisualizationS
 public:
 
 	UPROPERTY(config, EditAnywhere, Category = "Colors Mode")
-	EChaosVDParticleDebugColorMode ParticleColorMode = EChaosVDParticleDebugColorMode::None;
+	EChaosVDParticleDebugColorMode ParticleColorMode = EChaosVDParticleDebugColorMode::State;
 	
 	UPROPERTY(config, EditAnywhere, Category = "Colors By Shape", meta=(EditCondition = "ParticleColorMode == EChaosVDParticleDebugColorMode::ShapeType", EditConditionHides))
 	FChaosDebugDrawColorsByShapeType ColorsByShapeType;
diff --git a/Engine/Plugins/ChaosVD/Source/ChaosVD/Private/Widgets/SChaosVDViewportToolbar.cpp b/Engine/Plugins/ChaosVD/Source/ChaosVD/Private/Widgets/SChaosVDViewportToolbar.cpp
--- a/Engine/Plugins/ChaosVD/Source/ChaosVD/Private/Widgets/SChaosVDViewportToolbar.cpp
+++ b/Engine/Plugins/ChaosVD/Source/ChaosVD/Private/Widgets/SChaosVDViewportToolbar.cpp
@@ -231,6 +231,7 @@ TSharedRef<SWidget> SChaosVDViewportToolbar::GenerateShowMenu() const
 																		.IncludeFlag(FEngineShowFlags::SF_AntiAliasing)
 																		.IncludeFlag(FEngineShowFlags::SF_Grid)
 																		.IncludeFlag(FEngineShowFlags::SF_Translucency)
+																		.IncludeFlag(FEngineShowFlags::SF_MeshEdges)
 																		.IncludeFlag(FEngineShowFlags::SF_HitProxies)
 																		.IncludeFlag(FEngineShowFlags::SF_Fog)
 																		.IncludeFlag(FEngineShowFlags::SF_Pivot);
