[ChaosVD] Adding key shortcuts for playback controls

Now you can use

- CTRL +  Space to Play/Pause
- SHIFT + Space to Stop
- Period to go to Next Frame
- Comma to go to Previous Frame
- SHIFT + Comma to go to Previous Solver Stage
- SHIFT + Period to go to Next Solver Stage

And CTRL + [0-9] To set a solver or game track as active. An active track is the track on which the playback key shortcuts will act.

Lastly, now the tracks UI show a selection dot to indicate which track is active (The UX of this needs to be improved though)

#okforgitgub public

diff --git a/Engine/Plugins/ChaosVD/Source/ChaosVD/Private/ChaosVDCommands.cpp b/Engine/Plugins/ChaosVD/Source/ChaosVD/Private/ChaosVDCommands.cpp
--- a/Engine/Plugins/ChaosVD/Source/ChaosVD/Private/ChaosVDCommands.cpp
+++ b/Engine/Plugins/ChaosVD/Source/ChaosVD/Private/ChaosVDCommands.cpp
@@ -15,6 +15,16 @@ void FChaosVDCommands::RegisterCommands()
 	UI_COMMAND(ToggleFollowSelectedObject, "Follow Selected Object", "Start or Stop following the selected object", EUserInterfaceActionType::ToggleButton, FInputChord(EKeys::F8));
 	UI_COMMAND(OverridePlaybackFrameRate, "Override Recorded Framerate", "When enabled, allows to playback the recording at a fixed framerate", EUserInterfaceActionType::ToggleButton, FInputChord(EModifierKey::Control, EKeys::R));
 	UI_COMMAND(AllowTranslucentSelection, "Allow Translucent Selection", "Allows translucent objects to be selected", EUserInterfaceActionType::ToggleButton, FInputChord(EKeys::T));
+	
+	// Playback Controls commands
+    UI_COMMAND(PlayPauseTrack, "Play/Pause", "Plays or pauses the playback for the current active track", EUserInterfaceActionType::Button, FInputChord(FInputChord(EModifierKey::Control, EKeys::SpaceBar)));
+    UI_COMMAND(StopTrack, "Stop", "Stops the playback for the current active track", EUserInterfaceActionType::Button, FInputChord(EModifierKey::Shift, EKeys::SpaceBar));
+
+    UI_COMMAND(NextFrame, "Next Frame", "Plays the next frame of the current track", EUserInterfaceActionType::Button, FInputChord(EKeys::Period));
+    UI_COMMAND(PrevFrame, "Prev Frame", "Plays the previous frame of the current track", EUserInterfaceActionType::Button, FInputChord(EKeys::Comma));
+
+    UI_COMMAND(NextStage, "Next Stage", "Plays the next solver stage of the current track", EUserInterfaceActionType::Button, FInputChord(FInputChord(EModifierKey::Shift, EKeys::Period)));
+    UI_COMMAND(PrevStage, "Prev Stage", "Plays the previous solver stage of the current track", EUserInterfaceActionType::Button, FInputChord(FInputChord(EModifierKey::Shift, EKeys::EKeys::Comma)));
 }
 
 #undef LOCTEXT_NAMESPACE
diff --git a/Engine/Plugins/ChaosVD/Source/ChaosVD/Private/ChaosVDCommands.h b/Engine/Plugins/ChaosVD/Source/ChaosVD/Private/ChaosVDCommands.h
--- a/Engine/Plugins/ChaosVD/Source/ChaosVD/Private/ChaosVDCommands.h
+++ b/Engine/Plugins/ChaosVD/Source/ChaosVD/Private/ChaosVDCommands.h
@@ -18,4 +18,12 @@ public:
 	TSharedPtr<FUICommandInfo> ToggleFollowSelectedObject;
 	TSharedPtr<FUICommandInfo> OverridePlaybackFrameRate;
 	TSharedPtr<FUICommandInfo> AllowTranslucentSelection;
+
+	// Playback Controls commands
+	TSharedPtr<FUICommandInfo> PlayPauseTrack;
+	TSharedPtr<FUICommandInfo> StopTrack;
+	TSharedPtr<FUICommandInfo> NextFrame;
+	TSharedPtr<FUICommandInfo> PrevFrame;
+	TSharedPtr<FUICommandInfo> NextStage;
+	TSharedPtr<FUICommandInfo> PrevStage;
 };
\ No newline at end of file
diff --git a/Engine/Plugins/ChaosVD/Source/ChaosVD/Private/ChaosVDPlaybackController.cpp b/Engine/Plugins/ChaosVD/Source/ChaosVD/Private/ChaosVDPlaybackController.cpp
--- a/Engine/Plugins/ChaosVD/Source/ChaosVD/Private/ChaosVDPlaybackController.cpp
+++ b/Engine/Plugins/ChaosVD/Source/ChaosVD/Private/ChaosVDPlaybackController.cpp
@@ -33,10 +33,9 @@ bool FChaosVDTrackInfo::AreSameTrack(const TSharedRef<const FChaosVDTrackInfo>&
 	return TrackA->TrackID == TrackB->TrackID && TrackA->TrackType == TrackB->TrackType; 
 }
 
-FChaosVDPlaybackController::FChaosVDPlaybackController(const TWeakPtr<FChaosVDScene>& InSceneToControl)
+FChaosVDPlaybackController::FChaosVDPlaybackController(const TWeakPtr<FChaosVDScene>& InSceneToControl) : ActiveTrack(MakeShared<FChaosVDTrackInfo>())
 {
 	SceneToControl = InSceneToControl;
-
 	RecordingStoppedHandle = FChaosVDRuntimeModule::Get().RegisterRecordingStopCallback(FChaosVDRecordingStateChangedDelegate::FDelegate::CreateRaw(this, &FChaosVDPlaybackController::HandleDisconnectedFromSession));
 }
 
@@ -197,6 +196,15 @@ void FChaosVDPlaybackController::PlaySolverStepData(int32 TrackID, const TShared
 	}
 }
 
+TSharedRef<FChaosVDTrackInfo> FChaosVDPlaybackController::CreateTrackInfo(int32 SlotIndex)
+{
+	TSharedRef<FChaosVDTrackInfo> NewTrack = MakeShared<FChaosVDTrackInfo>();
+
+	NewTrack->TrackSlot = SlotIndex == INDEX_NONE ?  ++LastAssignedTrackSlot : SlotIndex;
+
+	return NewTrack;
+}
+
 void FChaosVDPlaybackController::GoToRecordedSolverStage_AssumesLocked(const int32 InTrackID, const int32 FrameNumber, const int32 StageNumber, FGuid InstigatorID)
 {
 	if (const TSharedPtr<FChaosVDScene> SceneToControlSharedPtr = SceneToControl.Pin())
@@ -815,8 +823,6 @@ void FChaosVDPlaybackController::StopPlayback(const FGuid& InstigatorGUID)
 	
 		GoToTrackFrameAndSync(InstigatorGUID, GameFramesTrack->TrackType, GameFramesTrack->TrackID, FrameNumber, StepNumber);
 	}
-	
-	CurrentPlayingTrack = nullptr;
 
 	VisitAvailableTracks([this](const TSharedPtr<FChaosVDTrackInfo>& TrackInfo)
 	{
@@ -910,28 +916,37 @@ void FChaosVDPlaybackController::UpdateTrackVisibility(EChaosVDTrackType Type, i
 
 void FChaosVDPlaybackController::HandleFramePlaybackControlInput(EChaosVDPlaybackButtonsID ButtonID, const TSharedRef<const FChaosVDTrackInfo>& InTrackInfoRef, FGuid Instigator)
 {
+	if (!IsRecordingLoaded())
+	{
+		return;
+	}
+
+	if (InTrackInfoRef != ActiveTrack)
+	{
+		if (ButtonID == EChaosVDPlaybackButtonsID::Play && ActiveTrack->bIsPlaying)
+		{
+			UE_LOG(LogChaosVDEditor, Warning, TEXT("[%s] Attempted to play when thera was another track playing. This should not happen, this is probably a widget with controls not properly disabled."), ANSI_TO_TCHAR(__FUNCTION__));
+			ActiveTrack->bIsPlaying = false;
+		}
+	}
+
+	// We receive a const ref from the UI as the track info referenced there is read only, but once we are in the controller again we can and want to modify it.
+	ActiveTrack = ConstCastSharedRef<FChaosVDTrackInfo>(InTrackInfoRef);
+
 	switch (ButtonID)
 	{
 	case EChaosVDPlaybackButtonsID::Play:
 		{
-			if (!ensure(!CurrentPlayingTrack))
-			{
-				UE_LOG(LogChaosVDEditor, Warning, TEXT("[%s] Attempted to play when thera was another track playing. This should not happen, this is probably a widget with controls not properly disabled."), ANSI_TO_TCHAR(__FUNCTION__));
-				CurrentPlayingTrack->bIsPlaying = false;
-			}
-
 			bPauseRequested = false;
 
-			// We receive a const ref from the UI as the track info referenced there is read only, but once we are in the controller again we can and want to modify it.
-			CurrentPlayingTrack = ConstCastSharedRef<FChaosVDTrackInfo>(InTrackInfoRef);
-			CurrentPlayingTrack->bIsPlaying = true;
+			ActiveTrack->bIsPlaying = true;
 			break;
 		}
 	case EChaosVDPlaybackButtonsID::Pause:
 		{
-			if (ensure(CurrentPlayingTrack.IsValid()))
+			if (ensure(ActiveTrack->bIsPlaying))
 			{
-				CurrentPlayingTrack->bIsPlaying = false;
+				ActiveTrack->bIsPlaying = false;
 			}
 			else
 			{
@@ -940,8 +955,6 @@ void FChaosVDPlaybackController::HandleFramePlaybackControlInput(EChaosVDPlaybac
 
 			bPauseRequested = true;
 
-			CurrentPlayingTrack = nullptr;
-
 			break;
 		}
 	case EChaosVDPlaybackButtonsID::Stop:
@@ -970,6 +983,19 @@ void FChaosVDPlaybackController::HandleFramePlaybackControlInput(EChaosVDPlaybac
 
 void FChaosVDPlaybackController::HandleFrameStagePlaybackControlInput(EChaosVDPlaybackButtonsID ButtonID, const TSharedRef<const FChaosVDTrackInfo>& InTrackInfoRef, FGuid Instigator)
 {
+	if (!IsRecordingLoaded())
+	{
+		return;
+	}
+
+	if (ActiveTrack->bIsPlaying)
+	{
+		return;
+	}
+
+	// We receive a const ref from the UI as the track info referenced there is read only, but once we are in the controller again we can and want to modify it.
+	ActiveTrack = ConstCastSharedRef<FChaosVDTrackInfo>(InTrackInfoRef);
+
 	switch (ButtonID)
 	{
 	case EChaosVDPlaybackButtonsID::Next:
@@ -998,17 +1024,12 @@ void FChaosVDPlaybackController::HandleFrameStagePlaybackControlInput(EChaosVDPl
 
 void FChaosVDPlaybackController::TickPlayback(float DeltaTime)
 {
-	if (!CurrentPlayingTrack)
-	{
-		return;
-	}
-
-	if (!CurrentPlayingTrack->bIsPlaying)
+	if (!ActiveTrack->IsValidTrack() || !ActiveTrack->bIsPlaying)
 	{
 		return;
 	}
 
-	if (!IsPlayingLiveSession() && CurrentPlayingTrack->CurrentFrame == CurrentPlayingTrack->MaxFrames - 1)
+	if (!IsPlayingLiveSession() && ActiveTrack->CurrentFrame == ActiveTrack->MaxFrames - 1)
 	{
 		StopPlayback(PlaybackSelfInstigatorID);
 		return;
@@ -1016,16 +1037,16 @@ void FChaosVDPlaybackController::TickPlayback(float DeltaTime)
 
 	CurrentPlaybackTime += DeltaTime;
 
-	const float CurrentFrameTime = GetFrameTimeForTrack(CurrentPlayingTrack->TrackType, CurrentPlayingTrack->TrackID, CurrentPlayingTrack.ToSharedRef());
+	const float CurrentFrameTime = GetFrameTimeForTrack(ActiveTrack->TrackType, ActiveTrack->TrackID, ActiveTrack);
 
-	const bool bIsGameTrack = CurrentPlayingTrack->TrackType == EChaosVDTrackType::Game;
+	const bool bIsGameTrack = ActiveTrack->TrackType == EChaosVDTrackType::Game;
 
 	while (CurrentPlaybackTime > CurrentFrameTime)
 	{
 		CurrentPlaybackTime -= CurrentFrameTime;
-		const int32 NextFrame = CurrentPlayingTrack->CurrentFrame + 1;
-		const int32 NextStageNumber =  bIsGameTrack ? 0 : GetTrackLastStageAtFrame_AssumesLocked(CurrentPlayingTrack->TrackType, CurrentPlayingTrack->TrackID, NextFrame);
-		GoToTrackFrameAndSync(PlaybackSelfInstigatorID, CurrentPlayingTrack->TrackType, CurrentPlayingTrack->TrackID, NextFrame, NextStageNumber);
+		const int32 NextFrame = ActiveTrack->CurrentFrame + 1;
+		const int32 NextStageNumber =  bIsGameTrack ? 0 : GetTrackLastStageAtFrame_AssumesLocked(ActiveTrack->TrackType, ActiveTrack->TrackID, NextFrame);
+		GoToTrackFrameAndSync(PlaybackSelfInstigatorID, ActiveTrack->TrackType, ActiveTrack->TrackID, NextFrame, NextStageNumber);
 	}
 }
 
@@ -1120,7 +1141,22 @@ void FChaosVDPlaybackController::ToggleTrackSyncEnabled(const TSharedRef<const F
 
 bool FChaosVDPlaybackController::IsPlaying() const
 {
-	return CurrentPlayingTrack && CurrentPlayingTrack->bIsPlaying;
+	return ActiveTrack->bIsPlaying;
+}
+
+void FChaosVDPlaybackController::TrySetActiveTrack(int32 SlotIndex)
+{
+	VisitAvailableTracks([this, SlotIndex](const TSharedPtr<FChaosVDTrackInfo>& TrackInfo)
+	{
+		if (TrackInfo && TrackInfo->TrackSlot == SlotIndex)
+		{
+			ActiveTrack->bIsPlaying = false;
+			ActiveTrack = TrackInfo.ToSharedRef();
+			return false;
+		}
+
+		return true;
+	});
 }
 
 void FChaosVDPlaybackController::UpdateSolverTracksData()
@@ -1143,7 +1179,7 @@ void FChaosVDPlaybackController::UpdateSolverTracksData()
 
 		if (!SolverTrackInfo.IsValid())
 		{
-			SolverTrackInfo = MakeShared<FChaosVDTrackInfo>();
+			SolverTrackInfo = CreateTrackInfo();
 			SolverTrackInfo->CurrentFrame = 0;
 			SolverTrackInfo->CurrentStage = 0;
 		};
@@ -1170,10 +1206,16 @@ void FChaosVDPlaybackController::HandleCurrentRecordingUpdated()
 	TSharedPtr<FChaosVDTrackInfo>& GameTrackInfo = TrackInfoPerType[EChaosVDTrackType::Game].FindOrAdd(GameTrackID);
 	if (!GameTrackInfo.IsValid())
 	{
-		GameTrackInfo = MakeShared<FChaosVDTrackInfo>();
+		GameTrackInfo = CreateTrackInfo(GameTrackSlot);
 		GameTrackInfo->TrackID = GameTrackID;
 		GameTrackInfo->CurrentFrame = 0;
 		GameTrackInfo->CurrentStage = 0;
+
+		// If no track is active, let's default to the game tracks as it is the only one guaranteed to always exists
+		if (!ActiveTrack->IsValidTrack())
+		{
+			ActiveTrack = GameTrackInfo.ToSharedRef();
+		}
 	}
 
 	GameTrackInfo->MaxFrames = LoadedRecording.IsValid() ? LoadedRecording->GetAvailableGameFrames_AssumesLocked().Num() : INDEX_NONE;
diff --git a/Engine/Plugins/ChaosVD/Source/ChaosVD/Private/ChaosVDPlaybackController.h b/Engine/Plugins/ChaosVD/Source/ChaosVD/Private/ChaosVDPlaybackController.h
--- a/Engine/Plugins/ChaosVD/Source/ChaosVD/Private/ChaosVDPlaybackController.h
+++ b/Engine/Plugins/ChaosVD/Source/ChaosVD/Private/ChaosVDPlaybackController.h
@@ -37,6 +37,7 @@ enum class EChaosVDTrackType : int32
 struct FChaosVDTrackInfo
 {
 	int32 TrackID = INDEX_NONE;
+	int32 TrackSlot = INDEX_NONE;
 	EChaosVDTrackType TrackType = EChaosVDTrackType::Invalid;
 	int32 CurrentFrame = INDEX_NONE;
 	int32 CurrentStage = INDEX_NONE;
@@ -53,6 +54,8 @@ struct FChaosVDTrackInfo
 	bool operator==(const FChaosVDTrackInfo& Other) const;
 
 	static bool AreSameTrack(const TSharedRef<const FChaosVDTrackInfo>& TrackA,  const TSharedRef<const FChaosVDTrackInfo>& TrackB);
+
+	bool IsValidTrack() const { return TrackID != INDEX_NONE; };
 };
 
 struct FChaosVDQueuedTrackInfoUpdate
@@ -93,6 +96,7 @@ public:
 
 	/** ID used for the Game Track */
 	static constexpr int32 GameTrackID  = 0;
+	static constexpr int32 GameTrackSlot  = 0;
 	static constexpr int32 InvalidFrameRateOverride  = -1;
 	static constexpr float FallbackFrameTime = 1.0f / 60.0f;
 
@@ -271,7 +275,7 @@ public:
 
 	void TickPlayback(float DeltaTime);
 
-	TSharedPtr<FChaosVDTrackInfo> GetCurrentPlayingTrackInfo() const { return CurrentPlayingTrack; }
+	TSharedRef<FChaosVDTrackInfo> GetActiveTrackInfo() const { return ActiveTrack; }
 
 	void GetTracksByType(EChaosVDTrackType Type, TArray<TSharedPtr<FChaosVDTrackInfo>>& OutTracks);
 	void SyncTracks(const TSharedRef<const FChaosVDTrackInfo>& FromTrack, EChaosVDSyncTimelinesMode TrackSyncMode = EChaosVDSyncTimelinesMode::RecordedTimestamp);
@@ -284,8 +288,9 @@ public:
 	EChaosVDSyncTimelinesMode GetTimelineSyncMode() const { return CurrentSyncMode; }
 	void SetTimelineSyncMode(EChaosVDSyncTimelinesMode SyncMode) { CurrentSyncMode = SyncMode ; }
 
-protected:
+	void TrySetActiveTrack(int32 SlotIndex);
 
+protected:
 	/** Updates (or adds) solvers data from the loaded recording to the solver tracks */
 	void UpdateSolverTracksData();
 
@@ -312,6 +317,8 @@ protected:
 	template <typename TVisitorCallback>
 	void VisitAvailableTracks(const TVisitorCallback& VisitorCallback);
 
+	TSharedRef<FChaosVDTrackInfo> CreateTrackInfo(int32 SlotIndex = INDEX_NONE);
+
 	/** Map containing all track info, by track type*/
 	TMap<EChaosVDTrackType, TrackInfoByIDMap> TrackInfoPerType;
 
@@ -354,10 +361,13 @@ protected:
 
 	FDelegateHandle RecordingStoppedHandle;
 
-	TSharedPtr<FChaosVDTrackInfo> CurrentPlayingTrack;
+	TSharedRef<FChaosVDTrackInfo> ActiveTrack;
 
 	float CurrentPlaybackTime = 0.0f;
 
+	/** Counter used to create Track Slot IDs. Game track is always Slot 0 */
+	int32 LastAssignedTrackSlot = 0;
+
 	EChaosVDSyncTimelinesMode CurrentSyncMode = EChaosVDSyncTimelinesMode::RecordedTimestamp;
 };
 
diff --git a/Engine/Plugins/ChaosVD/Source/ChaosVD/Private/ChaosVDSolversTracksTab.cpp b/Engine/Plugins/ChaosVD/Source/ChaosVD/Private/ChaosVDSolversTracksTab.cpp
--- a/Engine/Plugins/ChaosVD/Source/ChaosVD/Private/ChaosVDSolversTracksTab.cpp
+++ b/Engine/Plugins/ChaosVD/Source/ChaosVD/Private/ChaosVDSolversTracksTab.cpp
@@ -22,7 +22,7 @@ TSharedRef<SDockTab> FChaosVDSolversTracksTab::HandleTabSpawnRequest(const FSpaw
 	{
 		SolverTracksTab->SetContent
 		(
-			SAssignNew(SolverTracksWidget, SChaosVDSolverTracks, MainTabPtr->GetChaosVDEngineInstance()->GetPlaybackController())
+			SAssignNew(SolverTracksWidget, SChaosVDSolverTracks, MainTabPtr->GetChaosVDEngineInstance()->GetPlaybackController(), OwningTabWidget)
 		);
 	}
 	else
diff --git a/Engine/Plugins/ChaosVD/Source/ChaosVD/Private/Widgets/SChaosVDGameFramesPlaybackControls.cpp b/Engine/Plugins/ChaosVD/Source/ChaosVD/Private/Widgets/SChaosVDGameFramesPlaybackControls.cpp
--- a/Engine/Plugins/ChaosVD/Source/ChaosVD/Private/Widgets/SChaosVDGameFramesPlaybackControls.cpp
+++ b/Engine/Plugins/ChaosVD/Source/ChaosVD/Private/Widgets/SChaosVDGameFramesPlaybackControls.cpp
@@ -93,8 +93,9 @@ bool SChaosVDGameFramesPlaybackControls::CanPlayback() const
 				bIsCompatibleSyncMode &= !bHasTracksWithNetworkSyncData;
 			}
 
-			TSharedPtr<FChaosVDTrackInfo> CurrentTrackBeingPlayed = PlaybackControllerPtr->GetCurrentPlayingTrackInfo();
-			bCanControlPlayback = (!CurrentTrackBeingPlayed || (CurrentTrackBeingPlayed && CurrentTrackBeingPlayed->TrackType == EChaosVDTrackType::Game)) && bIsCompatibleSyncMode;
+			TSharedRef<FChaosVDTrackInfo> ActivePlaybackTrack = PlaybackControllerPtr->GetActiveTrackInfo();
+			bool bActiveTrackIsThisTrack = FChaosVDTrackInfo::AreSameTrack(ActivePlaybackTrack, GameTrackInfoRef);
+			bCanControlPlayback = bIsCompatibleSyncMode && (bActiveTrackIsThisTrack || (!bActiveTrackIsThisTrack && !ActivePlaybackTrack->bIsPlaying));
 		}
 	}
 
diff --git a/Engine/Plugins/ChaosVD/Source/ChaosVD/Private/Widgets/SChaosVDMainTab.cpp b/Engine/Plugins/ChaosVD/Source/ChaosVD/Private/Widgets/SChaosVDMainTab.cpp
--- a/Engine/Plugins/ChaosVD/Source/ChaosVD/Private/Widgets/SChaosVDMainTab.cpp
+++ b/Engine/Plugins/ChaosVD/Source/ChaosVD/Private/Widgets/SChaosVDMainTab.cpp
@@ -74,6 +74,7 @@ void SChaosVDMainTab::Construct(const FArguments& InArgs, TSharedPtr<FChaosVDEng
 {
 	EditorModeTools = MakeShared<FChaosVDEditorModeTools>(InChaosVDEngine->GetCurrentScene());
 
+	GlobalCommandList = MakeShared<FUICommandList>();
 	EditorModeTools->SetToolkitHost(StaticCastSharedRef<SChaosVDMainTab>(AsShared()));
 	ChaosVDEngine = InChaosVDEngine;
 	OwnerTab = InArgs._OwnerTab;
@@ -484,6 +485,66 @@ TSharedRef<IStructureDetailsView> SChaosVDMainTab::CreateStructureDetailsView(co
 	return DetailsView;
 }
 
+void SChaosVDMainTab::ProccessKeyEventForPlaybackTrackSelector(const FKeyEvent& InKeyEvent)
+{
+	if (TSharedPtr<FChaosVDPlaybackController> PlaybackController = GetChaosVDEngineInstance()->GetPlaybackController())
+	{
+		uint32 KeyCode = InKeyEvent.GetKeyCode();
+	
+		int32 TrackSlotIndex = INDEX_NONE;
+
+		constexpr int32 AlphaNumKeyCodeLowerBound = 48;
+		constexpr int32 AlphaNumKeyCodeUpperBound = 57;
+	
+		constexpr int32 NumPadNumberKeyCodeLowerBound = 96;
+		constexpr int32 NumPadNumberKeyCodeUpperBound = 105;
+		
+		if (KeyCode >= AlphaNumKeyCodeLowerBound && KeyCode <= AlphaNumKeyCodeUpperBound)
+		{
+			TrackSlotIndex = KeyCode - AlphaNumKeyCodeUpperBound;
+		}
+
+		if (KeyCode >= NumPadNumberKeyCodeLowerBound && KeyCode <= NumPadNumberKeyCodeUpperBound)
+		{
+			TrackSlotIndex = KeyCode - NumPadNumberKeyCodeLowerBound;
+		}
+
+		if (TrackSlotIndex != INDEX_NONE)
+		{
+			PlaybackController->TrySetActiveTrack(TrackSlotIndex);
+		}
+	}
+}
+
+FReply SChaosVDMainTab::OnKeyDown(const FGeometry& MyGeometry, const FKeyEvent& InKeyEvent)
+{
+	if (InKeyEvent.IsControlDown())
+	{
+		bShowTrackSelectorKeyShortcut = true;
+
+		ProccessKeyEventForPlaybackTrackSelector(InKeyEvent);
+	}
+
+	if (!GlobalCommandList->ProcessCommandBindings(InKeyEvent))
+	{
+		return SCompoundWidget::OnKeyDown(MyGeometry, InKeyEvent);
+	}
+
+	return FReply::Handled();
+}
+
+FReply SChaosVDMainTab::OnKeyUp(const FGeometry& MyGeometry, const FKeyEvent& InKeyEvent)
+{
+	bShowTrackSelectorKeyShortcut = false;
+
+	return SCompoundWidget::OnKeyUp(MyGeometry, InKeyEvent);
+}
+
+bool SChaosVDMainTab::ShouldShowTracksKeyShortcuts()
+{
+	return bShowTrackSelectorKeyShortcut;
+}
+
 void SChaosVDMainTab::SetCustomPropertyLayouts(IDetailsView* DetailsView)
 {
 	if (!DetailsView)
diff --git a/Engine/Plugins/ChaosVD/Source/ChaosVD/Private/Widgets/SChaosVDMainTab.h b/Engine/Plugins/ChaosVD/Source/ChaosVD/Private/Widgets/SChaosVDMainTab.h
--- a/Engine/Plugins/ChaosVD/Source/ChaosVD/Private/Widgets/SChaosVDMainTab.h
+++ b/Engine/Plugins/ChaosVD/Source/ChaosVD/Private/Widgets/SChaosVDMainTab.h
@@ -23,6 +23,7 @@ class FChaosVDPlaybackViewportTab;
 class FChaosVDObjectDetailsTab;
 class FChaosVDWorldOutlinerTab;
 class FStructOnScope;
+class FUICommandList;
 class SButton;
 class SDockTab;
 
@@ -83,6 +84,23 @@ public:
 	TSharedRef<IDetailsView> CreateDetailsView(const FDetailsViewArgs& InDetailsViewArgs);
 
 	TSharedRef<IStructureDetailsView> CreateStructureDetailsView(const FDetailsViewArgs& InDetailsViewArgs, const FStructureDetailsViewArgs& InStructureDetailsViewArgs, const TSharedPtr<FStructOnScope>& InStructData = nullptr, const FText& CustomName = FText::GetEmpty());
+	void ProccessKeyEventForPlaybackTrackSelector(const FKeyEvent& InKeyEvent);
+
+
+	TSharedPtr<FUICommandList> GetGlobalUICommandList() const
+	{
+		return GlobalCommandList;
+	}
+
+	virtual FReply OnKeyDown(const FGeometry& MyGeometry, const FKeyEvent& InKeyEvent) override;
+	virtual FReply OnKeyUp(const FGeometry& MyGeometry, const FKeyEvent& InKeyEvent) override;
+
+	virtual bool SupportsKeyboardFocus() const override
+	{
+		return true;
+	}
+
+	bool ShouldShowTracksKeyShortcuts();
 
 private:
 
@@ -141,6 +159,11 @@ private:
 	FText GetConnectButtonText() const;
 	FText GetConnectButtonTooltipText() const;
 
+	bool bShowTrackSelectorKeyShortcut = false;
+
+	/** Commandlist used for any UI actions that need to be processed globally regardless on which specific widget we are in */
+	TSharedPtr<FUICommandList> GlobalCommandList;
+
 	static inline const FName MainToolBarName = FName("ChaosVD.MainToolBar");
 };
 
diff --git a/Engine/Plugins/ChaosVD/Source/ChaosVD/Private/Widgets/SChaosVDPlaybackViewport.cpp b/Engine/Plugins/ChaosVD/Source/ChaosVD/Private/Widgets/SChaosVDPlaybackViewport.cpp
--- a/Engine/Plugins/ChaosVD/Source/ChaosVD/Private/Widgets/SChaosVDPlaybackViewport.cpp
+++ b/Engine/Plugins/ChaosVD/Source/ChaosVD/Private/Widgets/SChaosVDPlaybackViewport.cpp
@@ -8,10 +8,14 @@
 #include "ChaosVDPlaybackViewportClient.h"
 #include "ChaosVDScene.h"
 #include "EditorModeManager.h"
+#include "EditorViewportCommands.h"
 #include "SChaosVDGameFramesPlaybackControls.h"
+#include "SChaosVDMainTab.h"
 #include "Framework/Application/SlateApplication.h"
 #include "Widgets/SChaosVDTimelineWidget.h"
 #include "Widgets/SChaosVDViewportToolbar.h"
+#include "Widgets/Input/SCheckBox.h"
+#include "Widgets/Layout/SSeparator.h"
 #include "Widgets/Text/STextBlock.h"
 
 #define LOCTEXT_NAMESPACE "ChaosVisualDebugger"
@@ -61,6 +65,30 @@ void SChaosVDPlaybackViewport::Construct(const FArguments& InArgs, TWeakPtr<FCha
 	check(PlaybackViewportClient.IsValid());
 
 	PlaybackViewportClient->SetScene(InScene);
+
+	TAttribute<EVisibility> KeyShortcutVisibilityAttribute;
+	KeyShortcutVisibilityAttribute.BindLambda([WeakThis = AsWeak()]()
+	{
+		if (TSharedPtr<SChaosVDPlaybackViewport> PlaybackViewportWidget = StaticCastSharedPtr<SChaosVDPlaybackViewport>(WeakThis.Pin()))
+		{
+			return PlaybackViewportWidget->GetTrackSelectorKeyVisibility();
+		}
+
+		return EVisibility::Collapsed;
+	});
+
+	TAttribute<ECheckBoxState> GameTrackIsActiveAttribute;
+	GameTrackIsActiveAttribute.BindLambda([WeakThis = AsWeak(), WeakPlaybackController = InPlaybackController]()
+	{
+		TSharedPtr<SChaosVDPlaybackViewport> PlaybackViewportWidget = StaticCastSharedPtr<SChaosVDPlaybackViewport>(WeakThis.Pin());
+		TSharedPtr<FChaosVDPlaybackController> PlaybackController = StaticCastSharedPtr<FChaosVDPlaybackController>(WeakPlaybackController.Pin());
+		if (!PlaybackViewportWidget || !PlaybackController)
+		{
+			return ECheckBoxState::Undetermined;
+		}
+		
+		return PlaybackController->GetActiveTrackInfo()->TrackType == EChaosVDTrackType::Game ? ECheckBoxState::Checked : ECheckBoxState::Unchecked;
+	});
 	
 	if (UChaosVDEditorMode* CVDEdMode = Cast<UChaosVDEditorMode>(EditorModeTools->GetActiveScriptableMode(UChaosVDEditorMode::EM_ChaosVisualDebugger)))
 	{
@@ -89,11 +117,50 @@ void SChaosVDPlaybackViewport::Construct(const FArguments& InArgs, TWeakPtr<FCha
 			SNew(SVerticalBox)
 			+SVerticalBox::Slot()
 			.AutoHeight()
+			.HAlign(HAlign_Center)
 			.Padding(0.0f, 0.0f, 0.0f, 2.0f)
 			[
-				SNew(STextBlock)
-				.Justification(ETextJustify::Center)
-				.Text(LOCTEXT("PlaybackViewportWidgetGameFramesLabel", "Game Frames" ))
+				SNew(SHorizontalBox)
+				+SHorizontalBox::Slot()
+				.HAlign(HAlign_Left)
+				.VAlign(VAlign_Center)
+				.AutoWidth()
+				.Padding(2.0f, 0.0f, 8.0f, 0.0f)
+				[
+					SNew(SCheckBox)
+					.IsEnabled(false)
+					.Style(FAppStyle::Get(), "Menu.RadioButton")
+					.IsChecked(GameTrackIsActiveAttribute)
+				]
+				+SHorizontalBox::Slot()
+				.HAlign(HAlign_Left)
+				.VAlign(VAlign_Center)
+				.AutoWidth()
+				.Padding(0.f, 0.f, 0.f, 0.f)
+				[
+					SNew(STextBlock)
+					.Justification(ETextJustify::Center)
+					.Text(LOCTEXT("PlaybackViewportWidgetGameFramesLabel", "Game Frames" ))
+				]
+				+SHorizontalBox::Slot()
+				.AutoWidth()
+				.Padding(2.0f, 0.0f, 2.0f, 0.0f)
+				[
+					SNew(SSeparator)
+					.Visibility(KeyShortcutVisibilityAttribute)
+					.Orientation(Orient_Vertical)
+					.Thickness(1.f)
+				]
+				+SHorizontalBox::Slot()
+				.HAlign(HAlign_Right)
+				.VAlign(VAlign_Center)
+				.AutoWidth()
+				.Padding(0.f, 0.f, 0.f, 0.f)
+				[
+					SNew(STextBlock)
+					.Visibility(KeyShortcutVisibilityAttribute)
+					.Text(FText::AsCultureInvariant(TEXT("CTRL + 0")))
+				]
 			]
 			+SVerticalBox::Slot()
 			[
@@ -117,10 +184,96 @@ TSharedPtr<FExtender> SChaosVDPlaybackViewport::GetExtenders() const
 	return Extender;
 }
 
+void SChaosVDPlaybackViewport::BindGlobalUICommands()
+{
+	const FChaosVDCommands& Commands = FChaosVDCommands::Get();
+
+	TSharedPtr<SChaosVDMainTab> CVDToolkitHost = StaticCastSharedPtr<SChaosVDMainTab>(EditorModeTools->GetToolkitHost());
+	TSharedPtr<FUICommandList> GlobalUICommandsList = CVDToolkitHost ? CVDToolkitHost->GetGlobalUICommandList() : nullptr;
+
+	if (!GlobalUICommandsList)
+	{
+		return;
+	}
+
+	FUIAction PlayPausePlaybackAction;
+	PlayPausePlaybackAction.ExecuteAction.BindSPLambda(SharedThis(this), [WeakThis = AsWeak()]()
+	{
+		if (TSharedPtr<SChaosVDPlaybackViewport> PlaybackViewportWidget = StaticCastSharedPtr<SChaosVDPlaybackViewport>(WeakThis.Pin()))
+		{
+			if (TSharedPtr<FChaosVDPlaybackController> PlaybackControllerPtr = PlaybackViewportWidget->PlaybackController.Pin())
+			{
+				if (PlaybackControllerPtr->GetActiveTrackInfo()->bIsPlaying)
+				{
+					PlaybackViewportWidget->HandleFramePlaybackControlInput(EChaosVDPlaybackButtonsID::Pause);
+				}
+				else
+				{
+					PlaybackViewportWidget->HandleFramePlaybackControlInput(EChaosVDPlaybackButtonsID::Play);
+				}
+			}
+		}
+	});
+
+	GlobalUICommandsList->MapAction(Commands.PlayPauseTrack, PlayPausePlaybackAction);
+
+	FUIAction StopPlaybackAction;
+	StopPlaybackAction.ExecuteAction.BindSP(SharedThis(this), &SChaosVDPlaybackViewport::HandleFramePlaybackControlInput, EChaosVDPlaybackButtonsID::Stop);
+	GlobalUICommandsList->MapAction(Commands.StopTrack, StopPlaybackAction);
+
+	FUIAction NextFramePlaybackAction;
+	NextFramePlaybackAction.ExecuteAction.BindSP(SharedThis(this), &SChaosVDPlaybackViewport::HandleFramePlaybackControlInput, EChaosVDPlaybackButtonsID::Prev);
+	GlobalUICommandsList->MapAction(Commands.PrevFrame, NextFramePlaybackAction);
+
+	FUIAction PrevFramePlaybackAction;
+	PrevFramePlaybackAction.ExecuteAction.BindSP(SharedThis(this), &SChaosVDPlaybackViewport::HandleFramePlaybackControlInput, EChaosVDPlaybackButtonsID::Next);
+	GlobalUICommandsList->MapAction(Commands.NextFrame, PrevFramePlaybackAction);
+
+	FUIAction NexStagePlaybackAction;
+	NexStagePlaybackAction.ExecuteAction.BindSP(SharedThis(this), &SChaosVDPlaybackViewport::HandleFrameStagePlaybackControlInput, EChaosVDPlaybackButtonsID::Prev);
+	GlobalUICommandsList->MapAction(Commands.PrevStage, NexStagePlaybackAction);
+
+	FUIAction PrevStagePlaybackAction;
+	PrevStagePlaybackAction.ExecuteAction.BindSP(SharedThis(this), &SChaosVDPlaybackViewport::HandleFrameStagePlaybackControlInput, EChaosVDPlaybackButtonsID::Next);
+	GlobalUICommandsList->MapAction(Commands.NextStage, PrevStagePlaybackAction);
+}
+
+void SChaosVDPlaybackViewport::UnBindEditorViewportUnsupportedCommands()
+{
+	const FEditorViewportCommands& DefaultViewportCommands = FEditorViewportCommands::Get();
+
+	if (!CommandList)
+	{
+		return;
+	}
+
+	CommandList->UnmapAction(DefaultViewportCommands.ToggleRealTime);
+	CommandList->UnmapAction(DefaultViewportCommands.ToggleStats);
+	CommandList->UnmapAction(DefaultViewportCommands.ToggleFPS);
+	CommandList->UnmapAction(DefaultViewportCommands.ScreenCaptureForProjectThumbnail);
+	CommandList->UnmapAction(DefaultViewportCommands.RelativeCoordinateSystem_World);
+	CommandList->UnmapAction(DefaultViewportCommands.RelativeCoordinateSystem_Local);
+	CommandList->UnmapAction(DefaultViewportCommands.CycleTransformGizmoCoordSystem);
+	CommandList->UnmapAction(DefaultViewportCommands.ToggleInGameExposure);
+	CommandList->UnmapAction(DefaultViewportCommands.ToggleAutoExposure);
+	CommandList->UnmapAction(DefaultViewportCommands.ToggleInViewportContextMenu);
+	CommandList->UnmapAction(DefaultViewportCommands.ToggleOverrideViewportScreenPercentage);
+	CommandList->UnmapAction(DefaultViewportCommands.OpenEditorPerformanceProjectSettings);
+	CommandList->UnmapAction(DefaultViewportCommands.OpenEditorPerformanceEditorPreferences);
+	CommandList->UnmapAction(DefaultViewportCommands.DetailLightingMode);
+	CommandList->UnmapAction(DefaultViewportCommands.LightingOnlyMode);
+	CommandList->UnmapAction(DefaultViewportCommands.LightComplexityMode);
+	CommandList->UnmapAction(DefaultViewportCommands.ShaderComplexityMode);
+	CommandList->UnmapAction(DefaultViewportCommands.QuadOverdrawMode);
+	CommandList->UnmapAction(DefaultViewportCommands.LightmapDensityMode);
+}
+
 void SChaosVDPlaybackViewport::BindCommands()
 {
 	SEditorViewport::BindCommands();
 
+	UnBindEditorViewportUnsupportedCommands();
+
 	const FChaosVDCommands& Commands = FChaosVDCommands::Get();
 
 	if (ensure(Client))
@@ -153,6 +306,8 @@ void SChaosVDPlaybackViewport::BindCommands()
 		});
 		CommandList->MapAction(Commands.AllowTranslucentSelection, ToggleTranslucentGeometrySelectionAction);
 	}
+
+	BindGlobalUICommands();
 }
 
 EVisibility SChaosVDPlaybackViewport::GetTransformToolbarVisibility() const
@@ -212,6 +367,16 @@ void SChaosVDPlaybackViewport::OnFocusViewportToSelection()
 	PlaybackViewportClient->FocusOnSelectedObject();
 }
 
+EVisibility SChaosVDPlaybackViewport::GetTrackSelectorKeyVisibility() const
+{
+	if (TSharedPtr<SChaosVDMainTab> CVDToolkitHost = StaticCastSharedPtr<SChaosVDMainTab>(EditorModeTools->GetToolkitHost()))
+	{
+		return CVDToolkitHost->ShouldShowTracksKeyShortcuts() ? EVisibility::Visible : EVisibility::Collapsed;
+	}
+
+	return EVisibility::Collapsed;
+}
+
 TSharedRef<FEditorViewportClient> SChaosVDPlaybackViewport::MakeEditorViewportClient()
 {
 	TSharedPtr<FChaosVDPlaybackViewportClient> NewViewport = MakeShared<FChaosVDPlaybackViewportClient>(EditorModeTools, GetViewportWidget());
@@ -266,6 +431,25 @@ void SChaosVDPlaybackViewport::HandlePostSelectionChange(const UTypedElementSele
 	PlaybackViewportClient->bNeedsRedraw = true;
 }
 
+void SChaosVDPlaybackViewport::HandleFramePlaybackControlInput(EChaosVDPlaybackButtonsID ButtonID)
+{
+	if (const TSharedPtr<FChaosVDPlaybackController> PlaybackControllerPtr = PlaybackController.Pin())
+	{
+		PlaybackControllerPtr->HandleFramePlaybackControlInput(ButtonID, PlaybackControllerPtr->GetActiveTrackInfo(), GetInstigatorID());
+	}
+}
+
+void SChaosVDPlaybackViewport::HandleFrameStagePlaybackControlInput(EChaosVDPlaybackButtonsID ButtonID)
+{
+	if (const TSharedPtr<FChaosVDPlaybackController> PlaybackControllerPtr = PlaybackController.Pin())
+	{
+		if (PlaybackControllerPtr->GetActiveTrackInfo()->TrackType == EChaosVDTrackType::Solver)
+		{
+			PlaybackControllerPtr->HandleFrameStagePlaybackControlInput(ButtonID, PlaybackControllerPtr->GetActiveTrackInfo(), GetInstigatorID());
+		}	
+	}
+}
+
 void SChaosVDPlaybackViewport::OnPlaybackSceneUpdated()
 {
 	PlaybackViewportClient->HandleCVDSceneUpdated();
diff --git a/Engine/Plugins/ChaosVD/Source/ChaosVD/Private/Widgets/SChaosVDPlaybackViewport.h b/Engine/Plugins/ChaosVD/Source/ChaosVD/Private/Widgets/SChaosVDPlaybackViewport.h
--- a/Engine/Plugins/ChaosVD/Source/ChaosVD/Private/Widgets/SChaosVDPlaybackViewport.h
+++ b/Engine/Plugins/ChaosVD/Source/ChaosVD/Private/Widgets/SChaosVDPlaybackViewport.h
@@ -2,6 +2,7 @@
 
 #pragma once
 
+#include "ChaosVDCommands.h"
 #include "ChaosVDPlaybackControllerInstigator.h"
 #include "ChaosVDPlaybackControllerObserver.h"
 #include "SCommonEditorViewportToolbarBase.h"
@@ -43,6 +44,10 @@ public:
 
 	virtual void BindCommands() override;
 
+	void BindGlobalUICommands();
+
+	void UnBindEditorViewportUnsupportedCommands();
+
 	virtual EVisibility GetTransformToolbarVisibility() const override;
 
 	void GoToLocation(const FVector& InLocation) const;
@@ -61,6 +66,8 @@ public:
 	virtual void OnFocusViewportToSelection() override;
 
 protected:
+	
+	EVisibility GetTrackSelectorKeyVisibility() const;
 
 	virtual TSharedRef<FEditorViewportClient> MakeEditorViewportClient() override;
 	
@@ -70,6 +77,9 @@ protected:
 	virtual void HandlePlaybackControllerDataUpdated(TWeakPtr<FChaosVDPlaybackController> InController) override;
 	virtual void HandlePostSelectionChange(const UTypedElementSelectionSet* ChangesSelectionSet) override;
 
+	void HandleFramePlaybackControlInput(EChaosVDPlaybackButtonsID ButtonID);
+	void HandleFrameStagePlaybackControlInput(EChaosVDPlaybackButtonsID ButtonID);
+
 	void OnPlaybackSceneUpdated();
 	void OnSolverVisibilityUpdated(int32 SolverID, bool bNewVisibility);
 	void BindToSceneUpdateEvents();
diff --git a/Engine/Plugins/ChaosVD/Source/ChaosVD/Private/Widgets/SChaosVDSolverPlaybackControls.cpp b/Engine/Plugins/ChaosVD/Source/ChaosVD/Private/Widgets/SChaosVDSolverPlaybackControls.cpp
--- a/Engine/Plugins/ChaosVD/Source/ChaosVD/Private/Widgets/SChaosVDSolverPlaybackControls.cpp
+++ b/Engine/Plugins/ChaosVD/Source/ChaosVD/Private/Widgets/SChaosVDSolverPlaybackControls.cpp
@@ -157,8 +157,9 @@ bool SChaosVDSolverPlaybackControls::CanPlayback() const
 				return false;
 			}
 
-			TSharedPtr<FChaosVDTrackInfo> CurrentTrackBeingPlayed = CurrentPlaybackControllerPtr->GetCurrentPlayingTrackInfo();
-			bCanControlPlayback = !CurrentTrackBeingPlayed || (CurrentTrackBeingPlayed && CurrentTrackBeingPlayed->TrackType == EChaosVDTrackType::Solver && CurrentTrackBeingPlayed->TrackID == SolverTrackInfoRef->TrackID);
+			TSharedRef<FChaosVDTrackInfo> ActivePlaybackTrack = CurrentPlaybackControllerPtr->GetActiveTrackInfo();
+			bool bActiveTrackIsThisTrack = FChaosVDTrackInfo::AreSameTrack(ActivePlaybackTrack, SolverTrackInfoRef);
+			bCanControlPlayback =  bActiveTrackIsThisTrack || (!bActiveTrackIsThisTrack && !ActivePlaybackTrack->bIsPlaying);
 		}
 	}
 
diff --git a/Engine/Plugins/ChaosVD/Source/ChaosVD/Private/Widgets/SChaosVDSolverTracks.cpp b/Engine/Plugins/ChaosVD/Source/ChaosVD/Private/Widgets/SChaosVDSolverTracks.cpp
--- a/Engine/Plugins/ChaosVD/Source/ChaosVD/Private/Widgets/SChaosVDSolverTracks.cpp
+++ b/Engine/Plugins/ChaosVD/Source/ChaosVD/Private/Widgets/SChaosVDSolverTracks.cpp
@@ -4,24 +4,29 @@
 
 #include "ChaosVDModule.h"
 #include "ChaosVDPlaybackController.h"
+#include "SChaosVDMainTab.h"
 #include "SChaosVDWarningMessageBox.h"
 #include "Settings/ChaosVDSolverTrackSettings.h"
 #include "ToolMenu.h"
 #include "ToolMenus.h"
 #include "ToolMenuEntry.h"
 #include "ToolMenuSection.h"
+#include "Framework/Application/SlateApplication.h"
 #include "Utils/ChaosVDUserInterfaceUtils.h"
 #include "Widgets/SChaosVDPlaybackViewport.h"
 #include "Widgets/SChaosVDSolverPlaybackControls.h"
+#include "Widgets/Input/SCheckBox.h"
 #include "Widgets/Layout/SExpandableArea.h"
 #include "Widgets/Layout/SSeparator.h"
 
 #define LOCTEXT_NAMESPACE "ChaosVisualDebugger"
 
-void SChaosVDSolverTracks::Construct(const FArguments& InArgs, TWeakPtr<FChaosVDPlaybackController> InPlaybackController)
+void SChaosVDSolverTracks::Construct(const FArguments& InArgs, TWeakPtr<FChaosVDPlaybackController> InPlaybackController, TWeakPtr<SChaosVDMainTab> MainTab)
 {
 	MenuName = FName(TEXT("ChaosVD.SolverTracks.MenuToolbar"));
 
+	MainTabWeakPtr = MainTab;
+
 	ChildSlot
 	[
 		SNew(SVerticalBox)
@@ -61,7 +66,7 @@ void SChaosVDSolverTracks::Construct(const FArguments& InArgs, TWeakPtr<FChaosVD
 	}
 	else
 	{
-		UE_LOG(LogChaosVDEditor, Error, TEXT("[%s] Solver tracks contructed with an invalid player controler. The solver tracks widget will not be functional"), ANSI_TO_TCHAR(__FUNCTION__))	
+		UE_LOG(LogChaosVDEditor, Error, TEXT("[%s] Solver tracks constructed with an invalid player controller. The solver tracks widget will not be functional"), ANSI_TO_TCHAR(__FUNCTION__))	
 	}
 	
 	if (UChaosVDSolverTrackSettings* Settings = FChaosVDSettingsManager::Get().GetSettingsObject<UChaosVDSolverTrackSettings>())
@@ -163,6 +168,25 @@ TSharedRef<ITableRow> SChaosVDSolverTracks::MakeSolverTrackControlsFromTrackInfo
 							.HAlign(HAlign_Left)
 							.VAlign(VAlign_Center)
 							.AutoWidth()
+							.Padding(2.0f, 0.0f, 8.0f, 0.0f)
+							[
+								SNew(SCheckBox)
+								.IsEnabled(false)
+								.Style(FAppStyle::Get(), "Menu.RadioButton")
+								.IsChecked_Lambda([WeakThis = AsWeak(), WeakTrackInfo = TrackInfo.ToWeakPtr()]()
+								{
+									if (TSharedPtr<SChaosVDSolverTracks> SolverTracksWidget = StaticCastSharedPtr<SChaosVDSolverTracks>(WeakThis.Pin()))
+									{
+										return SolverTracksWidget->IsActiveTrack(WeakTrackInfo) ? ECheckBoxState::Checked : ECheckBoxState::Unchecked;
+									}
+
+									return ECheckBoxState::Undetermined;
+								})
+							]
+							+SHorizontalBox::Slot()
+							.HAlign(HAlign_Left)
+							.VAlign(VAlign_Center)
+							.AutoWidth()
 							.Padding(0.f, 0.f, 0.f, 0.f)
 							[
 								SNew(STextBlock)
@@ -197,6 +221,17 @@ TSharedRef<ITableRow> SChaosVDSolverTracks::MakeSolverTrackControlsFromTrackInfo
 								})
 								.WarningText(LOCTEXT("IncomatibleSyncModeWarning", " Incompatible sync mode selected | Attempting to fallback to TimeStamp sync mode for this track | Controls disabled"))
                             ]
+							+SHorizontalBox::Slot()
+							.HAlign(HAlign_Right)
+							.VAlign(VAlign_Center)
+							.AutoWidth()
+							.Padding(0.f, 0.f, 2.f, 0.f)
+							[
+								SNew(STextBlock)
+								.Visibility(this, &SChaosVDSolverTracks::GetSelectorKeyVisibility, TrackInfo->TrackSlot)
+								.Text(FText::FormatOrdered(LOCTEXT("TrackSelectorModifier", "CTRL + {0}"), TrackInfo->TrackSlot))
+								.Font(FCoreStyle::Get().GetFontStyle("ExpandableArea.TitleFont"))
+							]
 						]
 						.BodyContent()
 						[
@@ -226,6 +261,36 @@ TSharedRef<ITableRow> SChaosVDSolverTracks::MakeSolverTrackControlsFromTrackInfo
 		];
 }
 
+EVisibility SChaosVDSolverTracks::GetSelectorKeyVisibility(int32 TrackSlot) const
+{
+	// Currently we only support selecting tracks using CTRL + [0-9]. We never have that many tracks. If at some point we do
+	// Like when we add support for other solve types, we can create another combination (or input chord), and then update this code 
+	constexpr int32 MaxSlotAddressableByKeyboard = 9;
+	if (TrackSlot > MaxSlotAddressableByKeyboard)
+	{
+		return EVisibility::Collapsed;
+	}
+	
+	if (TSharedPtr<SChaosVDMainTab> MainTabPtr = MainTabWeakPtr.Pin())
+	{
+		return MainTabPtr->ShouldShowTracksKeyShortcuts() ? EVisibility::Visible : EVisibility::Collapsed;
+	}
+	
+	return EVisibility::Collapsed;
+}
+
+bool SChaosVDSolverTracks::IsActiveTrack(const TWeakPtr<const FChaosVDTrackInfo>& TrackInfo) const
+{
+	TSharedPtr<const FChaosVDTrackInfo> TrackInfoPtr = TrackInfo.Pin();
+	TSharedPtr<FChaosVDPlaybackController> PlaybackControllerPtr = PlaybackController.Pin();
+	if (!TrackInfoPtr || !PlaybackControllerPtr)
+	{
+		return false;
+	}
+
+	return FChaosVDTrackInfo::AreSameTrack(TrackInfoPtr.ToSharedRef(), PlaybackControllerPtr->GetActiveTrackInfo());
+}
+
 TSharedRef<SWidget> SChaosVDSolverTracks::GenerateToolbarWidget()
 {
 	RegisterMenus();
diff --git a/Engine/Plugins/ChaosVD/Source/ChaosVD/Private/Widgets/SChaosVDSolverTracks.h b/Engine/Plugins/ChaosVD/Source/ChaosVD/Private/Widgets/SChaosVDSolverTracks.h
--- a/Engine/Plugins/ChaosVD/Source/ChaosVD/Private/Widgets/SChaosVDSolverTracks.h
+++ b/Engine/Plugins/ChaosVD/Source/ChaosVD/Private/Widgets/SChaosVDSolverTracks.h
@@ -10,8 +10,10 @@
 
 #include "SChaosVDSolverTracks.generated.h"
 
-struct FChaosVDTrackInfo;
 class FChaosVDPlaybackController;
+class SChaosVDMainTab;
+
+struct FChaosVDTrackInfo;
 
 UCLASS()
 class CHAOSVD_API UChaosVDSolverTracksToolbarMenuContext : public UObject
@@ -33,7 +35,7 @@ public:
 	SLATE_BEGIN_ARGS( SChaosVDSolverTracks ){}
 	SLATE_END_ARGS()
 
-	void Construct(const FArguments& InArgs, TWeakPtr<FChaosVDPlaybackController> InPlaybackController);
+	void Construct(const FArguments& InArgs, TWeakPtr<FChaosVDPlaybackController> InPlaybackController, TWeakPtr<SChaosVDMainTab> MainTab);
 	
 	~SChaosVDSolverTracks();
 
@@ -49,6 +51,10 @@ private:
 
 	TArray<TSharedPtr<const FChaosVDTrackInfo>> CachedTrackInfoArray;
 
+	EVisibility GetSelectorKeyVisibility(int32 TrackSlot) const;
+
+	bool IsActiveTrack(const TWeakPtr<const FChaosVDTrackInfo>& TrackInfo) const;
+
 	TSharedRef<SWidget> GenerateToolbarWidget();
 	TSharedRef<SWidget> GenerateSyncModeMenuWidget();
 
@@ -56,5 +62,9 @@ private:
 
 	void RegisterMenus();
 
+	bool bShowSelectorKeyShortcut = false;
+	
+	TWeakPtr<SChaosVDMainTab> MainTabWeakPtr;
+
 	FName MenuName;
 };
