[ChaosVD] CVD Extension System (Initial Implementation - Part 1)

This CL implements the base of CVD extension system, and moves the Recorded Generic Debug Draw system into a new module meant to contain all the built-in extensions (which can also be used as an example to implement custom extensions in other modules).

Overall all changes

- Moved the minimum required headers to support the Generic Recorded Debug Draw extension
- Exported the minim amount of methods used by [REDACTED_PERSON]
- Updated AChaosVDBaseDataContainer to handle addition of custom data components from CVD Extensions
- Implemented FChaosVDExtension. This will base class for any CVD extension.
- Implemented a simple UStruct based, custom data system. CVD extensions can use this to store they data in the correct frame data, so they don't have to handle storage of the data (which is always an option).

Something to note, in this change, the API for the extension system is not final. It is just the minimum amount of changes needed to move a exiting system into being a CVD extension in another module.

Further refinement will be done in Part 2

diff --git a/Engine/Plugins/ChaosVD/ChaosVD.uplugin b/Engine/Plugins/ChaosVD/ChaosVD.uplugin
--- a/Engine/Plugins/ChaosVD/ChaosVD.uplugin
+++ b/Engine/Plugins/ChaosVD/ChaosVD.uplugin
@@ -25,6 +25,11 @@
 			"Name": "ChaosVDBlueprint",
 			"Type": "RuntimeAndProgram",
 			"LoadingPhase": "Default"
+		},
+		{
+			"Name": "ChaosVDBuiltInExtensions",
+			"Type": "EditorAndProgram",
+			"LoadingPhase": "Default"
 		}
 	],
 	"Plugins": [
diff --git a/Engine/Plugins/ChaosVD/Source/ChaosVD/Private/Actors/ChaosVDDataContainerBaseActor.cpp b/Engine/Plugins/ChaosVD/Source/ChaosVD/Private/Actors/ChaosVDDataContainerBaseActor.cpp
--- a/Engine/Plugins/ChaosVD/Source/ChaosVD/Private/Actors/ChaosVDDataContainerBaseActor.cpp
+++ b/Engine/Plugins/ChaosVD/Source/ChaosVD/Private/Actors/ChaosVDDataContainerBaseActor.cpp
@@ -1,13 +1,21 @@
 // Copyright Epic Games, Inc. All Rights Reserved.
 
-#include "ChaosVDDataContainerBaseActor.h"
+#include "Actors/ChaosVDDataContainerBaseActor.h"
+
+#include "ChaosVDModule.h"
 #include "Components/ChaosVDSolverDataComponent.h"
+#include "ExtensionsSystem/ChaosVDExtensionsManager.h"
 
 AChaosVDDataContainerBaseActor::AChaosVDDataContainerBaseActor()
 {
 	PrimaryActorTick.bCanEverTick = false;
 }
 
+AChaosVDDataContainerBaseActor::~AChaosVDDataContainerBaseActor()
+{
+	FChaosVDExtensionsManager::Get().OnExtensionRegistered().RemoveAll(this);
+}
+
 void AChaosVDDataContainerBaseActor::UpdateFromNewGameFrameData(const FChaosVDGameFrameData& InGameFrameData)
 {
 	TInlineComponentArray<UChaosVDSolverDataComponent*> SolverDataComponents(this);
@@ -20,6 +28,18 @@ void AChaosVDDataContainerBaseActor::UpdateFromNewGameFrameData(const FChaosVDGa
 	}
 }
 
+void AChaosVDDataContainerBaseActor::UpdateFromNewSolverFrameData(const FChaosVDSolverFrameData& InSolverFrameData)
+{
+	TInlineComponentArray<UChaosVDSolverDataComponent*> SolverDataComponents(this);
+	for (UChaosVDSolverDataComponent* Component : SolverDataComponents)
+	{
+		if (Component)
+		{
+			Component->UpdateFromSolverFrameData(InSolverFrameData);
+		}
+	}
+}
+
 void AChaosVDDataContainerBaseActor::Destroyed()
 {
 	TInlineComponentArray<UChaosVDSolverDataComponent*> SolverDataComponents(this);
@@ -34,3 +54,47 @@ void AChaosVDDataContainerBaseActor::Destroyed()
 	Super::Destroyed();
 }
 
+void AChaosVDDataContainerBaseActor::PostActorCreated()
+{
+	Super::PostActorCreated();
+	
+	FChaosVDExtensionsManager::Get().EnumerateExtensions([this](const TSharedRef<FChaosVDExtension>& Extension)
+	{		
+		HandlePostInitializationExtensionRegistered(Extension);
+
+		return true;
+	});
+
+	FChaosVDExtensionsManager::Get().OnExtensionRegistered().AddUObject(this, &AChaosVDDataContainerBaseActor::HandlePostInitializationExtensionRegistered);
+}
+
+void AChaosVDDataContainerBaseActor::HandlePostInitializationExtensionRegistered(const TSharedRef<FChaosVDExtension>& NewExtension)
+{
+	TConstArrayView<TSubclassOf<UActorComponent>> CustomDataComponentsClasses = NewExtension->GetSolverDataComponentsClasses();
+
+	for (const TSubclassOf<UActorComponent>& ComponentClass : CustomDataComponentsClasses)
+	{
+		UActorComponent* DataComponent = NewObject<UActorComponent>(this, ComponentClass);
+		AddOwnedComponent(DataComponent);
+		DataComponent->RegisterComponent();
+
+		if (UChaosVDSolverDataComponent* SolverDataComponent = Cast<UChaosVDSolverDataComponent>(DataComponent))
+		{
+			SolverDataComponent->SetSolverID(SolverDataID);
+		}
+	}
+}
+
+void AChaosVDDataContainerBaseActor::SetSolverID(int32 InSolverID)
+{
+	SolverDataID = InSolverID;
+	
+	TInlineComponentArray<UChaosVDSolverDataComponent*> SolverDataComponents(this);
+	for (UChaosVDSolverDataComponent* Component : SolverDataComponents)
+	{
+		if (Component)
+		{
+			Component->SetSolverID(InSolverID);
+		}
+	}
+}
diff --git a/Engine/Plugins/ChaosVD/Source/ChaosVD/Private/Actors/ChaosVDGameFrameInfoActor.cpp b/Engine/Plugins/ChaosVD/Source/ChaosVD/Private/Actors/ChaosVDGameFrameInfoActor.cpp
--- a/Engine/Plugins/ChaosVD/Source/ChaosVD/Private/Actors/ChaosVDGameFrameInfoActor.cpp
+++ b/Engine/Plugins/ChaosVD/Source/ChaosVD/Private/Actors/ChaosVDGameFrameInfoActor.cpp
@@ -2,23 +2,10 @@
 
 #include "Actors/ChaosVDGameFrameInfoActor.h"
 
-#include "Components/ChaosVDGenericDebugDrawDataComponent.h"
 #include "Components/ChaosVDSceneQueryDataComponent.h"
 
 AChaosVDGameFrameInfoActor::AChaosVDGameFrameInfoActor()
 {
 	PrimaryActorTick.bCanEverTick = false;
 
-	GenericDebugDrawDataComponent = CreateDefaultSubobject<UChaosVDGenericDebugDrawDataComponent>(TEXT("UChaosVDGenericDebugDrawDataComponent"));
-
-	constexpr int32 GenericGameFrameDataSolverID = INDEX_NONE;
-	GenericDebugDrawDataComponent->SetSolverID(GenericGameFrameDataSolverID);
-}
-
-void AChaosVDGameFrameInfoActor::CleanUp()
-{
-	if (GenericDebugDrawDataComponent)
-	{
-		GenericDebugDrawDataComponent->ClearData();
-	}
 }
diff --git a/Engine/Plugins/ChaosVD/Source/ChaosVD/Private/Actors/ChaosVDGameFrameInfoActor.h b/Engine/Plugins/ChaosVD/Source/ChaosVD/Private/Actors/ChaosVDGameFrameInfoActor.h
--- a/Engine/Plugins/ChaosVD/Source/ChaosVD/Private/Actors/ChaosVDGameFrameInfoActor.h
+++ b/Engine/Plugins/ChaosVD/Source/ChaosVD/Private/Actors/ChaosVDGameFrameInfoActor.h
@@ -2,7 +2,7 @@
 
 #pragma once
 
-#include "ChaosVDDataContainerBaseActor.h"
+#include "Actors/ChaosVDDataContainerBaseActor.h"
 #include "GameFramework/Actor.h"
 #include "UObject/ObjectPtr.h"
 
@@ -19,11 +19,4 @@ class AChaosVDGameFrameInfoActor : public AChaosVDDataContainerBaseActor
 
 public:
 	AChaosVDGameFrameInfoActor();
-
-	UChaosVDGenericDebugDrawDataComponent* GetGenericDebugDrawDataComponent() const { return GenericDebugDrawDataComponent.Get(); }
-
-	void CleanUp();
-
-	UPROPERTY()
-	TObjectPtr<UChaosVDGenericDebugDrawDataComponent> GenericDebugDrawDataComponent;
 };
diff --git a/Engine/Plugins/ChaosVD/Source/ChaosVD/Private/Actors/ChaosVDSolverInfoActor.cpp b/Engine/Plugins/ChaosVD/Source/ChaosVD/Private/Actors/ChaosVDSolverInfoActor.cpp
--- a/Engine/Plugins/ChaosVD/Source/ChaosVD/Private/Actors/ChaosVDSolverInfoActor.cpp
+++ b/Engine/Plugins/ChaosVD/Source/ChaosVD/Private/Actors/ChaosVDSolverInfoActor.cpp
@@ -7,8 +7,6 @@
 #include "ChaosVDScene.h"
 #include "ChaosVDSettingsManager.h"
 #include "EditorActorFolders.h"
-#include "Components/ChaosVDGenericDebugDrawDataComponent.h"
-#include "Components/ChaosVDGTAccelerationStructuresDataComponent.h"
 #include "Components/ChaosVDParticleDataComponent.h"
 #include "Components/ChaosVDSceneQueryDataComponent.h"
 #include "Components/ChaosVDSolverCharacterGroundConstraintDataComponent.h"
@@ -27,9 +25,8 @@ AChaosVDSolverInfoActor::AChaosVDSolverInfoActor()
 	ParticleDataComponent = CreateDefaultSubobject<UChaosVDParticleDataComponent>(TEXT("ParticleCollisionDataComponent"));
 	JointsDataComponent = CreateDefaultSubobject<UChaosVDSolverJointConstraintDataComponent>(TEXT("JointDataComponent"));
 	CharacterGroundConstraintDataComponent = CreateDefaultSubobject<UChaosVDSolverCharacterGroundConstraintDataComponent>(TEXT("CharacterGroundConstraintDataComponent"));
-	GTAccelerationStructuresDataComponent = CreateDefaultSubobject<UChaosVDGTAccelerationStructuresDataComponent>(TEXT("GTAccelerationStructuresDataComponent"));
 	SceneQueryDataComponent = CreateDefaultSubobject<UChaosVDSceneQueryDataComponent>(TEXT("ChaosVDSceneQueryDataComponent"));
-	GenericDebugDrawDataComponent = CreateDefaultSubobject<UChaosVDGenericDebugDrawDataComponent>(TEXT("UChaosVDGenericDebugDrawDataComponent"));
+
 	bIsServer = false;
 
 	if (UChaosVDParticleVisualizationSettings* ParticleVisualizationSettings = FChaosVDSettingsManager::Get().GetSettingsObject<UChaosVDParticleVisualizationSettings>())
@@ -43,20 +40,6 @@ AChaosVDSolverInfoActor::AChaosVDSolverInfoActor()
 	}
 }
 
-void AChaosVDSolverInfoActor::SetSolverID(int32 InSolverID)
-{
-	SolverID = InSolverID;
-	
-	TInlineComponentArray<UChaosVDSolverDataComponent*> SolverDataComponents(this);
-	for (UChaosVDSolverDataComponent* Component : SolverDataComponents)
-	{
-		if (Component)
-		{
-			Component->SetSolverID(InSolverID);
-		}
-	}
-}
-
 void AChaosVDSolverInfoActor::SetSolverName(const FName& InSolverName)
 {
 	SolverName = InSolverName;
diff --git a/Engine/Plugins/ChaosVD/Source/ChaosVD/Private/Actors/ChaosVDSolverInfoActor.h b/Engine/Plugins/ChaosVD/Source/ChaosVD/Private/Actors/ChaosVDSolverInfoActor.h
--- a/Engine/Plugins/ChaosVD/Source/ChaosVD/Private/Actors/ChaosVDSolverInfoActor.h
+++ b/Engine/Plugins/ChaosVD/Source/ChaosVD/Private/Actors/ChaosVDSolverInfoActor.h
@@ -2,7 +2,7 @@
 
 #pragma once
 
-#include "ChaosVDDataContainerBaseActor.h"
+#include "Actors/ChaosVDDataContainerBaseActor.h"
 #include "ChaosVDParticleActor.h"
 #include "ChaosVDSceneSelectionObserver.h"
 #include "GameFramework/Actor.h"
@@ -13,7 +13,6 @@ class UChaosVDSolverCharacterGroundConstraintDataComponent;
 class UChaosVDSolverJointConstraintDataComponent;
 struct FChaosVDParticleDataWrapper;
 class AChaosVDParticleActor;
-class UChaosVDGenericDebugDrawDataComponent;
 class UChaosVDParticleDataComponent;
 class UChaosVDSceneQueryDataComponent;
 class UChaosVDSolverCollisionDataComponent;
@@ -31,7 +30,6 @@ public:
 
 	AChaosVDSolverInfoActor();
 
-	void SetSolverID(int32 InSolverID);
 	int32 GetSolverID() const { return SolverID; }
 
 	void SetSolverName(const FName& InSolverName);
@@ -43,15 +41,13 @@ public:
 	virtual void SetScene(TWeakPtr<FChaosVDScene> InScene) override;
 
 	void SetSimulationTransform(const FTransform& InSimulationTransform) { SimulationTransform = InSimulationTransform; }
-	const FTransform& GetSimulationTransform() const { return SimulationTransform; }
+	virtual const FTransform& GetSimulationTransform() const override { return SimulationTransform; }
 
 	UChaosVDSolverCollisionDataComponent* GetCollisionDataComponent() { return CollisionDataComponent; }
 	UChaosVDParticleDataComponent* GetParticleDataComponent() { return ParticleDataComponent; }
 	UChaosVDSolverJointConstraintDataComponent* GetJointsDataComponent() { return JointsDataComponent; }
 	UChaosVDSolverCharacterGroundConstraintDataComponent* GetCharacterGroundConstraintDataComponent() { return CharacterGroundConstraintDataComponent; }
-	UChaosVDGTAccelerationStructuresDataComponent* GetGTAccelerationStructuresDataComponent() { return GTAccelerationStructuresDataComponent; }
 	UChaosVDSceneQueryDataComponent* GetSceneQueryDataComponent() const { return SceneQueryDataComponent.Get(); }
-	UChaosVDGenericDebugDrawDataComponent* GetGenericDebugDrawDataComponent() const { return GenericDebugDrawDataComponent.Get(); }
 
 	void RegisterParticleActor(int32 ParticleID, AChaosVDParticleActor* ParticleActor);
 
@@ -122,14 +118,8 @@ protected:
 	UPROPERTY()
 	TObjectPtr<UChaosVDSolverCharacterGroundConstraintDataComponent> CharacterGroundConstraintDataComponent;
 
-	UPROPERTY()
-	TObjectPtr<UChaosVDGTAccelerationStructuresDataComponent> GTAccelerationStructuresDataComponent;
-
 	UPROPERTY()
 	TObjectPtr<UChaosVDSceneQueryDataComponent> SceneQueryDataComponent;
-	
-	UPROPERTY()
-	TObjectPtr<UChaosVDGenericDebugDrawDataComponent> GenericDebugDrawDataComponent;
 };
 
 template <typename TCallback>
diff --git a/Engine/Plugins/ChaosVD/Source/ChaosVD/Private/ChaosVDGeometryBuilder.cpp b/Engine/Plugins/ChaosVD/Source/ChaosVD/Private/ChaosVDGeometryBuilder.cpp
--- a/Engine/Plugins/ChaosVD/Source/ChaosVD/Private/ChaosVDGeometryBuilder.cpp
+++ b/Engine/Plugins/ChaosVD/Source/ChaosVD/Private/ChaosVDGeometryBuilder.cpp
@@ -249,7 +249,7 @@ void FChaosVDGeometryBuilder::Initialize(const TWeakPtr<FChaosVDScene>& ChaosVDS
 	};
 
 	MeshComponentsWaitingForGeometry = MakeUnique<FObjectsWaitingGeometryList<FMeshComponentWeakPtr>>(ProcessMeshComponent, NSLOCTEXT("ChaosVisualDebugger", "GeometryGenNotification","Mesh Components"), ShouldProcessObjectsForKey);
-	MeshComponentsWaitingForMaterial = MakeUnique<FObjectsWaitingProcessingQueue<FMeshComponentWeakPtr>>(UpdateMeshMaterialForComponent, NSLOCTEXT("ChaosVisualDebugger", "GeometryMaterialNotification","Material instances"));
+	MeshComponentsWaitingForMaterial = MakeUnique<FObjectsWaitingProcessingQueue<FMeshComponentWeakPtr>>(UpdateMeshMaterialForComponent, NSLOCTEXT("ChaosVisualDebugger", "GeometryMaterialNotification","Component Materials"));
 	GeometryTasksPendingLaunch = MakeUnique<FObjectsWaitingProcessingQueue<TSharedPtr<FChaosVDGeometryGenerationTask>>>(LaunchGeometryGenerationTaskDeferred, NSLOCTEXT("ChaosVisualDebugger", "GeometryTaskLauchNotification","Static Meshes"));
 
 	GameThreadTickDelegate = FTSTicker::GetCoreTicker().AddTicker(FTickerDelegate::CreateRaw(this, &FChaosVDGeometryBuilder::GameThreadTick));
@@ -597,7 +597,7 @@ void FChaosVDGeometryBuilder::DestroyMeshComponent(UMeshComponent* MeshComponent
 	MeshComponentsPendingDisposal.Add(MeshComponent);
 }
 
-void FChaosVDGeometryBuilder::RequestMaterialInstance(UMeshComponent* MeshComponent)
+void FChaosVDGeometryBuilder::RequestMaterialUpdate(UMeshComponent* MeshComponent)
 {
 	MeshComponentsWaitingForMaterial->EnqueueObject(MeshComponent);
 }
diff --git a/Engine/Plugins/ChaosVD/Source/ChaosVD/Private/ChaosVDModule.cpp b/Engine/Plugins/ChaosVD/Source/ChaosVD/Private/ChaosVDModule.cpp
--- a/Engine/Plugins/ChaosVD/Source/ChaosVD/Private/ChaosVDModule.cpp
+++ b/Engine/Plugins/ChaosVD/Source/ChaosVD/Private/ChaosVDModule.cpp
@@ -16,6 +16,7 @@
 #include "WorkspaceMenuStructure.h"
 #include "WorkspaceMenuStructureModule.h"
 #include "DetailsCustomizations/ChaosVDShapeDataCustomization.h"
+#include "ExtensionsSystem/ChaosVDExtensionsManager.h"
 
 #define LOCTEXT_NAMESPACE "ChaosVisualDebugger"
 
@@ -82,6 +83,8 @@ void FChaosVDModule::ShutdownModule()
 	FCoreDelegates::OnEnginePreExit.RemoveAll(this);
 
 	CloseActiveInstances();
+
+	FChaosVDExtensionsManager::TearDown();
 	
 	FChaosVDSettingsManager::TearDown();
 }
diff --git a/Engine/Plugins/ChaosVD/Source/ChaosVD/Private/ChaosVDPlaybackController.cpp b/Engine/Plugins/ChaosVD/Source/ChaosVD/Private/ChaosVDPlaybackController.cpp
--- a/Engine/Plugins/ChaosVD/Source/ChaosVD/Private/ChaosVDPlaybackController.cpp
+++ b/Engine/Plugins/ChaosVD/Source/ChaosVD/Private/ChaosVDPlaybackController.cpp
@@ -7,6 +7,7 @@
 #include "ChaosVDRecording.h"
 #include "ChaosVDRuntimeModule.h"
 #include "ChaosVDScene.h"
+#include "ExtensionsSystem/ChaosVDExtensionsManager.h"
 #include "Misc/MessageDialog.h"
 #include "Trace/ChaosVDTraceManager.h"
 #include "Trace/ChaosVDTraceProvider.h"
@@ -85,7 +86,7 @@ bool FChaosVDPlaybackController::LoadChaosVDRecordingFromTraceSession(const FCha
 	
 	if (TSharedPtr<FChaosVDScene> ScenePtr = SceneToControl.Pin())
 	{
-		ScenePtr->LoadedRecording = LoadedRecording;
+		ScenePtr->SetLoadedRecording(LoadedRecording);
 	}
 	
 	bHasPendingGTUpdateBroadcast = true;
@@ -150,7 +151,7 @@ void FChaosVDPlaybackController::PlayFromClosestKeyFrame_AssumesLocked(const int
 		constexpr bool bRequestingKeyFrameOnly = true;
 		if (const FChaosVDSolverFrameData* SolverFrameData = LoadedRecording->GetSolverFrameData_AssumesLocked(InTrackID, FrameNumber, bRequestingKeyFrameOnly))
 		{
-			PlaySolverStepData(InTrackID, InSceneToControl.AsShared(), *SolverFrameData, SolverStage);
+			PlaySolverStageData(InTrackID, InSceneToControl.AsShared(), *SolverFrameData, SolverStage);
 		}
 		else
 		{
@@ -165,7 +166,7 @@ void FChaosVDPlaybackController::PlayFromClosestKeyFrame_AssumesLocked(const int
 	FChaosVDSolverFrameData CollapsedFrameData;
 	LoadedRecording->CollapseSolverFramesRange_AssumesLocked(InTrackID, KeyFrameNumber, FrameNumber, CollapsedFrameData);
 
-	PlaySolverStepData(InTrackID, InSceneToControl.AsShared(), CollapsedFrameData, SolverStage);
+	PlaySolverStageData(InTrackID, InSceneToControl.AsShared(), CollapsedFrameData, SolverStage);
 }
 
 void FChaosVDPlaybackController::EnqueueTrackInfoUpdate(const TSharedRef<const FChaosVDTrackInfo>& InTrackInfo, FGuid InstigatorID)
@@ -183,16 +184,16 @@ void FChaosVDPlaybackController::EnqueueGeometryDataUpdate(const Chaos::FConstIm
 	GeometryDataUpdateGTQueue.Enqueue({NewGeometry, GeometryID });
 }
 
-void FChaosVDPlaybackController::PlaySolverStepData(int32 TrackID, const TSharedRef<FChaosVDScene>& InSceneToControlSharedPtr, const FChaosVDSolverFrameData& InSolverFrameData, int32 StepIndex)
+void FChaosVDPlaybackController::PlaySolverStageData(int32 TrackID, const TSharedRef<FChaosVDScene>& InSceneToControlSharedPtr, const FChaosVDSolverFrameData& InSolverFrameData, int32 StageIndex)
 {
-	if (InSolverFrameData.SolverSteps.IsValidIndex(StepIndex))
+	if (InSolverFrameData.SolverSteps.IsValidIndex(StageIndex))
 	{
-		InSceneToControlSharedPtr->UpdateFromRecordedStepData(TrackID, InSolverFrameData.SolverSteps[StepIndex], InSolverFrameData);
+		InSceneToControlSharedPtr->UpdateFromRecordedStepData(TrackID, InSolverFrameData.SolverSteps[StageIndex], InSolverFrameData);
 	}
 	else
 	{
 		// This is common if we stop PIE, change worlds, and PIE again without stopping the recording
-		UE_LOG(LogChaosVDEditor, Verbose, TEXT("[%s] Tried to scrub to an invalid step | Step Number [%d] ..."), ANSI_TO_TCHAR(__FUNCTION__), StepIndex);
+		UE_LOG(LogChaosVDEditor, Verbose, TEXT("[%s] Tried to scrub to an invalid step | Step Number [%d] ..."), ANSI_TO_TCHAR(__FUNCTION__), StageIndex);
 	}
 }
 
@@ -261,12 +262,12 @@ void FChaosVDPlaybackController::GoToRecordedSolverStage_AssumesLocked(const int
 				{
 					for (int32 StageIndex = 0; StageIndex <= TargetStageIndex; StageIndex++)
 					{
-						PlaySolverStepData(InTrackID, SceneToControlSharedPtr.ToSharedRef(), *SolverFrameData, StageIndex);
+						PlaySolverStageData(InTrackID, SceneToControlSharedPtr.ToSharedRef(), *SolverFrameData, StageIndex);
 					}
 				}
 				else
 				{
-					PlaySolverStepData(InTrackID, SceneToControlSharedPtr.ToSharedRef(), *SolverFrameData, TargetStageIndex);
+					PlaySolverStageData(InTrackID, SceneToControlSharedPtr.ToSharedRef(), *SolverFrameData, TargetStageIndex);
 				}
 
 				if (CurrentTrackInfo->CurrentFrame != FrameNumber)
@@ -283,7 +284,7 @@ void FChaosVDPlaybackController::GoToRecordedSolverStage_AssumesLocked(const int
 				CurrentTrackInfo->bIsReSimulated = SolverFrameData->bIsResimulated;
 
 				CurrentTrackInfo->CurrentStageNames.Reset();
-				Algo::Transform(SolverFrameData->SolverSteps, CurrentTrackInfo->CurrentStageNames, &FChaosVDStepData::StepName);
+				Algo::Transform(SolverFrameData->SolverSteps, CurrentTrackInfo->CurrentStageNames, &FChaosVDFrameStageData::StepName);
 			
 				EnqueueTrackInfoUpdate(CurrentTrackInfo.ToSharedRef(), InstigatorID);
 			}
@@ -383,7 +384,7 @@ void FChaosVDPlaybackController::GoToTrackFrame_AssumesLockedAndSync(FGuid Insti
 	}
 }
 
-int32 FChaosVDPlaybackController::GetTrackStepsNumberAtFrame_AssumesLocked(EChaosVDTrackType TrackType, const int32 InTrackID, const int32 FrameNumber) const
+int32 FChaosVDPlaybackController::GetTrackStagesNumberAtFrame_AssumesLocked(EChaosVDTrackType TrackType, const int32 InTrackID, const int32 FrameNumber) const
 {
 	if (!LoadedRecording.IsValid())
 	{
@@ -414,7 +415,7 @@ int32 FChaosVDPlaybackController::GetTrackStepsNumberAtFrame_AssumesLocked(EChao
 	}
 }
 
-const FChaosVDStepsContainer* FChaosVDPlaybackController::GetTrackStepsDataAtFrame_AssumesLocked(EChaosVDTrackType TrackType, int32 InTrackID, int32 FrameNumber) const
+const FChaosVDFrameStagesContainer* FChaosVDPlaybackController::GetTrackStagesDataAtFrame_AssumesLocked(EChaosVDTrackType TrackType, int32 InTrackID, int32 FrameNumber) const
 {
 	if (!LoadedRecording.IsValid())
 	{
@@ -563,7 +564,7 @@ int32 FChaosVDPlaybackController::GetTrackCurrentFrame(EChaosVDTrackType TrackTy
 	return INDEX_NONE;
 }
 
-int32 FChaosVDPlaybackController::GetTrackCurrentStep(EChaosVDTrackType TrackType, const int32 InTrackID) const
+int32 FChaosVDPlaybackController::GetTrackCurrentStage(EChaosVDTrackType TrackType, const int32 InTrackID) const
 {
 	if (const TrackInfoByIDMap* TrackInfoByID = TrackInfoPerType.Find(TrackType))
 	{
@@ -589,7 +590,7 @@ int32 FChaosVDPlaybackController::GetTrackLastStageAtFrame_AssumesLocked(EChaosV
 	{
 		case EChaosVDTrackType::Solver:
 		{
-			const int32 AvailableSteps = GetTrackStepsNumberAtFrame_AssumesLocked(EChaosVDTrackType::Solver, InTrackID, InFrameNumber);
+			const int32 AvailableSteps = GetTrackStagesNumberAtFrame_AssumesLocked(EChaosVDTrackType::Solver, InTrackID, InFrameNumber);
 			return AvailableSteps == INDEX_NONE ? INDEX_NONE: AvailableSteps -1;
 			break;
 		}
@@ -619,19 +620,19 @@ TSharedPtr<FChaosVDTrackInfo> FChaosVDPlaybackController::GetMutableTrackInfo(EC
 	return nullptr;
 }
 
-void FChaosVDPlaybackController::LockTrackInCurrentStep(EChaosVDTrackType TrackType, int32 TrackID)
+void FChaosVDPlaybackController::LockTrackInCurrentStage(EChaosVDTrackType TrackType, int32 TrackID)
 {
 	if (TSharedPtr<FChaosVDTrackInfo> TrackInfo = GetMutableTrackInfo(TrackType, TrackID))
 	{
-		TrackInfo->LockedOnStep = TrackInfo->CurrentStage;
+		TrackInfo->LockedOnStage = TrackInfo->CurrentStage;
 	}
 }
 
-void FChaosVDPlaybackController::UnlockTrackStep(EChaosVDTrackType TrackType, int32 TrackID)
+void FChaosVDPlaybackController::UnlockTrackStage(EChaosVDTrackType TrackType, int32 TrackID)
 {
 	if (TSharedPtr<FChaosVDTrackInfo> TrackInfo = GetMutableTrackInfo(TrackType, TrackID))
 	{
-		TrackInfo->LockedOnStep = INDEX_NONE;
+		TrackInfo->LockedOnStage = INDEX_NONE;
 	}
 }
 
@@ -741,6 +742,12 @@ bool FChaosVDPlaybackController::Tick(float DeltaTime)
 
 	if (bHasPendingGTUpdateBroadcast)
 	{
+		FChaosVDExtensionsManager::Get().EnumerateExtensions([ThisWeakPtr](const TSharedRef<FChaosVDExtension>& Extension)
+		{
+			Extension->HandlePlaybackControllerDataUpdated(ThisWeakPtr);
+			return true;
+		});
+
 		ControllerUpdatedDelegate.Broadcast(ThisWeakPtr);
 		bHasPendingGTUpdateBroadcast = false;
 	}
@@ -751,6 +758,12 @@ bool FChaosVDPlaybackController::Tick(float DeltaTime)
 		while (TrackInfoUpdateGTQueue.Dequeue(TrackInfoUpdate))
 		{
 			OnTrackFrameUpdated().Broadcast(ThisWeakPtr, TrackInfoUpdate.TrackInfo, TrackInfoUpdate.InstigatorID);
+
+			FChaosVDExtensionsManager::Get().EnumerateExtensions([ThisWeakPtr, &TrackInfoUpdate](const TSharedRef<FChaosVDExtension>& Extension)
+			{
+				Extension->HandleControllerTrackFrameUpdated(ThisWeakPtr, TrackInfoUpdate.TrackInfo, TrackInfoUpdate.InstigatorID);
+				return true;
+			});
 		}
 	}
 
diff --git a/Engine/Plugins/ChaosVD/Source/ChaosVD/Private/ChaosVDPlaybackViewportClient.cpp b/Engine/Plugins/ChaosVD/Source/ChaosVD/Private/ChaosVDPlaybackViewportClient.cpp
--- a/Engine/Plugins/ChaosVD/Source/ChaosVD/Private/ChaosVDPlaybackViewportClient.cpp
+++ b/Engine/Plugins/ChaosVD/Source/ChaosVD/Private/ChaosVDPlaybackViewportClient.cpp
@@ -16,7 +16,6 @@
 #include "UnrealWidget.h"
 #include "Actors/ChaosVDGameFrameInfoActor.h"
 #include "Actors/ChaosVDSolverInfoActor.h"
-#include "Components/ChaosVDGenericDebugDrawDataComponent.h"
 #include "Components/InstancedStaticMeshComponent.h"
 #include "Visualizers/ChaosVDDebugDrawUtils.h"
 #include "Widgets/SChaosVDMainTab.h"
diff --git a/Engine/Plugins/ChaosVD/Source/ChaosVD/Private/ChaosVDScene.cpp b/Engine/Plugins/ChaosVD/Source/ChaosVD/Private/ChaosVDScene.cpp
--- a/Engine/Plugins/ChaosVD/Source/ChaosVD/Private/ChaosVDScene.cpp
+++ b/Engine/Plugins/ChaosVD/Source/ChaosVD/Private/ChaosVDScene.cpp
@@ -143,7 +143,7 @@ void FChaosVDScene::AddReferencedObjects(FReferenceCollector& Collector)
 	Collector.AddStableReferenceArray(&AvailableDataContainerActors);
 }
 
-void FChaosVDScene::UpdateFromRecordedStepData(const int32 SolverID, const FChaosVDStepData& InRecordedStepData, const FChaosVDSolverFrameData& InFrameData)
+void FChaosVDScene::UpdateFromRecordedStepData(const int32 SolverID, const FChaosVDFrameStageData& InRecordedStepData, const FChaosVDSolverFrameData& InFrameData)
 {
 	AChaosVDSolverInfoActor* SolverSceneData = nullptr;
 	if (AChaosVDSolverInfoActor** SolverSceneDataPtrPtr = SolverDataContainerBySolverID.Find(SolverID))
@@ -247,7 +247,7 @@ void FChaosVDScene::UpdateFromRecordedStepData(const int32 SolverID, const FChao
 	OnSceneUpdated().Broadcast();
 }
 
-void FChaosVDScene::UpdateParticlesCollisionData(const FChaosVDStepData& InRecordedStepData, int32 SolverID)
+void FChaosVDScene::UpdateParticlesCollisionData(const FChaosVDFrameStageData& InRecordedStepData, int32 SolverID)
 {
 	if (AChaosVDSolverInfoActor* SolverDataInfoContainer = SolverDataContainerBySolverID.FindChecked(SolverID))
 	{
@@ -258,7 +258,7 @@ void FChaosVDScene::UpdateParticlesCollisionData(const FChaosVDStepData& InRecor
 	}
 }
 
-void FChaosVDScene::UpdateJointConstraintsData(const FChaosVDStepData& InRecordedStepData, int32 SolverID)
+void FChaosVDScene::UpdateJointConstraintsData(const FChaosVDFrameStageData& InRecordedStepData, int32 SolverID)
 {
 	if (AChaosVDSolverInfoActor* SolverDataInfoContainer = SolverDataContainerBySolverID.FindChecked(SolverID))
 	{
diff --git a/Engine/Plugins/ChaosVD/Source/ChaosVD/Private/ChaosVDTabSpawnerBase.cpp b/Engine/Plugins/ChaosVD/Source/ChaosVD/Private/ChaosVDTabSpawnerBase.cpp
--- a/Engine/Plugins/ChaosVD/Source/ChaosVD/Private/ChaosVDTabSpawnerBase.cpp
+++ b/Engine/Plugins/ChaosVD/Source/ChaosVD/Private/ChaosVDTabSpawnerBase.cpp
@@ -11,6 +11,7 @@
 
 FChaosVDTabSpawnerBase::FChaosVDTabSpawnerBase(const FName& InTabID, TSharedPtr<FTabManager> InTabManager, TWeakPtr<SChaosVDMainTab> InOwningTabWidget)
 {
+	TabID = InTabID;
 	OwningTabWidget = InOwningTabWidget;
 	InTabManager->RegisterTabSpawner(InTabID, FOnSpawnTab::CreateRaw(this, &FChaosVDTabSpawnerBase::HandleTabSpawnRequest));
 }
@@ -21,10 +22,10 @@ void FChaosVDTabSpawnerBase::HandleTabClosed(TSharedRef<SDockTab> InTabClosed)
 	InTabClosed->SetOnTabClosed(SDockTab::FOnTabClosedCallback());
 }
 
-void FChaosVDTabSpawnerBase::HandleTabSpawned(TSharedRef<SDockTab> InTabClosed)
+void FChaosVDTabSpawnerBase::HandleTabSpawned(TSharedRef<SDockTab> InTabSpawned)
 {
-	OnTabSpawned().Broadcast(InTabClosed);
-	InTabClosed->SetOnTabClosed(SDockTab::FOnTabClosedCallback::CreateRaw(this, &FChaosVDTabSpawnerBase::HandleTabClosed));
+	OnTabSpawned().Broadcast(InTabSpawned);
+	InTabSpawned->SetOnTabClosed(SDockTab::FOnTabClosedCallback::CreateRaw(this, &FChaosVDTabSpawnerBase::HandleTabClosed));
 }
 
 TSharedRef<SWidget> FChaosVDTabSpawnerBase::GenerateErrorWidget()
diff --git a/Engine/Plugins/ChaosVD/Source/ChaosVD/Private/ChaosVDTabsIDs.cpp b/Engine/Plugins/ChaosVD/Source/ChaosVD/Private/ChaosVDTabsIDs.cpp
--- a/Engine/Plugins/ChaosVD/Source/ChaosVD/Private/ChaosVDTabsIDs.cpp
+++ b/Engine/Plugins/ChaosVD/Source/ChaosVD/Private/ChaosVDTabsIDs.cpp
@@ -1,7 +1,6 @@
 // Copyright Epic Games, Inc. All Rights Reserved.
 
 #include "ChaosVDTabsIDs.h"
-#include "UObject/NameTypes.h"
 
 const FName FChaosVDTabID::ChaosVisualDebuggerTab = TEXT("ChaosVisualDebuggerTab");
 const FName FChaosVDTabID::PlaybackViewport = TEXT("PlaybackViewport");
diff --git a/Engine/Plugins/ChaosVD/Source/ChaosVD/Private/Components/ChaosVDConstraintDataComponent.h b/Engine/Plugins/ChaosVD/Source/ChaosVD/Private/Components/ChaosVDConstraintDataComponent.h
--- a/Engine/Plugins/ChaosVD/Source/ChaosVD/Private/Components/ChaosVDConstraintDataComponent.h
+++ b/Engine/Plugins/ChaosVD/Source/ChaosVD/Private/Components/ChaosVDConstraintDataComponent.h
@@ -3,7 +3,7 @@
 #pragma once
 
 #include "Components/ChaosVDConstraintDataHelpers.h"
-#include "ChaosVDSolverDataComponent.h"
+#include "Components/ChaosVDSolverDataComponent.h"
 #include "DataWrappers/ChaosVDParticleDataWrapper.h"
 #include "ChaosVDConstraintDataComponent.generated.h"
 
diff --git a/Engine/Plugins/ChaosVD/Source/ChaosVD/Private/Components/ChaosVDInstancedStaticMeshComponent.cpp b/Engine/Plugins/ChaosVD/Source/ChaosVD/Private/Components/ChaosVDInstancedStaticMeshComponent.cpp
--- a/Engine/Plugins/ChaosVD/Source/ChaosVD/Private/Components/ChaosVDInstancedStaticMeshComponent.cpp
+++ b/Engine/Plugins/ChaosVD/Source/ChaosVD/Private/Components/ChaosVDInstancedStaticMeshComponent.cpp
@@ -155,7 +155,7 @@ void UChaosVDInstancedStaticMeshComponent::Initialize()
 	{
 		if (!HasOverrideMaterials())
 		{
-			GeometryBuilder->RequestMaterialInstance(this);
+			GeometryBuilder->RequestMaterialUpdate(this);
 		}
 	}
 }
diff --git a/Engine/Plugins/ChaosVD/Source/ChaosVD/Private/Components/ChaosVDSceneQueryDataComponent.h b/Engine/Plugins/ChaosVD/Source/ChaosVD/Private/Components/ChaosVDSceneQueryDataComponent.h
--- a/Engine/Plugins/ChaosVD/Source/ChaosVD/Private/Components/ChaosVDSceneQueryDataComponent.h
+++ b/Engine/Plugins/ChaosVD/Source/ChaosVD/Private/Components/ChaosVDSceneQueryDataComponent.h
@@ -3,7 +3,7 @@
 #pragma once
 
 #include "DataWrappers/ChaosVDQueryDataWrappers.h"
-#include "ChaosVDSolverDataComponent.h"
+#include "Components/ChaosVDSolverDataComponent.h"
 #include "ChaosVDSolverDataSelection.h"
 
 #include "ChaosVDSceneQueryDataComponent.generated.h"
diff --git a/Engine/Plugins/ChaosVD/Source/ChaosVD/Private/Components/ChaosVDSolverCollisionDataComponent.h b/Engine/Plugins/ChaosVD/Source/ChaosVD/Private/Components/ChaosVDSolverCollisionDataComponent.h
--- a/Engine/Plugins/ChaosVD/Source/ChaosVD/Private/Components/ChaosVDSolverCollisionDataComponent.h
+++ b/Engine/Plugins/ChaosVD/Source/ChaosVD/Private/Components/ChaosVDSolverCollisionDataComponent.h
@@ -4,11 +4,11 @@
 
 #include "ChaosVDConstraintDataHelpers.h"
 
+#include "Components/ChaosVDSolverDataComponent.h"
 #include "Containers/Array.h"
 #include "Containers/Map.h"
 #include "DataWrappers/ChaosVDCollisionDataWrappers.h"
 #include "GameFramework/Info.h"
-#include "ChaosVDSolverDataComponent.h"
 #include "Templates/SharedPointer.h"
 
 #include "ChaosVDSolverCollisionDataComponent.generated.h"
diff --git a/Engine/Plugins/ChaosVD/Source/ChaosVD/Private/Components/ChaosVDStaticMeshComponent.cpp b/Engine/Plugins/ChaosVD/Source/ChaosVD/Private/Components/ChaosVDStaticMeshComponent.cpp
--- a/Engine/Plugins/ChaosVD/Source/ChaosVD/Private/Components/ChaosVDStaticMeshComponent.cpp
+++ b/Engine/Plugins/ChaosVD/Source/ChaosVD/Private/Components/ChaosVDStaticMeshComponent.cpp
@@ -74,7 +74,7 @@ void UChaosVDStaticMeshComponent::UpdateInstanceColor(const TSharedPtr<FChaosVDM
 
 			Chaos::VisualDebugger::Utils::EnumAddToggleFlag(MeshComponentAttributeFlags, EChaosVDMeshAttributesFlags::TranslucentGeometry);
 			
-			GeometryBuilder->RequestMaterialInstance(this);
+			GeometryBuilder->RequestMaterialUpdate(this);
 		}
 	}
 }
@@ -216,6 +216,6 @@ void UChaosVDStaticMeshComponent::Initialize()
 	TSharedPtr<FChaosVDGeometryBuilder> GeometryBuilder = GeometryBuilderWeakPtr.Pin();
 	if (ensure(GeometryBuilder))
 	{
-		GeometryBuilder->RequestMaterialInstance(this);
+		GeometryBuilder->RequestMaterialUpdate(this);
 	}
 }
diff --git a/Engine/Plugins/ChaosVD/Source/ChaosVD/Private/DetailsCustomizations/ChaosVDQueryDataWrappersCustomizationDetails.cpp b/Engine/Plugins/ChaosVD/Source/ChaosVD/Private/DetailsCustomizations/ChaosVDQueryDataWrappersCustomizationDetails.cpp
--- a/Engine/Plugins/ChaosVD/Source/ChaosVD/Private/DetailsCustomizations/ChaosVDQueryDataWrappersCustomizationDetails.cpp
+++ b/Engine/Plugins/ChaosVD/Source/ChaosVD/Private/DetailsCustomizations/ChaosVDQueryDataWrappersCustomizationDetails.cpp
@@ -98,7 +98,7 @@ void FChaosVDCollisionChannelsCustomizationBase::CustomizeChildren(TSharedRef<IP
 	TSharedPtr<SChaosVDMainTab> MainTabPtr =  MainTabWeakPtr.Pin(); 
 	TSharedPtr<FChaosVDScene> Scene = MainTabPtr ? MainTabPtr->GetChaosVDEngineInstance()->GetCurrentScene() : nullptr;
 
-	if (TSharedPtr<FChaosVDRecording> CVDRecording = Scene ? Scene->LoadedRecording : nullptr)
+	if (TSharedPtr<FChaosVDRecording> CVDRecording = Scene ? Scene->GetLoadedRecording() : nullptr)
 	{
 		UpdateCollisionChannelsInfoCache(CVDRecording->GetCollisionChannelsInfoContainer());
 	}
diff --git a/Engine/Plugins/ChaosVD/Source/ChaosVD/Private/DetailsCustomizations/ChaosVDShapeDataCustomization.cpp b/Engine/Plugins/ChaosVD/Source/ChaosVD/Private/DetailsCustomizations/ChaosVDShapeDataCustomization.cpp
--- a/Engine/Plugins/ChaosVD/Source/ChaosVD/Private/DetailsCustomizations/ChaosVDShapeDataCustomization.cpp
+++ b/Engine/Plugins/ChaosVD/Source/ChaosVD/Private/DetailsCustomizations/ChaosVDShapeDataCustomization.cpp
@@ -59,7 +59,7 @@ void FChaosVDShapeDataCustomization::CustomizeChildren(TSharedRef<IPropertyHandl
 		return;
 	}
 
-	if (TSharedPtr<FChaosVDRecording> CVDRecording = Scene->LoadedRecording)
+	if (TSharedPtr<FChaosVDRecording> CVDRecording = Scene->GetLoadedRecording())
 	{
 		UpdateCollisionChannelsInfoCache(CVDRecording->GetCollisionChannelsInfoContainer());
 	}
diff --git a/Engine/Plugins/ChaosVD/Source/ChaosVD/Private/ExtensionsSystem/ChaosVDExtensionsManager.cpp b/Engine/Plugins/ChaosVD/Source/ChaosVD/Private/ExtensionsSystem/ChaosVDExtensionsManager.cpp
new file mode 100644
--- /dev/null
+++ b/Engine/Plugins/ChaosVD/Source/ChaosVD/Private/ExtensionsSystem/ChaosVDExtensionsManager.cpp
@@ -0,0 +1,61 @@
+// Copyright Epic Games, Inc. All Rights Reserved.
+
+#include "ExtensionsSystem/ChaosVDExtensionsManager.h"
+
+#include "ChaosVDModule.h"
+#include "Features/IModularFeatures.h"
+#include "Misc/LazySingleton.h"
+
+FChaosVDExtensionsManager::FChaosVDExtensionsManager()
+{
+}
+
+FChaosVDExtensionsManager::~FChaosVDExtensionsManager()
+{
+	IModularFeatures::Get().OnModularFeatureRegistered().RemoveAll(this);
+	IModularFeatures::Get().OnModularFeatureUnregistered().RemoveAll(this);
+	
+	EnumerateExtensions([this](const TSharedRef<FChaosVDExtension>& InExtension)
+	{
+		UnRegisterExtension(InExtension);
+
+		return true;
+	});
+}
+
+FChaosVDExtensionsManager& FChaosVDExtensionsManager::Get()
+{
+	return TLazySingleton<FChaosVDExtensionsManager>::Get();
+}
+
+void FChaosVDExtensionsManager::TearDown()
+{
+	TLazySingleton<FChaosVDExtensionsManager>::TearDown();
+}
+
+void FChaosVDExtensionsManager::RegisterExtension(const TSharedRef<FChaosVDExtension>& InExtension)
+{
+	FName ExtensionType = InExtension->GetExtensionType();
+	if (ensure(!AvailableExtensions.Contains(ExtensionType)))
+	{
+		AvailableExtensions.Add(ExtensionType, InExtension);
+
+		ExtensionRegisteredEvent.Broadcast(InExtension);
+
+		UE_LOG(LogChaosVDEditor, Log, TEXT("[%s] Registering CVD Extension [%s] ..."), ANSI_TO_TCHAR(__FUNCTION__), *ExtensionType.ToString());
+	}
+	else
+	{
+		UE_LOG(LogChaosVDEditor, Warning, TEXT("[%s] CVD Extension [%s] already registered (or another extension is using the same id). Skipping ..."), ANSI_TO_TCHAR(__FUNCTION__), *ExtensionType.ToString());
+	}
+}
+
+void FChaosVDExtensionsManager::UnRegisterExtension(const TSharedRef<FChaosVDExtension>& InExtension)
+{
+	FName ExtensionType = InExtension->GetExtensionType();
+	AvailableExtensions.Remove(ExtensionType);
+	
+	ExtensionRegisteredEvent.Broadcast(InExtension);
+
+	UE_LOG(LogChaosVDEditor, Log, TEXT("[%s] UnRegistering CVD Extension [%s] ..."), ANSI_TO_TCHAR(__FUNCTION__), *ExtensionType.ToString());
+}
diff --git a/Engine/Plugins/ChaosVD/Source/ChaosVD/Private/Settings/ChaosVDCharacterConstraintsVisualizationSettings.h b/Engine/Plugins/ChaosVD/Source/ChaosVD/Private/Settings/ChaosVDCharacterConstraintsVisualizationSettings.h
--- a/Engine/Plugins/ChaosVD/Source/ChaosVD/Private/Settings/ChaosVDCharacterConstraintsVisualizationSettings.h
+++ b/Engine/Plugins/ChaosVD/Source/ChaosVD/Private/Settings/ChaosVDCharacterConstraintsVisualizationSettings.h
@@ -2,7 +2,7 @@
 
 #pragma once
 
-#include "ChaosVDCoreSettings.h"
+#include "Settings/ChaosVDCoreSettings.h"
 #include "UObject/Object.h"
 #include "ChaosVDCharacterConstraintsVisualizationSettings.generated.h"
 
diff --git a/Engine/Plugins/ChaosVD/Source/ChaosVD/Private/Settings/ChaosVDCollisionVisualizationSettings.h b/Engine/Plugins/ChaosVD/Source/ChaosVD/Private/Settings/ChaosVDCollisionVisualizationSettings.h
--- a/Engine/Plugins/ChaosVD/Source/ChaosVD/Private/Settings/ChaosVDCollisionVisualizationSettings.h
+++ b/Engine/Plugins/ChaosVD/Source/ChaosVD/Private/Settings/ChaosVDCollisionVisualizationSettings.h
@@ -1,7 +1,8 @@
 // Copyright Epic Games, Inc. All Rights Reserved.
 
 #pragma once
-#include "ChaosVDCoreSettings.h"
+
+#include "Settings/ChaosVDCoreSettings.h"
 
 #include "ChaosVDCollisionVisualizationSettings.generated.h"
 
diff --git a/Engine/Plugins/ChaosVD/Source/ChaosVD/Private/Settings/ChaosVDCoreSettings.cpp b/Engine/Plugins/ChaosVD/Source/ChaosVD/Private/Settings/ChaosVDCoreSettings.cpp
--- a/Engine/Plugins/ChaosVD/Source/ChaosVD/Private/Settings/ChaosVDCoreSettings.cpp
+++ b/Engine/Plugins/ChaosVD/Source/ChaosVD/Private/Settings/ChaosVDCoreSettings.cpp
@@ -1,6 +1,6 @@
 // Copyright Epic Games, Inc. All Rights Reserved.
 
-#include "ChaosVDCoreSettings.h"
+#include "Settings/ChaosVDCoreSettings.h"
 
 #include "Misc/ConfigContext.h"
 #include "Widgets/SChaosVDPlaybackViewport.h"
diff --git a/Engine/Plugins/ChaosVD/Source/ChaosVD/Private/Settings/ChaosVDJointConstraintVisualizationSettings.h b/Engine/Plugins/ChaosVD/Source/ChaosVD/Private/Settings/ChaosVDJointConstraintVisualizationSettings.h
--- a/Engine/Plugins/ChaosVD/Source/ChaosVD/Private/Settings/ChaosVDJointConstraintVisualizationSettings.h
+++ b/Engine/Plugins/ChaosVD/Source/ChaosVD/Private/Settings/ChaosVDJointConstraintVisualizationSettings.h
@@ -1,7 +1,8 @@
 // Copyright Epic Games, Inc. All Rights Reserved.
 
 #pragma once
-#include "ChaosVDCoreSettings.h"
+
+#include "Settings/ChaosVDCoreSettings.h"
 
 #include "ChaosVDJointConstraintVisualizationSettings.generated.h"
 
diff --git a/Engine/Plugins/ChaosVD/Source/ChaosVD/Private/Trace/ChaosVDTraceAnalyzer.cpp b/Engine/Plugins/ChaosVD/Source/ChaosVD/Private/Trace/ChaosVDTraceAnalyzer.cpp
--- a/Engine/Plugins/ChaosVD/Source/ChaosVD/Private/Trace/ChaosVDTraceAnalyzer.cpp
+++ b/Engine/Plugins/ChaosVD/Source/ChaosVD/Private/Trace/ChaosVDTraceAnalyzer.cpp
@@ -124,7 +124,7 @@ bool FChaosVDTraceAnalyzer::OnEvent(uint16 RouteId, EStyle Style, const FOnEvent
 			{
 				if (FrameData->SolverSteps.Num() > 0)
 				{
-					FChaosVDStepData& LastSolverStage = FrameData->SolverSteps.Last();
+					FChaosVDFrameStageData& LastSolverStage = FrameData->SolverSteps.Last();
 					if (EnumHasAnyFlags(LastSolverStage.StageFlags, EChaosVDSolverStageFlags::Open) && ensure(!EnumHasAnyFlags(LastSolverStage.StageFlags, EChaosVDSolverStageFlags::ExplicitStage)))
 					{
 						// If the current Solver stage was implicitly generated, we need to close it before starting a new one.
@@ -134,7 +134,7 @@ bool FChaosVDTraceAnalyzer::OnEvent(uint16 RouteId, EStyle Style, const FOnEvent
 				}
 
 				// Add an empty step. It will be filled out by the particle (and later on other objects/elements) events
-				FChaosVDStepData& NewSolverStageData = FrameData->SolverSteps.AddDefaulted_GetRef();
+				FChaosVDFrameStageData& NewSolverStageData = FrameData->SolverSteps.AddDefaulted_GetRef();
 
 				FWideStringView DebugNameView;
 				EventData.GetString("StepName", DebugNameView);
@@ -181,7 +181,7 @@ bool FChaosVDTraceAnalyzer::OnEvent(uint16 RouteId, EStyle Style, const FOnEvent
 		{
 			const int32 DataID = EventData.GetValue<int32>("DataID");
 			
-			FChaosVDBinaryDataContainer& DataContainer = ChaosVDTraceProvider->FindOrAddUnprocessedData(DataID);
+			FChaosVDTraceProvider::FBinaryDataContainer& DataContainer = ChaosVDTraceProvider->FindOrAddUnprocessedData(DataID);
 			DataContainer.bIsCompressed = EventData.GetValue<bool>("IsCompressed");
 			DataContainer.UncompressedSize = EventData.GetValue<uint32>("OriginalSize");
 			DataContainer.DataID = EventData.GetValue<int32>("DataID");
@@ -197,7 +197,7 @@ bool FChaosVDTraceAnalyzer::OnEvent(uint16 RouteId, EStyle Style, const FOnEvent
 		{
 			const int32 DataID = EventData.GetValue<int32>("DataID");	
 
-			FChaosVDBinaryDataContainer& DataContainer = ChaosVDTraceProvider->FindOrAddUnprocessedData(DataID);
+			FChaosVDTraceProvider::FBinaryDataContainer& DataContainer = ChaosVDTraceProvider->FindOrAddUnprocessedData(DataID);
 
 			const TArrayView<const uint8> SerializedDataChunk = EventData.GetArrayView<uint8>("RawData");
 			DataContainer.RawData.Append(SerializedDataChunk.GetData(), SerializedDataChunk.Num());
diff --git a/Engine/Plugins/ChaosVD/Source/ChaosVD/Private/Trace/ChaosVDTraceAnalyzer.h b/Engine/Plugins/ChaosVD/Source/ChaosVD/Private/Trace/ChaosVDTraceAnalyzer.h
--- a/Engine/Plugins/ChaosVD/Source/ChaosVD/Private/Trace/ChaosVDTraceAnalyzer.h
+++ b/Engine/Plugins/ChaosVD/Source/ChaosVD/Private/Trace/ChaosVDTraceAnalyzer.h
@@ -7,7 +7,7 @@
 #include "Templates/SharedPointer.h"
 #include "ChaosVisualDebugger/ChaosVDOptionalDataChannel.h"
 
-struct FChaosVDStepData;
+struct FChaosVDFrameStageData;
 class FChaosVDTraceProvider;
 
 namespace TraceServices { class IAnalysisSession; }
diff --git a/Engine/Plugins/ChaosVD/Source/ChaosVD/Private/Trace/ChaosVDTraceProvider.cpp b/Engine/Plugins/ChaosVD/Source/ChaosVD/Private/Trace/ChaosVDTraceProvider.cpp
--- a/Engine/Plugins/ChaosVD/Source/ChaosVD/Private/Trace/ChaosVDTraceProvider.cpp
+++ b/Engine/Plugins/ChaosVD/Source/ChaosVD/Private/Trace/ChaosVDTraceProvider.cpp
@@ -4,17 +4,12 @@
 
 #include "ChaosVDModule.h"
 #include "ChaosVDRecording.h"
-
-#include "ChaosVisualDebugger/ChaosVDSerializedNameTable.h"
 #include "ChaosVisualDebugger/ChaosVisualDebuggerTrace.h"
-
 #include "Compression/OodleDataCompressionUtil.h"
-#include "DataProcessors/ChaosVDAABBTreeDataProcessor.h"
+#include "Containers/Ticker.h"
 #include "DataProcessors/ChaosVDCollisionChannelsInfoDataProcessor.h"
-#include "DataProcessors/ChaosVDDebugDrawBoxDataProcessor.h"
-#include "DataProcessors/ChaosVDDebugDrawImplicitObjectDataProcessor.h"
-#include "DataProcessors/ChaosVDDebugDrawLineDataProcessor.h"
-#include "DataProcessors/ChaosVDDebugDrawSphereDataProcessor.h"
+#include "ExtensionsSystem/ChaosVDExtensionsManager.h"
+#include "Misc/MessageDialog.h"
 #include "Trace/DataProcessors/ChaosVDCharacterGroundConstraintDataProcessor.h"
 #include "Trace/DataProcessors/ChaosVDConstraintDataProcessor.h"
 #include "Trace/DataProcessors/ChaosVDJointConstraintDataProcessor.h"
@@ -36,6 +31,8 @@ namespace Chaos::VisualDebugger::Cvars
 		TEXT("The max number of game frames we will queue before processing them while the trace files is being loaded."));
 }
 
+#define LOCTEXT_NAMESPACE "ChaosVisualDebugger"
+
 FName FChaosVDTraceProvider::ProviderName("ChaosVDProvider");
 
 FChaosVDTraceProvider::FChaosVDTraceProvider(TraceServices::IAnalysisSession& InSession): Session(InSession)
@@ -85,6 +82,19 @@ void FChaosVDTraceProvider::GetAvailablePendingSolverIDsAtGameFrame(const TShare
 	}
 }
 
+FString FChaosVDTraceProvider::GenerateFormattedStringListFromSet(const TSet<FString>& StringsSet)
+{
+	FString FormattedString = TEXT("");
+	for (const FString& ListEntry : StringsSet)
+	{
+		FormattedString.Append(TEXT("- "));
+		FormattedString.Append(ListEntry);
+		FormattedString.Append(TEXT("\n"));
+	}
+
+	return MoveTemp(FormattedString);
+}
+
 void FChaosVDTraceProvider::CommitProcessedGameFramesToRecording()
 {
 	TArray<int32, TInlineAllocator<16>> SolverIDs;
@@ -157,16 +167,16 @@ TWeakPtr<FChaosVDGameFrameData> FChaosVDTraceProvider::GetCurrentGameFrame()
 	return CurrentGameFrame;
 }
 
-FChaosVDBinaryDataContainer& FChaosVDTraceProvider::FindOrAddUnprocessedData(const int32 DataID)
+FChaosVDTraceProvider::FBinaryDataContainer& FChaosVDTraceProvider::FindOrAddUnprocessedData(const int32 DataID)
 {
-	if (const TSharedPtr<FChaosVDBinaryDataContainer>* UnprocessedData = UnprocessedDataByID.Find(DataID))
+	if (const TSharedPtr<FBinaryDataContainer>* UnprocessedData = UnprocessedDataByID.Find(DataID))
 	{
 		check(UnprocessedData->IsValid());
 		return *UnprocessedData->Get();
 	}
 	else
 	{
-		const TSharedPtr<FChaosVDBinaryDataContainer> DataContainer = MakeShared<FChaosVDBinaryDataContainer>(DataID);
+		const TSharedPtr<FBinaryDataContainer> DataContainer = MakeShared<FBinaryDataContainer>(DataID);
 		UnprocessedDataByID.Add(DataID, DataContainer);
 		return *DataContainer.Get();
 	}
@@ -176,9 +186,9 @@ bool FChaosVDTraceProvider::ProcessBinaryData(const int32 DataID)
 {
 	RegisterDefaultDataProcessorsIfNeeded();
 
-	if (const TSharedPtr<FChaosVDBinaryDataContainer>* UnprocessedDataPtr = UnprocessedDataByID.Find(DataID))
+	if (const TSharedPtr<FBinaryDataContainer>* UnprocessedDataPtr = UnprocessedDataByID.Find(DataID))
 	{
-		const TSharedPtr<FChaosVDBinaryDataContainer> UnprocessedData = *UnprocessedDataPtr;
+		const TSharedPtr<FBinaryDataContainer> UnprocessedData = *UnprocessedDataPtr;
 		if (UnprocessedData.IsValid())
 		{
 			UnprocessedData->bIsReady = true;
@@ -207,13 +217,14 @@ bool FChaosVDTraceProvider::ProcessBinaryData(const int32 DataID)
 					}
 					else
 					{
-						UE_LOG(LogChaosVDEditor, Warning, TEXT("[%s] Failed to serialize Binary Data with ID [%d] | Type [%s]"), ANSI_TO_TCHAR(__FUNCTION__), DataID, *UnprocessedData->TypeName);
+						UE_LOG(LogChaosVDEditor, Verbose, TEXT("[%s] Failed to serialize Binary Data with ID [%d] | Type [%s]"), ANSI_TO_TCHAR(__FUNCTION__), DataID, *UnprocessedData->TypeName);
+						TypesFailedToSerialize.Add(UnprocessedData->TypeName);
 					}
 				}
 			}
 			else
 			{
-				UE_LOG(LogChaosVDEditor, Warning, TEXT("[%s] Data processor for type [%s] not found"), ANSI_TO_TCHAR(__FUNCTION__), *UnprocessedData->TypeName);
+				MissingDataProcessors.Add(UnprocessedData->TypeName);	
 			}
 		}
 
@@ -235,6 +246,32 @@ void FChaosVDTraceProvider::RegisterDataProcessor(TSharedPtr<FChaosVDDataProcess
 
 void FChaosVDTraceProvider::HandleAnalysisComplete()
 {
+	if (!MissingDataProcessors.IsEmpty())
+	{
+		FString MissingProcessorNameList = GenerateFormattedStringListFromSet(MissingDataProcessors);
+	
+		FText MissingDataProcessorsMessage = FText::FormatOrdered(LOCTEXT("MissingDataProcessorMessage", "This recording was made with CVD extensions that are not supported in this version. \n\nAs a result, the following data types could not be read and will be ignored : \n\n {0}"), FText::FromString(MissingProcessorNameList));
+
+		FTSTicker::GetCoreTicker().AddTicker(FTickerDelegate::CreateLambda([MissingDataProcessorsMessage](float DeltaTime)
+		{
+			FMessageDialog::Open(EAppMsgCategory::Warning, EAppMsgType::Ok, MissingDataProcessorsMessage, LOCTEXT("MissingDataProcessorMessageTitle", "Partially unsupported CVD Recording"));
+			return false;
+		}));
+	}
+
+	if (!TypesFailedToSerialize.IsEmpty())
+	{
+		FString FailedTypeList = GenerateFormattedStringListFromSet(TypesFailedToSerialize);
+	
+		FText MissingDataProcessorsMessage = FText::FormatOrdered(LOCTEXT("FailedSerializationMessage", "The following data types were part of the recording, but they couldn't be read : \n\n {0} \n\n Visualization related to that data will not be shown."), FText::FromString(FailedTypeList));
+
+		FTSTicker::GetCoreTicker().AddTicker(FTickerDelegate::CreateLambda([MissingDataProcessorsMessage](float DeltaTime)
+		{
+			FMessageDialog::Open(EAppMsgCategory::Error, EAppMsgType::Ok, MissingDataProcessorsMessage, LOCTEXT("FailedSerializationMessageTitle", "Failed to read data"));
+			return false;
+		}));
+	}
+
 	UnprocessedDataByID.Reset();
 
 	UE_LOG(LogChaosVDEditor, Log, TEXT("Trace Analysis complete for session [%s] | Calculating data loaded stats..."), Session.GetName());
@@ -265,12 +302,12 @@ void FChaosVDTraceProvider::HandleAnalysisComplete()
 	UE_LOG(LogChaosVDEditor, Log, TEXT("Total size of loaded data => [%s]"), *FText::AsMemory(TotalBytes, &SizeFormattingOptions,nullptr, EMemoryUnitStandard::IEC).ToString());
 }
 
-FChaosVDStepData* FChaosVDTraceProvider::GetCurrentSolverStageDataForCurrentFrame(int32 SolverID, EChaosVDSolverStageAccessorFlags Flags)
+FChaosVDFrameStageData* FChaosVDTraceProvider::GetCurrentSolverStageDataForCurrentFrame(int32 SolverID, EChaosVDSolverStageAccessorFlags Flags)
 {
 	auto CreateInBetweenSolverStage = [](FChaosVDSolverFrameData& InFrameData)
 	{
 		// Add an empty step. It will be filled out by the particle (and later on other objects/elements) events
-		FChaosVDStepData& SolverStageData = InFrameData.SolverSteps.AddDefaulted_GetRef();
+		FChaosVDFrameStageData& SolverStageData = InFrameData.SolverSteps.AddDefaulted_GetRef();
 		SolverStageData.StepName = TEXT("Between Stage Data");
 		EnumAddFlags(SolverStageData.StageFlags, EChaosVDSolverStageFlags::Open);
 
@@ -287,7 +324,7 @@ FChaosVDStepData* FChaosVDTraceProvider::GetCurrentSolverStageDataForCurrentFram
 			}
 		}
 
-		FChaosVDStepData& CurrentSolverStage = FrameData->SolverSteps.Last();
+		FChaosVDFrameStageData& CurrentSolverStage = FrameData->SolverSteps.Last();
 		if (EnumHasAnyFlags(CurrentSolverStage.StageFlags, EChaosVDSolverStageFlags::Open))
 		{
 			return &CurrentSolverStage;
@@ -348,31 +385,17 @@ void FChaosVDTraceProvider::RegisterDefaultDataProcessorsIfNeeded()
 	TSharedPtr<FChaosVDArchiveHeaderProcessor> ArchiveHeaderDataProcessor = MakeShared<FChaosVDArchiveHeaderProcessor>();
 	ArchiveHeaderDataProcessor->SetTraceProvider(AsShared());
 	RegisterDataProcessor(ArchiveHeaderDataProcessor);
-
-	TSharedPtr<FChaosVDAABBTreeDataProcessor> AABBTreeDataProcessor = MakeShared<FChaosVDAABBTreeDataProcessor>();
-	AABBTreeDataProcessor->SetTraceProvider(AsShared());
-	RegisterDataProcessor(AABBTreeDataProcessor);
-
-	TSharedPtr<FChaosVDDebugDrawBoxDataProcessor> DebugDrawBoxesDataProcessor = MakeShared<FChaosVDDebugDrawBoxDataProcessor>();
-	DebugDrawBoxesDataProcessor->SetTraceProvider(AsShared());
-	RegisterDataProcessor(DebugDrawBoxesDataProcessor);
-
-	TSharedPtr<FChaosVDDebugDrawLineDataProcessor> DebugDrawLinesDataProcessor = MakeShared<FChaosVDDebugDrawLineDataProcessor>();
-	DebugDrawLinesDataProcessor->SetTraceProvider(AsShared());
-	RegisterDataProcessor(DebugDrawLinesDataProcessor);
-
-	TSharedPtr<FChaosVDDebugDrawSphereDataProcessor> DebugDrawSpheresDataProcessor = MakeShared<FChaosVDDebugDrawSphereDataProcessor>();
-	DebugDrawSpheresDataProcessor->SetTraceProvider(AsShared());
-	RegisterDataProcessor(DebugDrawSpheresDataProcessor);
-
-	TSharedPtr<FChaosVDDebugDrawImplicitObjectDataProcessor> DebugDrawImplicitObjectDataProcessor = MakeShared<FChaosVDDebugDrawImplicitObjectDataProcessor>();
-	DebugDrawImplicitObjectDataProcessor->SetTraceProvider(AsShared());
-	RegisterDataProcessor(DebugDrawImplicitObjectDataProcessor);
 	
 	TSharedPtr<FChaosVDCollisionChannelsInfoDataProcessor> CollisionChannelsInfoDataProcessor = MakeShared<FChaosVDCollisionChannelsInfoDataProcessor>();
 	CollisionChannelsInfoDataProcessor->SetTraceProvider(AsShared());
 	RegisterDataProcessor(CollisionChannelsInfoDataProcessor);
 
+	FChaosVDExtensionsManager::Get().EnumerateExtensions([this](const TSharedRef<FChaosVDExtension>& Extension)
+	{
+		Extension->RegisterDataProcessorsInstancesForProvider(StaticCastSharedRef<FChaosVDTraceProvider>(AsShared()));
+		return true;
+	});
+
 	bDefaultDataProcessorsRegistered = true;
 }
 
@@ -388,3 +411,5 @@ void FChaosVDTraceProvider::DeQueueGameFrameForProcessing(TSharedPtr<FChaosVDGam
 	CurrentGameFrameQueue.Dequeue(OutFrameData);
 	CurrentGameFrameQueueSize--;
 }
+
+#undef LOCTEXT_NAMESPACE
diff --git a/Engine/Plugins/ChaosVD/Source/ChaosVD/Private/Trace/DataProcessors/ChaosVDCharacterGroundConstraintDataProcessor.h b/Engine/Plugins/ChaosVD/Source/ChaosVD/Private/Trace/DataProcessors/ChaosVDCharacterGroundConstraintDataProcessor.h
--- a/Engine/Plugins/ChaosVD/Source/ChaosVD/Private/Trace/DataProcessors/ChaosVDCharacterGroundConstraintDataProcessor.h
+++ b/Engine/Plugins/ChaosVD/Source/ChaosVD/Private/Trace/DataProcessors/ChaosVDCharacterGroundConstraintDataProcessor.h
@@ -3,7 +3,7 @@
 #pragma once
 
 #pragma once
-#include "ChaosVDDataProcessorBase.h"
+#include "Trace/DataProcessors/ChaosVDDataProcessorBase.h"
 
 struct FChaosVDSolverFrameData;
 struct FChaosVDConstraint;
diff --git a/Engine/Plugins/ChaosVD/Source/ChaosVD/Private/Trace/DataProcessors/ChaosVDCollisionChannelsInfoDataProcessor.h b/Engine/Plugins/ChaosVD/Source/ChaosVD/Private/Trace/DataProcessors/ChaosVDCollisionChannelsInfoDataProcessor.h
--- a/Engine/Plugins/ChaosVD/Source/ChaosVD/Private/Trace/DataProcessors/ChaosVDCollisionChannelsInfoDataProcessor.h
+++ b/Engine/Plugins/ChaosVD/Source/ChaosVD/Private/Trace/DataProcessors/ChaosVDCollisionChannelsInfoDataProcessor.h
@@ -1,7 +1,7 @@
 ﻿// Copyright Epic Games, Inc. All Rights Reserved.
 
 #pragma once
-#include "ChaosVDDataProcessorBase.h"
+#include "Trace/DataProcessors/ChaosVDDataProcessorBase.h"
 
 struct FChaosVDSolverFrameData;
 
diff --git a/Engine/Plugins/ChaosVD/Source/ChaosVD/Private/Trace/DataProcessors/ChaosVDConstraintDataProcessor.cpp b/Engine/Plugins/ChaosVD/Source/ChaosVD/Private/Trace/DataProcessors/ChaosVDConstraintDataProcessor.cpp
--- a/Engine/Plugins/ChaosVD/Source/ChaosVD/Private/Trace/DataProcessors/ChaosVDConstraintDataProcessor.cpp
+++ b/Engine/Plugins/ChaosVD/Source/ChaosVD/Private/Trace/DataProcessors/ChaosVDConstraintDataProcessor.cpp
@@ -26,7 +26,7 @@ bool FChaosVDConstraintDataProcessor::ProcessRawData(const TArray<uint8>& InData
 
 	if (bSuccess)
 	{
-		FChaosVDStepData* CurrentSolverStage = ProviderSharedPtr->GetCurrentSolverStageDataForCurrentFrame(RecordedConstraint.SolverID, EChaosVDSolverStageAccessorFlags::None);
+		FChaosVDFrameStageData* CurrentSolverStage = ProviderSharedPtr->GetCurrentSolverStageDataForCurrentFrame(RecordedConstraint.SolverID, EChaosVDSolverStageAccessorFlags::None);
 		if (ensureMsgf(CurrentSolverStage, TEXT("A MidPhase was traced without a valid step scope")))
 		{
 			AddConstraintToParticleIDMap(RecordedConstraint, RecordedConstraint.Particle0Index, *CurrentSolverStage);
@@ -37,7 +37,7 @@ bool FChaosVDConstraintDataProcessor::ProcessRawData(const TArray<uint8>& InData
 	return bSuccess;
 }
 
-void FChaosVDConstraintDataProcessor::AddConstraintToParticleIDMap(const FChaosVDConstraint& InConstraintData, int32 ParticleID, FChaosVDStepData& InSolverStageData)
+void FChaosVDConstraintDataProcessor::AddConstraintToParticleIDMap(const FChaosVDConstraint& InConstraintData, int32 ParticleID, FChaosVDFrameStageData& InSolverStageData)
 {
 	if (TArray<FChaosVDConstraint>* ParticleConstraints = InSolverStageData.RecordedConstraintsByParticleID.Find(ParticleID))
 	{
diff --git a/Engine/Plugins/ChaosVD/Source/ChaosVD/Private/Trace/DataProcessors/ChaosVDConstraintDataProcessor.h b/Engine/Plugins/ChaosVD/Source/ChaosVD/Private/Trace/DataProcessors/ChaosVDConstraintDataProcessor.h
--- a/Engine/Plugins/ChaosVD/Source/ChaosVD/Private/Trace/DataProcessors/ChaosVDConstraintDataProcessor.h
+++ b/Engine/Plugins/ChaosVD/Source/ChaosVD/Private/Trace/DataProcessors/ChaosVDConstraintDataProcessor.h
@@ -1,10 +1,10 @@
 ﻿// Copyright Epic Games, Inc. All Rights Reserved.
 
 #pragma once
-#include "ChaosVDDataProcessorBase.h"
+#include "Trace/DataProcessors/ChaosVDDataProcessorBase.h"
 
 struct FChaosVDSolverFrameData;
-struct FChaosVDStepData;
+struct FChaosVDFrameStageData;
 struct FChaosVDConstraint;
 
 /**
@@ -17,5 +17,5 @@ public:
 	
 	virtual bool ProcessRawData(const TArray<uint8>& InData) override;
 
-	void AddConstraintToParticleIDMap(const FChaosVDConstraint& InConstraintData, int32 ParticleID, FChaosVDStepData& InSolverStageData);
+	void AddConstraintToParticleIDMap(const FChaosVDConstraint& InConstraintData, int32 ParticleID, FChaosVDFrameStageData& InSolverStageData);
 };
diff --git a/Engine/Plugins/ChaosVD/Source/ChaosVD/Private/Trace/DataProcessors/ChaosVDDataProcessorBase.cpp b/Engine/Plugins/ChaosVD/Source/ChaosVD/Private/Trace/DataProcessors/ChaosVDDataProcessorBase.cpp
--- a/Engine/Plugins/ChaosVD/Source/ChaosVD/Private/Trace/DataProcessors/ChaosVDDataProcessorBase.cpp
+++ b/Engine/Plugins/ChaosVD/Source/ChaosVD/Private/Trace/DataProcessors/ChaosVDDataProcessorBase.cpp
@@ -1,6 +1,6 @@
 // Copyright Epic Games, Inc. All Rights Reserved.
 
-#include "ChaosVDDataProcessorBase.h"
+#include "Trace/DataProcessors/ChaosVDDataProcessorBase.h"
 
 FChaosVDDataProcessorBase::~FChaosVDDataProcessorBase()
 {
@@ -26,3 +26,12 @@ void FChaosVDDataProcessorBase::SetTraceProvider(const TSharedPtr<FChaosVDTraceP
 {
 	TraceProvider = InProvider;
 }
+
+FChaosVDGenericDataProcessor::FChaosVDGenericDataProcessor(FStringView InCompatibleType, const TFunction<bool(const TArray<uint8>&)>& InProcessDataCallback) : FChaosVDDataProcessorBase(InCompatibleType), ProcessDataCallback(InProcessDataCallback)
+{
+}
+
+bool FChaosVDGenericDataProcessor::ProcessRawData(const TArray<uint8>& InData)
+{
+	return FChaosVDDataProcessorBase::ProcessRawData(InData) && ProcessDataCallback && ProcessDataCallback(InData);
+}
diff --git a/Engine/Plugins/ChaosVD/Source/ChaosVD/Private/Trace/DataProcessors/ChaosVDJointConstraintDataProcessor.cpp b/Engine/Plugins/ChaosVD/Source/ChaosVD/Private/Trace/DataProcessors/ChaosVDJointConstraintDataProcessor.cpp
--- a/Engine/Plugins/ChaosVD/Source/ChaosVD/Private/Trace/DataProcessors/ChaosVDJointConstraintDataProcessor.cpp
+++ b/Engine/Plugins/ChaosVD/Source/ChaosVD/Private/Trace/DataProcessors/ChaosVDJointConstraintDataProcessor.cpp
@@ -26,7 +26,7 @@ bool FChaosVDJointConstraintDataProcessor::ProcessRawData(const TArray<uint8>& I
 
 	if (bSuccess)
 	{
-		FChaosVDStepData* CurrentSolverStage = ProviderSharedPtr->GetCurrentSolverStageDataForCurrentFrame(JointConstraint->SolverID, EChaosVDSolverStageAccessorFlags::None);
+		FChaosVDFrameStageData* CurrentSolverStage = ProviderSharedPtr->GetCurrentSolverStageDataForCurrentFrame(JointConstraint->SolverID, EChaosVDSolverStageAccessorFlags::None);
 		if (ensureMsgf(CurrentSolverStage, TEXT("A Joint Constraint was traced without a valid step scope")))
 		{
 			CurrentSolverStage->RecordedJointConstraints.Add(JointConstraint);
diff --git a/Engine/Plugins/ChaosVD/Source/ChaosVD/Private/Trace/DataProcessors/ChaosVDMidPhaseDataProcessor.cpp b/Engine/Plugins/ChaosVD/Source/ChaosVD/Private/Trace/DataProcessors/ChaosVDMidPhaseDataProcessor.cpp
--- a/Engine/Plugins/ChaosVD/Source/ChaosVD/Private/Trace/DataProcessors/ChaosVDMidPhaseDataProcessor.cpp
+++ b/Engine/Plugins/ChaosVD/Source/ChaosVD/Private/Trace/DataProcessors/ChaosVDMidPhaseDataProcessor.cpp
@@ -26,7 +26,7 @@ bool FChaosVDMidPhaseDataProcessor::ProcessRawData(const TArray<uint8>& InData)
 
 	if (bSuccess)
 	{
-		FChaosVDStepData* CurrentSolverStage = ProviderSharedPtr->GetCurrentSolverStageDataForCurrentFrame(MidPhase->SolverID, EChaosVDSolverStageAccessorFlags::None);
+		FChaosVDFrameStageData* CurrentSolverStage = ProviderSharedPtr->GetCurrentSolverStageDataForCurrentFrame(MidPhase->SolverID, EChaosVDSolverStageAccessorFlags::None);
 
 		if (ensureMsgf(CurrentSolverStage, TEXT("A MidPhase was traced without a valid step scope")))
 		{
@@ -40,7 +40,7 @@ bool FChaosVDMidPhaseDataProcessor::ProcessRawData(const TArray<uint8>& InData)
 	return bSuccess;
 }
 
-void FChaosVDMidPhaseDataProcessor::AddMidPhaseToParticleIDMap(const TSharedPtr<FChaosVDParticlePairMidPhase>& MidPhaseData, int32 ParticleID, FChaosVDStepData& InSolverStageData)
+void FChaosVDMidPhaseDataProcessor::AddMidPhaseToParticleIDMap(const TSharedPtr<FChaosVDParticlePairMidPhase>& MidPhaseData, int32 ParticleID, FChaosVDFrameStageData& InSolverStageData)
 {
 	if (TArray<TSharedPtr<FChaosVDParticlePairMidPhase>>* ParticleMidPhases = InSolverStageData.RecordedMidPhasesByParticleID.Find(ParticleID))
 	{
diff --git a/Engine/Plugins/ChaosVD/Source/ChaosVD/Private/Trace/DataProcessors/ChaosVDMidPhaseDataProcessor.h b/Engine/Plugins/ChaosVD/Source/ChaosVD/Private/Trace/DataProcessors/ChaosVDMidPhaseDataProcessor.h
--- a/Engine/Plugins/ChaosVD/Source/ChaosVD/Private/Trace/DataProcessors/ChaosVDMidPhaseDataProcessor.h
+++ b/Engine/Plugins/ChaosVD/Source/ChaosVD/Private/Trace/DataProcessors/ChaosVDMidPhaseDataProcessor.h
@@ -1,11 +1,11 @@
 ﻿// Copyright Epic Games, Inc. All Rights Reserved.
 
 #pragma once
-#include "ChaosVDDataProcessorBase.h"
+#include "Trace/DataProcessors/ChaosVDDataProcessorBase.h"
 #include "Templates/SharedPointer.h"
 
 struct FChaosVDSolverFrameData;
-struct FChaosVDStepData;
+struct FChaosVDFrameStageData;
 struct FChaosVDParticlePairMidPhase;
 
 /**
@@ -18,5 +18,5 @@ public:
 
 	virtual bool ProcessRawData(const TArray<uint8>& InData) override;
 
-	void AddMidPhaseToParticleIDMap(const TSharedPtr<FChaosVDParticlePairMidPhase>& MidPhaseData, int32 ParticleID, FChaosVDStepData& InSolverStageData);
+	void AddMidPhaseToParticleIDMap(const TSharedPtr<FChaosVDParticlePairMidPhase>& MidPhaseData, int32 ParticleID, FChaosVDFrameStageData& InSolverStageData);
 };
diff --git a/Engine/Plugins/ChaosVD/Source/ChaosVD/Private/Trace/DataProcessors/ChaosVDTraceParticleDataProcessor.cpp b/Engine/Plugins/ChaosVD/Source/ChaosVD/Private/Trace/DataProcessors/ChaosVDTraceParticleDataProcessor.cpp
--- a/Engine/Plugins/ChaosVD/Source/ChaosVD/Private/Trace/DataProcessors/ChaosVDTraceParticleDataProcessor.cpp
+++ b/Engine/Plugins/ChaosVD/Source/ChaosVD/Private/Trace/DataProcessors/ChaosVDTraceParticleDataProcessor.cpp
@@ -29,7 +29,7 @@ bool FChaosVDTraceParticleDataProcessor::ProcessRawData(const TArray<uint8>& InD
 	if (bSuccess)
 	{
 		EChaosVDSolverStageAccessorFlags StageAccessorFlags = EChaosVDSolverStageAccessorFlags::CreateNewIfEmpty | EChaosVDSolverStageAccessorFlags::CreateNewIfClosed; 
-		if (FChaosVDStepData* CurrentSolverStage = ProviderSharedPtr->GetCurrentSolverStageDataForCurrentFrame(ParticleData->SolverID, StageAccessorFlags))
+		if (FChaosVDFrameStageData* CurrentSolverStage = ProviderSharedPtr->GetCurrentSolverStageDataForCurrentFrame(ParticleData->SolverID, StageAccessorFlags))
 		{
 			CurrentSolverStage->RecordedParticlesData.Add(ParticleData);
 		}
diff --git a/Engine/Plugins/ChaosVD/Source/ChaosVD/Private/Utils/ChaosVDUserInterfaceUtils.cpp b/Engine/Plugins/ChaosVD/Source/ChaosVD/Private/Utils/ChaosVDUserInterfaceUtils.cpp
--- a/Engine/Plugins/ChaosVD/Source/ChaosVD/Private/Utils/ChaosVDUserInterfaceUtils.cpp
+++ b/Engine/Plugins/ChaosVD/Source/ChaosVD/Private/Utils/ChaosVDUserInterfaceUtils.cpp
@@ -1,6 +1,6 @@
 // Copyright Epic Games, Inc. All Rights Reserved.
 
-#include "ChaosVDUserInterfaceUtils.h"
+#include "Utils/ChaosVDUserInterfaceUtils.h"
 #include "PropertyEditorModule.h"
 #include "Modules/ModuleManager.h"
 #include "ToolMenu.h"
diff --git a/Engine/Plugins/ChaosVD/Source/ChaosVD/Private/Visualizers/ChaosVDCharacterGroundConstraintsDataComponentVisualizer.cpp b/Engine/Plugins/ChaosVD/Source/ChaosVD/Private/Visualizers/ChaosVDCharacterGroundConstraintsDataComponentVisualizer.cpp
--- a/Engine/Plugins/ChaosVD/Source/ChaosVD/Private/Visualizers/ChaosVDCharacterGroundConstraintsDataComponentVisualizer.cpp
+++ b/Engine/Plugins/ChaosVD/Source/ChaosVD/Private/Visualizers/ChaosVDCharacterGroundConstraintsDataComponentVisualizer.cpp
@@ -71,7 +71,7 @@ void FChaosVDCharacterGroundConstraintDataComponentVisualizer::DrawVisualization
 		return;
 	}
 
-	const TSharedPtr<FChaosVDRecording> CVDRecording = CVDScene->LoadedRecording;
+	const TSharedPtr<FChaosVDRecording> CVDRecording = CVDScene->GetLoadedRecording();
 	if (!CVDRecording)
 	{
 		return;
diff --git a/Engine/Plugins/ChaosVD/Source/ChaosVD/Private/Visualizers/ChaosVDCharacterGroundConstraintsDataComponentVisualizer.h b/Engine/Plugins/ChaosVD/Source/ChaosVD/Private/Visualizers/ChaosVDCharacterGroundConstraintsDataComponentVisualizer.h
--- a/Engine/Plugins/ChaosVD/Source/ChaosVD/Private/Visualizers/ChaosVDCharacterGroundConstraintsDataComponentVisualizer.h
+++ b/Engine/Plugins/ChaosVD/Source/ChaosVD/Private/Visualizers/ChaosVDCharacterGroundConstraintsDataComponentVisualizer.h
@@ -1,8 +1,8 @@
 // Copyright Epic Games, Inc. All Rights Reserved.
 
 #pragma once
-#include "ChaosVDComponentVisualizerBase.h"
-#include "ComponentVisualizer.h"
+#include "Visualizers/ChaosVDComponentVisualizerBase.h"
+
 #include "Chaos/Core.h"
 
 struct FChaosVDSolverDataSelectionHandle;
diff --git a/Engine/Plugins/ChaosVD/Source/ChaosVD/Private/Visualizers/ChaosVDComponentVisualizerBase.cpp b/Engine/Plugins/ChaosVD/Source/ChaosVD/Private/Visualizers/ChaosVDComponentVisualizerBase.cpp
--- a/Engine/Plugins/ChaosVD/Source/ChaosVD/Private/Visualizers/ChaosVDComponentVisualizerBase.cpp
+++ b/Engine/Plugins/ChaosVD/Source/ChaosVD/Private/Visualizers/ChaosVDComponentVisualizerBase.cpp
@@ -1,6 +1,6 @@
 // Copyright Epic Games, Inc. All Rights Reserved.
 
-#include "ChaosVDComponentVisualizerBase.h"
+#include "Visualizers/ChaosVDComponentVisualizerBase.h"
 
 #include "ChaosVDEngine.h"
 #include "ChaosVDScene.h"
diff --git a/Engine/Plugins/ChaosVD/Source/ChaosVD/Private/Visualizers/ChaosVDJointConstraintsDataComponentVisualizer.cpp b/Engine/Plugins/ChaosVD/Source/ChaosVD/Private/Visualizers/ChaosVDJointConstraintsDataComponentVisualizer.cpp
--- a/Engine/Plugins/ChaosVD/Source/ChaosVD/Private/Visualizers/ChaosVDJointConstraintsDataComponentVisualizer.cpp
+++ b/Engine/Plugins/ChaosVD/Source/ChaosVD/Private/Visualizers/ChaosVDJointConstraintsDataComponentVisualizer.cpp
@@ -150,7 +150,7 @@ void FChaosVDJointConstraintsDataComponentVisualizer::DrawVisualization(const UA
 		return;
 	}
 
-	const TSharedPtr<FChaosVDRecording> CVDRecording = CVDScene->LoadedRecording;
+	const TSharedPtr<FChaosVDRecording> CVDRecording = CVDScene->GetLoadedRecording();
 	if (!CVDRecording)
 	{
 		return;
diff --git a/Engine/Plugins/ChaosVD/Source/ChaosVD/Private/Visualizers/ChaosVDJointConstraintsDataComponentVisualizer.h b/Engine/Plugins/ChaosVD/Source/ChaosVD/Private/Visualizers/ChaosVDJointConstraintsDataComponentVisualizer.h
--- a/Engine/Plugins/ChaosVD/Source/ChaosVD/Private/Visualizers/ChaosVDJointConstraintsDataComponentVisualizer.h
+++ b/Engine/Plugins/ChaosVD/Source/ChaosVD/Private/Visualizers/ChaosVDJointConstraintsDataComponentVisualizer.h
@@ -1,7 +1,7 @@
 ﻿// Copyright Epic Games, Inc. All Rights Reserved.
 
 #pragma once
-#include "ChaosVDComponentVisualizerBase.h"
+#include "Visualizers/ChaosVDComponentVisualizerBase.h"
 #include "ChaosVDSolverDataSelection.h"
 #include "Chaos/Core.h"
 #include "Components/ChaosVDSolverJointConstraintDataComponent.h"
diff --git a/Engine/Plugins/ChaosVD/Source/ChaosVD/Private/Visualizers/ChaosVDSolverCollisionDataComponentVisualizer.cpp b/Engine/Plugins/ChaosVD/Source/ChaosVD/Private/Visualizers/ChaosVDSolverCollisionDataComponentVisualizer.cpp
--- a/Engine/Plugins/ChaosVD/Source/ChaosVD/Private/Visualizers/ChaosVDSolverCollisionDataComponentVisualizer.cpp
+++ b/Engine/Plugins/ChaosVD/Source/ChaosVD/Private/Visualizers/ChaosVDSolverCollisionDataComponentVisualizer.cpp
@@ -2,7 +2,6 @@
 
 #include "ChaosVDSolverCollisionDataComponentVisualizer.h"
 
-#include "ChaosVDComponentVisualizerBase.h"
 #include "ChaosVDParticleActor.h"
 #include "ChaosVDScene.h"
 #include "ChaosVDSettingsManager.h"
diff --git a/Engine/Plugins/ChaosVD/Source/ChaosVD/Private/Widgets/SChaosVDCollisionDataInspector.cpp b/Engine/Plugins/ChaosVD/Source/ChaosVD/Private/Widgets/SChaosVDCollisionDataInspector.cpp
--- a/Engine/Plugins/ChaosVD/Source/ChaosVD/Private/Widgets/SChaosVDCollisionDataInspector.cpp
+++ b/Engine/Plugins/ChaosVD/Source/ChaosVD/Private/Widgets/SChaosVDCollisionDataInspector.cpp
@@ -7,7 +7,7 @@
 #include "IStructureDetailsView.h"
 #include "Modules/ModuleManager.h"
 #include "PropertyEditorModule.h"
-#include "SChaosVDMainTab.h"
+#include "Widgets/SChaosVDMainTab.h"
 #include "Visualizers/ChaosVDSolverCollisionDataComponentVisualizer.h"
 #include "Widgets/Input/SButton.h"
 #include "Widgets/Layout/SScrollBox.h"
diff --git a/Engine/Plugins/ChaosVD/Source/ChaosVD/Private/Widgets/SChaosVDConstraintDataInspector.cpp b/Engine/Plugins/ChaosVD/Source/ChaosVD/Private/Widgets/SChaosVDConstraintDataInspector.cpp
--- a/Engine/Plugins/ChaosVD/Source/ChaosVD/Private/Widgets/SChaosVDConstraintDataInspector.cpp
+++ b/Engine/Plugins/ChaosVD/Source/ChaosVD/Private/Widgets/SChaosVDConstraintDataInspector.cpp
@@ -5,7 +5,7 @@
 #include "ChaosVDScene.h"
 #include "IStructureDetailsView.h"
 #include "PropertyEditorModule.h"
-#include "SChaosVDMainTab.h"
+#include "Widgets/SChaosVDMainTab.h"
 #include "Actors/ChaosVDSolverInfoActor.h"
 #include "Widgets/SChaosVDWarningMessageBox.h"
 #include "Modules/ModuleManager.h"
diff --git a/Engine/Plugins/ChaosVD/Source/ChaosVD/Private/Widgets/SChaosVDMainTab.cpp b/Engine/Plugins/ChaosVD/Source/ChaosVD/Private/Widgets/SChaosVDMainTab.cpp
--- a/Engine/Plugins/ChaosVD/Source/ChaosVD/Private/Widgets/SChaosVDMainTab.cpp
+++ b/Engine/Plugins/ChaosVD/Source/ChaosVD/Private/Widgets/SChaosVDMainTab.cpp
@@ -34,8 +34,6 @@
 #include "StatusBarSubsystem.h"
 #include "ToolMenu.h"
 #include "ToolMenus.h"
-#include "Components/ChaosVDGenericDebugDrawDataComponent.h"
-#include "Components/ChaosVDGTAccelerationStructuresDataComponent.h"
 #include "Components/ChaosVDInstancedStaticMeshComponent.h"
 #include "Components/ChaosVDParticleDataComponent.h"
 #include "Components/ChaosVDSolverCharacterGroundConstraintDataComponent.h"
@@ -47,6 +45,9 @@
 #include "DetailsCustomizations/ChaosVDQueryDataWrappersCustomizationDetails.h"
 #include "DetailsCustomizations/ChaosVDSelectionMultipleViewCustomization.h"
 #include "DetailsCustomizations/ChaosVDShapeDataCustomization.h"
+#include "ExtensionsSystem/ChaosVDExtensionsManager.h"
+
+#include "Features/IModularFeatures.h"
 #include "HAL/FileManager.h"
 #include "Settings/ChaosVDMiscSettings.h"
 #include "Styling/StyleColors.h"
@@ -54,8 +55,6 @@
 #include "TabSpawers/ChaosVDSceneQueryBrowserTab.h"
 #include "Trace/ChaosVDTraceManager.h"
 #include "Visualizers/ChaosVDCharacterGroundConstraintsDataComponentVisualizer.h"
-#include "Visualizers/ChaosVDGenericDebugDrawDataComponentVisualizer.h"
-#include "Visualizers/ChaosVDGTAccelerationStructureDataComponentVisualizer.h"
 #include "Visualizers/ChaosVDJointConstraintsDataComponentVisualizer.h"
 #include "Visualizers/ChaosVDParticleDataComponentVisualizer.h"
 #include "Visualizers/ChaosVDSceneQueryDataComponentVisualizer.h"
@@ -84,8 +83,12 @@ void SChaosVDMainTab::Construct(const FArguments& InArgs, TSharedPtr<FChaosVDEng
 	RegisterComponentVisualizer(UChaosVDParticleDataComponent::StaticClass()->GetFName(), MakeShared<FChaosVDParticleDataComponentVisualizer>());
 	RegisterComponentVisualizer(UChaosVDSolverJointConstraintDataComponent::StaticClass()->GetFName(), MakeShared<FChaosVDJointConstraintsDataComponentVisualizer>());
 	RegisterComponentVisualizer(UChaosVDSolverCharacterGroundConstraintDataComponent::StaticClass()->GetFName(), MakeShared<FChaosVDCharacterGroundConstraintDataComponentVisualizer>());
-	RegisterComponentVisualizer(UChaosVDGTAccelerationStructuresDataComponent::StaticClass()->GetFName(), MakeShared<FChaosVDGTAccelerationStructureDataComponentVisualizer>());
-	RegisterComponentVisualizer(UChaosVDGenericDebugDrawDataComponent::StaticClass()->GetFName(), MakeShared<FChaosVDGenericDebugDrawDataComponentVisualizer>());
+	
+	FChaosVDExtensionsManager::Get().EnumerateExtensions([this](const TSharedRef<FChaosVDExtension>& Extension)
+	{
+		Extension->RegisterComponentVisualizers(StaticCastSharedRef<SChaosVDMainTab>(AsShared()));
+		return true;
+	});
 
 	TabManager = FGlobalTabmanager::Get()->NewTabManager(InArgs._OwnerTab.ToSharedRef()).ToSharedPtr();
 
@@ -100,6 +103,12 @@ void SChaosVDMainTab::Construct(const FArguments& InArgs, TSharedPtr<FChaosVDEng
 	RegisterTabSpawner<FChaosVDSceneQueryBrowserTab>(FChaosVDTabID::SceneQueryBrowser);
 	RegisterTabSpawner<FChaosVDRecordedLogTab>(FChaosVDTabID::RecordedOutputLog);
 
+	FChaosVDExtensionsManager::Get().EnumerateExtensions([this](const TSharedRef<FChaosVDExtension>& Extension)
+	{
+		Extension->RegisterCustomTabSpawners(StaticCastSharedRef<SChaosVDMainTab>(AsShared()));
+		return true;
+	});
+
 	StatusBarID = FName(FChaosVDTabID::StatusBar.ToString() + InChaosVDEngine->GetInstanceGuid().ToString());
 	
 	TSharedPtr<SWidget> StatusBarWidget;
@@ -177,6 +186,8 @@ void SChaosVDMainTab::Construct(const FArguments& InArgs, TSharedPtr<FChaosVDEng
 		]
 	];
 
+	FChaosVDExtensionsManager::Get().OnExtensionRegistered().AddSP(this, &SChaosVDMainTab::HandlePostInitializationExtensionRegistered);
+
 	// Make sure these tabs are always focused at the start
 	TabManager->TryInvokeTab(FChaosVDTabID::SolversTrack);
 	TabManager->TryInvokeTab(FChaosVDTabID::DetailsPanel);
@@ -186,6 +197,7 @@ void SChaosVDMainTab::Construct(const FArguments& InArgs, TSharedPtr<FChaosVDEng
 
 SChaosVDMainTab::~SChaosVDMainTab()
 {
+	FChaosVDExtensionsManager::Get().OnExtensionRegistered().RemoveAll(this);
 	CleanUpDisableCPUThrottlingDelegate();
 }
 
@@ -247,6 +259,20 @@ void SChaosVDMainTab::RegisterComponentVisualizer(FName ClassName, const TShared
 	}
 }
 
+void SChaosVDMainTab::HandlePostInitializationExtensionRegistered(const TSharedRef<FChaosVDExtension>& NewExtension)
+{	
+	NewExtension->RegisterCustomTabSpawners(StaticCastSharedRef<SChaosVDMainTab>(AsShared()));
+	NewExtension->RegisterComponentVisualizers(StaticCastSharedRef<SChaosVDMainTab>(AsShared()));
+
+	for (const TWeakPtr<IDetailsView>& DetailsPanel : CustomizedDetailsPanels)
+	{
+		if (TSharedPtr<IDetailsView> DetailsPanelPtr = DetailsPanel.Pin())
+		{
+			NewExtension->SetCustomPropertyLayouts(DetailsPanelPtr.Get(), StaticCastSharedRef<SChaosVDMainTab>(AsShared()));
+		}
+	}
+}
+
 void SChaosVDMainTab::HandleTabSpawned(TSharedRef<SDockTab> Tab, FName TabID)
 {
 	if (!ActiveTabsByID.Contains(TabID))
@@ -552,7 +578,6 @@ void SChaosVDMainTab::SetCustomPropertyLayouts(IDetailsView* DetailsView)
 		return;
 	}
 
-	FPropertyEditorModule& PropertyModule = FModuleManager::GetModuleChecked<FPropertyEditorModule>("PropertyEditor");
 	DetailsView->RegisterInstancedCustomPropertyLayout(AChaosVDParticleActor::StaticClass(), FOnGetDetailCustomizationInstance::CreateStatic(&FChaosVDParticleActorCustomization::MakeInstance, StaticCastWeakPtr<SChaosVDMainTab>(AsWeak())));
 	DetailsView->RegisterInstancedCustomPropertyLayout(UChaosVDInstancedStaticMeshComponent::StaticClass(), FOnGetDetailCustomizationInstance::CreateStatic(&FChaosVDGeometryComponentCustomization::MakeInstance));
 	DetailsView->RegisterInstancedCustomPropertyLayout(UChaosVDStaticMeshComponent::StaticClass(), FOnGetDetailCustomizationInstance::CreateStatic(&FChaosVDGeometryComponentCustomization::MakeInstance));
@@ -568,6 +593,19 @@ void SChaosVDMainTab::SetCustomPropertyLayouts(IDetailsView* DetailsView)
 	DetailsView->RegisterInstancedCustomPropertyTypeLayout(TEXT("ChaosVDCollisionResponseParams"), FOnGetPropertyTypeCustomizationInstance::CreateStatic(&FChaosVDCollisionResponseParamsCustomization::MakeInstance, StaticCastWeakPtr<SChaosVDMainTab>(AsWeak())));
 	DetailsView->RegisterInstancedCustomPropertyTypeLayout(TEXT("ChaosVDCollisionObjectQueryParams"), FOnGetPropertyTypeCustomizationInstance::CreateStatic(&FChaosVDCollisionObjectParamsCustomization::MakeInstance,StaticCastWeakPtr<SChaosVDMainTab>(AsWeak())));
 	DetailsView->RegisterInstancedCustomPropertyTypeLayout(TEXT("ChaosVDShapeCollisionData"), FOnGetPropertyTypeCustomizationInstance::CreateStatic(&FChaosVDShapeDataCustomization::MakeInstance, StaticCastWeakPtr<SChaosVDMainTab>(AsWeak())));
+
+	// We don't need to validate properties, and trying to do so seems to be costing between 2-4ms per tick!
+	DetailsView->SetCustomValidatePropertyNodesFunction(FOnValidateDetailsViewPropertyNodes::CreateLambda([](const FRootPropertyNodeList& Root){ return true;}));
+	
+	// We need to keep a weak ptr array of any panel we customized so we can apply any customization coming from late initialized extensions
+	// (if the details panels are still alive)
+	CustomizedDetailsPanels.Add(StaticCastWeakPtr<IDetailsView>(DetailsView->AsWeak()));
+
+	FChaosVDExtensionsManager::Get().EnumerateExtensions([this, DetailsView](const TSharedRef<FChaosVDExtension>& Extension)
+	{
+		Extension->SetCustomPropertyLayouts(DetailsView, StaticCastSharedRef<SChaosVDMainTab>(AsShared()));
+		return true;
+	});
 }
 
 void SChaosVDMainTab::BrowseChaosVDRecordingFromFolder(FStringView FolderPath)
diff --git a/Engine/Plugins/ChaosVD/Source/ChaosVD/Private/Widgets/SChaosVDSceneQueryBrowser.cpp b/Engine/Plugins/ChaosVD/Source/ChaosVD/Private/Widgets/SChaosVDSceneQueryBrowser.cpp
--- a/Engine/Plugins/ChaosVD/Source/ChaosVD/Private/Widgets/SChaosVDSceneQueryBrowser.cpp
+++ b/Engine/Plugins/ChaosVD/Source/ChaosVD/Private/Widgets/SChaosVDSceneQueryBrowser.cpp
@@ -11,7 +11,7 @@
 #include "EditorModeManager.h"
 #include "EditorViewportClient.h"
 #include "SChaosVDGameFramesPlaybackControls.h"
-#include "SChaosVDMainTab.h"
+#include "Widgets/SChaosVDMainTab.h"
 #include "SChaosVDSceneQueryTree.h"
 #include "SChaosVDWarningMessageBox.h"
 #include "Settings/ChaosVDSceneQueryVisualizationSettings.h"
@@ -751,7 +751,7 @@ void SChaosVDSceneQueryBrowser::HandleTreeItemFocused(const TSharedPtr<FChaosVDS
 		return;
 	}
 
-	TSharedPtr<FChaosVDRecording> RecordedData = ScenePtr->LoadedRecording;
+	TSharedPtr<FChaosVDRecording> RecordedData = ScenePtr->GetLoadedRecording();
 	if (!RecordedData)
 	{
 		return;
diff --git a/Engine/Plugins/ChaosVD/Source/ChaosVD/Public/Actors/ChaosVDDataContainerBaseActor.h b/Engine/Plugins/ChaosVD/Source/ChaosVD/Public/Actors/ChaosVDDataContainerBaseActor.h
new file mode 100644
--- /dev/null
+++ b/Engine/Plugins/ChaosVD/Source/ChaosVD/Public/Actors/ChaosVDDataContainerBaseActor.h
@@ -0,0 +1,61 @@
+// Copyright Epic Games, Inc. All Rights Reserved.
+
+#pragma once
+
+#include "GameFramework/Actor.h"
+#include "ChaosVDDataContainerBaseActor.generated.h"
+
+class FChaosVDExtension;
+class FChaosVDScene;
+
+struct FChaosVDSolverFrameData;
+struct FChaosVDGameFrameData;
+
+/** Base class for any CVD actor that will contain frame related data (either solver frame or game frame) */
+UCLASS(Abstract, NotBlueprintable, NotPlaceable)
+class CHAOSVD_API AChaosVDDataContainerBaseActor : public AActor
+{
+	GENERATED_BODY()
+
+public:
+	AChaosVDDataContainerBaseActor();
+	virtual ~AChaosVDDataContainerBaseActor() override;
+
+	virtual void UpdateFromNewGameFrameData(const FChaosVDGameFrameData& InGameFrameData);
+
+	virtual void UpdateFromNewSolverFrameData(const FChaosVDSolverFrameData& InSolverFrameData);
+
+	virtual void Destroyed() override;
+
+	virtual void SetScene(TWeakPtr<FChaosVDScene> InScene)
+	{
+		SceneWeakPtr = InScene;
+	}
+	
+	void SetSolverID(int32 InSolverID);
+
+	virtual const FTransform& GetSimulationTransform() const
+	{
+		return FTransform::Identity;
+	}
+
+	virtual bool IsVisible() const
+	{
+		return true;
+	}
+
+	virtual void PostActorCreated() override;
+	
+	TWeakPtr<FChaosVDScene> GetScene() const
+	{
+		return SceneWeakPtr;
+	}
+	
+protected:
+
+	void HandlePostInitializationExtensionRegistered(const TSharedRef<FChaosVDExtension>& NewExtension);
+	
+	TWeakPtr<FChaosVDScene> SceneWeakPtr;
+
+	int32 SolverDataID = INDEX_NONE;
+};
diff --git a/Engine/Plugins/ChaosVD/Source/ChaosVD/Public/ChaosVDModule.h b/Engine/Plugins/ChaosVD/Source/ChaosVD/Public/ChaosVDModule.h
--- a/Engine/Plugins/ChaosVD/Source/ChaosVD/Public/ChaosVDModule.h
+++ b/Engine/Plugins/ChaosVD/Source/ChaosVD/Public/ChaosVDModule.h
@@ -4,6 +4,8 @@
 
 #include "Templates/SharedPointer.h"
 #include "Modules/ModuleManager.h"
+#include "ExtensionsSystem/ChaosVDExtension.h"
+#include "Features/IModularFeatures.h"
 #include "Misc/Guid.h"
 
 class SWidget;
@@ -14,7 +16,6 @@ class SDockTab;
 class FSpawnTabArgs;
 struct FGuid;
 
-/** Main module class for the Chaos Visual Debugger editor */
 /** Main module class for the Chaos Visual Debugger editor */
 DECLARE_LOG_CATEGORY_EXTERN(LogChaosVDEditor, Log, Log);
 class FChaosVDModule : public IModuleInterface
@@ -30,7 +31,7 @@ public:
 	virtual void ShutdownModule() override;
 
 	/** Returns the Trace manager instance used by CVD */
-	TSharedPtr<FChaosVDTraceManager>& GetTraceManager() { return ChaosVDTraceManager; };
+	TSharedPtr<FChaosVDTraceManager>& GetTraceManager() { return ChaosVDTraceManager; }
 
 	/** Spawns a new instance of CVD */
 	void SpawnCVDTab();
@@ -63,3 +64,5 @@ private:
 
 	bool bIsShuttingDown = false;
 };
+
+
diff --git a/Engine/Plugins/ChaosVD/Source/ChaosVDBuiltInExtensions/Private/GenericDebugDraw/Components/ChaosVDGenericDebugDrawDataComponent.cpp b/Engine/Plugins/ChaosVD/Source/ChaosVDBuiltInExtensions/Private/GenericDebugDraw/Components/ChaosVDGenericDebugDrawDataComponent.cpp
new file mode 100644
--- /dev/null
+++ b/Engine/Plugins/ChaosVD/Source/ChaosVDBuiltInExtensions/Private/GenericDebugDraw/Components/ChaosVDGenericDebugDrawDataComponent.cpp
@@ -0,0 +1,32 @@
+// Copyright Epic Games, Inc. All Rights Reserved.
+
+#include "ChaosVDGenericDebugDrawDataComponent.h"
+
+#include "ChaosVDRecording.h"
+
+UChaosVDGenericDebugDrawDataComponent::UChaosVDGenericDebugDrawDataComponent()
+{
+	SetCanEverAffectNavigation(false);
+	bNavigationRelevant = false;
+	PrimaryComponentTick.bCanEverTick = false;
+}
+
+void UChaosVDGenericDebugDrawDataComponent::UpdateFromNewGameFrameData(const FChaosVDGameFrameData& InGameFrameData)
+{
+	if (TSharedPtr<FChaosVDDebugShapeDataContainer> DebugDrawShapeData = InGameFrameData.GetCustomDataHandler().GetData<FChaosVDDebugShapeDataContainer>())
+	{
+		CopyDataFromSourceMap<FChaosVDDebugDrawBoxDataWrapper>(DebugDrawShapeData->RecordedDebugDrawBoxesBySolverID, DebugDrawBoxes, SolverID);
+		CopyDataFromSourceMap<FChaosVDDebugDrawLineDataWrapper>(DebugDrawShapeData->RecordedDebugDrawLinesBySolverID, DebugDrawLines, SolverID);
+		CopyDataFromSourceMap<FChaosVDDebugDrawSphereDataWrapper>(DebugDrawShapeData->RecordedDebugDrawSpheresBySolverID, DebugDrawSpheres, SolverID);
+		CopyDataFromSourceMap<FChaosVDDebugDrawImplicitObjectDataWrapper>(DebugDrawShapeData->RecordedDebugDrawImplicitObjectsBySolverID, DebugDrawImplicitObjects, SolverID);
+	}
+}
+
+void UChaosVDGenericDebugDrawDataComponent::ClearData()
+{
+	DebugDrawBoxes.Reset();
+	DebugDrawLines.Reset();
+	DebugDrawSpheres.Reset();
+	DebugDrawLines.Reset();
+	DebugDrawImplicitObjects.Reset();
+}
diff --git a/Engine/Plugins/ChaosVD/Source/ChaosVDBuiltInExtensions/Private/GenericDebugDraw/Trace/DataProcessors/ChaosVDDebugDrawBoxDataProcessor.cpp b/Engine/Plugins/ChaosVD/Source/ChaosVDBuiltInExtensions/Private/GenericDebugDraw/Trace/DataProcessors/ChaosVDDebugDrawBoxDataProcessor.cpp
new file mode 100644
--- /dev/null
+++ b/Engine/Plugins/ChaosVD/Source/ChaosVDBuiltInExtensions/Private/GenericDebugDraw/Trace/DataProcessors/ChaosVDDebugDrawBoxDataProcessor.cpp
@@ -0,0 +1,38 @@
+// Copyright Epic Games, Inc. All Rights Reserved.
+
+#include "ChaosVDDebugDrawBoxDataProcessor.h"
+#include "DataWrappers/ChaosVDDebugShapeDataWrapper.h"
+#include "Visualizers/ChaosVDDebugDrawUtils.h"
+
+
+FChaosVDDebugDrawBoxDataProcessor::FChaosVDDebugDrawBoxDataProcessor() : FChaosVDDataProcessorBase(FChaosVDDebugDrawBoxDataWrapper::WrapperTypeName)
+{
+}
+
+bool FChaosVDDebugDrawBoxDataProcessor::ProcessRawData(const TArray<uint8>& InData)
+{
+	FChaosVDDataProcessorBase::ProcessRawData(InData);
+
+	const TSharedPtr<FChaosVDTraceProvider> ProviderSharedPtr = TraceProvider.Pin();
+	if (!ensure(ProviderSharedPtr.IsValid()))
+	{
+		return false;
+	}
+
+	const TSharedPtr<FChaosVDDebugDrawBoxDataWrapper> DebugDrawData = MakeShared<FChaosVDDebugDrawBoxDataWrapper>();
+	const bool bSuccess = Chaos::VisualDebugger::ReadDataFromBuffer(InData, *DebugDrawData, ProviderSharedPtr.ToSharedRef());
+	
+	if (bSuccess)
+	{
+		if (TSharedPtr<FChaosVDGameFrameData> CurrentFrameData = ProviderSharedPtr->GetCurrentGameFrame().Pin())
+		{
+			if (TSharedPtr<FChaosVDDebugShapeDataContainer> DebugDrawShapeData = CurrentFrameData->GetCustomDataHandler().GetOrAddDefaultData<FChaosVDDebugShapeDataContainer>())
+			{
+				DebugDrawShapeData->RecordedDebugDrawBoxesBySolverID.FindOrAdd(DebugDrawData->SolverID).Add(DebugDrawData);
+				CurrentFrameData->MarkDirty();
+			}
+		}
+	}
+
+	return bSuccess;
+}
diff --git a/Engine/Plugins/ChaosVD/Source/ChaosVDBuiltInExtensions/Private/GenericDebugDraw/Trace/DataProcessors/ChaosVDDebugDrawImplicitObjectDataProcessor.cpp b/Engine/Plugins/ChaosVD/Source/ChaosVDBuiltInExtensions/Private/GenericDebugDraw/Trace/DataProcessors/ChaosVDDebugDrawImplicitObjectDataProcessor.cpp
new file mode 100644
--- /dev/null
+++ b/Engine/Plugins/ChaosVD/Source/ChaosVDBuiltInExtensions/Private/GenericDebugDraw/Trace/DataProcessors/ChaosVDDebugDrawImplicitObjectDataProcessor.cpp
@@ -0,0 +1,38 @@
+// Copyright Epic Games, Inc. All Rights Reserved.
+
+#include "ChaosVDDebugDrawImplicitObjectDataProcessor.h"
+#include "DataWrappers/ChaosVDDebugShapeDataWrapper.h"
+#include "Visualizers/ChaosVDDebugDrawUtils.h"
+
+
+FChaosVDDebugDrawImplicitObjectDataProcessor::FChaosVDDebugDrawImplicitObjectDataProcessor() : FChaosVDDataProcessorBase(FChaosVDDebugDrawImplicitObjectDataWrapper::WrapperTypeName)
+{
+}
+
+bool FChaosVDDebugDrawImplicitObjectDataProcessor::ProcessRawData(const TArray<uint8>& InData)
+{
+	FChaosVDDataProcessorBase::ProcessRawData(InData);
+
+	const TSharedPtr<FChaosVDTraceProvider> ProviderSharedPtr = TraceProvider.Pin();
+	if (!ensure(ProviderSharedPtr.IsValid()))
+	{
+		return false;
+	}
+
+	const TSharedPtr<FChaosVDDebugDrawImplicitObjectDataWrapper> DebugDrawData = MakeShared<FChaosVDDebugDrawImplicitObjectDataWrapper>();
+	const bool bSuccess = Chaos::VisualDebugger::ReadDataFromBuffer(InData, *DebugDrawData, ProviderSharedPtr.ToSharedRef());
+	
+	if (bSuccess)
+	{
+		if (TSharedPtr<FChaosVDGameFrameData> CurrentFrameData = ProviderSharedPtr->GetCurrentGameFrame().Pin())
+		{
+			if (TSharedPtr<FChaosVDDebugShapeDataContainer> DebugDrawShapeData = CurrentFrameData->GetCustomDataHandler().GetOrAddDefaultData<FChaosVDDebugShapeDataContainer>())
+			{	
+				DebugDrawShapeData->RecordedDebugDrawImplicitObjectsBySolverID.FindOrAdd(DebugDrawData->SolverID).Add(DebugDrawData);
+				CurrentFrameData->MarkDirty();
+			}
+		}
+	}
+
+	return bSuccess;
+}
diff --git a/Engine/Plugins/ChaosVD/Source/ChaosVDBuiltInExtensions/Private/GenericDebugDraw/Trace/DataProcessors/ChaosVDDebugDrawLineDataProcessor.cpp b/Engine/Plugins/ChaosVD/Source/ChaosVDBuiltInExtensions/Private/GenericDebugDraw/Trace/DataProcessors/ChaosVDDebugDrawLineDataProcessor.cpp
new file mode 100644
--- /dev/null
+++ b/Engine/Plugins/ChaosVD/Source/ChaosVDBuiltInExtensions/Private/GenericDebugDraw/Trace/DataProcessors/ChaosVDDebugDrawLineDataProcessor.cpp
@@ -0,0 +1,38 @@
+// Copyright Epic Games, Inc. All Rights Reserved.
+
+#include "ChaosVDDebugDrawLineDataProcessor.h"
+#include "DataWrappers/ChaosVDDebugShapeDataWrapper.h"
+#include "Visualizers/ChaosVDDebugDrawUtils.h"
+
+
+FChaosVDDebugDrawLineDataProcessor::FChaosVDDebugDrawLineDataProcessor() : FChaosVDDataProcessorBase(FChaosVDDebugDrawLineDataWrapper::WrapperTypeName)
+{
+}
+
+bool FChaosVDDebugDrawLineDataProcessor::ProcessRawData(const TArray<uint8>& InData)
+{
+	FChaosVDDataProcessorBase::ProcessRawData(InData);
+
+	const TSharedPtr<FChaosVDTraceProvider> ProviderSharedPtr = TraceProvider.Pin();
+	if (!ensure(ProviderSharedPtr.IsValid()))
+	{
+		return false;
+	}
+
+	const TSharedPtr<FChaosVDDebugDrawLineDataWrapper> DebugDrawData = MakeShared<FChaosVDDebugDrawLineDataWrapper>();
+	const bool bSuccess = Chaos::VisualDebugger::ReadDataFromBuffer(InData, *DebugDrawData, ProviderSharedPtr.ToSharedRef());
+	
+	if (bSuccess)
+	{
+		if (TSharedPtr<FChaosVDGameFrameData> CurrentFrameData = ProviderSharedPtr->GetCurrentGameFrame().Pin())
+		{
+			if (TSharedPtr<FChaosVDDebugShapeDataContainer> DebugDrawShapeData = CurrentFrameData->GetCustomDataHandler().GetOrAddDefaultData<FChaosVDDebugShapeDataContainer>())
+			{
+				DebugDrawShapeData->RecordedDebugDrawLinesBySolverID.FindOrAdd(DebugDrawData->SolverID).Add(DebugDrawData);
+				CurrentFrameData->MarkDirty();
+			}
+		}
+	}
+
+	return bSuccess;
+}
\ No newline at end of file
diff --git a/Engine/Plugins/ChaosVD/Source/ChaosVDBuiltInExtensions/Private/GenericDebugDraw/Trace/DataProcessors/ChaosVDDebugDrawSphereDataProcessor.cpp b/Engine/Plugins/ChaosVD/Source/ChaosVDBuiltInExtensions/Private/GenericDebugDraw/Trace/DataProcessors/ChaosVDDebugDrawSphereDataProcessor.cpp
new file mode 100644
--- /dev/null
+++ b/Engine/Plugins/ChaosVD/Source/ChaosVDBuiltInExtensions/Private/GenericDebugDraw/Trace/DataProcessors/ChaosVDDebugDrawSphereDataProcessor.cpp
@@ -0,0 +1,38 @@
+// Copyright Epic Games, Inc. All Rights Reserved.
+
+#include "ChaosVDDebugDrawSphereDataProcessor.h"
+#include "DataWrappers/ChaosVDDebugShapeDataWrapper.h"
+#include "Visualizers/ChaosVDDebugDrawUtils.h"
+
+
+FChaosVDDebugDrawSphereDataProcessor::FChaosVDDebugDrawSphereDataProcessor() : FChaosVDDataProcessorBase(FChaosVDDebugDrawSphereDataWrapper::WrapperTypeName)
+{
+}
+
+bool FChaosVDDebugDrawSphereDataProcessor::ProcessRawData(const TArray<uint8>& InData)
+{
+	FChaosVDDataProcessorBase::ProcessRawData(InData);
+
+	const TSharedPtr<FChaosVDTraceProvider> ProviderSharedPtr = TraceProvider.Pin();
+	if (!ensure(ProviderSharedPtr.IsValid()))
+	{
+		return false;
+	}
+
+	const TSharedPtr<FChaosVDDebugDrawSphereDataWrapper> DebugDrawData = MakeShared<FChaosVDDebugDrawSphereDataWrapper>();
+	const bool bSuccess = Chaos::VisualDebugger::ReadDataFromBuffer(InData, *DebugDrawData, ProviderSharedPtr.ToSharedRef());
+	
+	if (bSuccess)
+	{
+		if (TSharedPtr<FChaosVDGameFrameData> CurrentFrameData = ProviderSharedPtr->GetCurrentGameFrame().Pin())
+		{
+			if (TSharedPtr<FChaosVDDebugShapeDataContainer> DebugDrawShapeData = CurrentFrameData->GetCustomDataHandler().GetOrAddDefaultData<FChaosVDDebugShapeDataContainer>())
+			{
+				DebugDrawShapeData->RecordedDebugDrawSpheresBySolverID.FindOrAdd(DebugDrawData->SolverID).Add(DebugDrawData);
+				CurrentFrameData->MarkDirty();
+			}
+		}
+	}
+
+	return bSuccess;
+}
diff --git a/Engine/Source/Runtime/Experimental/ChaosVDData/Private/ChaosVDRecording.cpp b/Engine/Source/Runtime/Experimental/ChaosVDData/Private/ChaosVDRecording.cpp
--- a/Engine/Source/Runtime/Experimental/ChaosVDData/Private/ChaosVDRecording.cpp
+++ b/Engine/Source/Runtime/Experimental/ChaosVDData/Private/ChaosVDRecording.cpp
@@ -4,6 +4,22 @@
 #include "Chaos/ImplicitObject.h"
 #include "ChaosVisualDebugger/ChaosVDSerializedNameTable.h"
 
+FName FChaosVDCustomUserDataHandle::GetTypeName() const
+{
+	const UStruct* Struct = DataStruct->GetStruct();
+	if (ensure(Struct != nullptr))
+	{
+		return Struct->GetFName();
+	}
+
+	return NAME_Name;
+}
+
+void FChaosVDCustomFrameData::AddData(const FChaosVDCustomUserDataHandle& InData)
+{
+	CustomDataHandlesByType.Add(InData.GetTypeName(), InData);
+}
+
 FChaosVDRecording::FChaosVDRecording()
 {
 	NameTable = MakeShared<Chaos::VisualDebugger::FChaosVDSerializableNameTable>();
@@ -379,7 +395,7 @@ void FChaosVDRecording::CollapseSolverFramesRange_AssumesLocked(int32 SolverID,
 
 			if (SolverFrameData->SolverSteps.Num() > 0)
 			{
-				for (const FChaosVDStepData& StepData : SolverFrameData->SolverSteps)
+				for (const FChaosVDFrameStageData& StepData : SolverFrameData->SolverSteps)
 				{
 					for (const TSharedPtr<FChaosVDParticleDataWrapper>& ParticleData : StepData.RecordedParticlesData)
 					{
@@ -411,17 +427,17 @@ void FChaosVDRecording::CollapseSolverFramesRange_AssumesLocked(int32 SolverID,
 				OutCollapsedFrameData.SimulationTransform = SolverFrameData->SimulationTransform;
 				OutCollapsedFrameData.DebugFName = SolverFrameData->DebugFName;
 
-				FChaosVDStepData CollapsedStepData;
-				CollapsedStepData.StepName = TEXT("Auto Generated Stage");
+				FChaosVDFrameStageData CollapsedStageData;
+				CollapsedStageData.StepName = TEXT("Auto Generated Stage");
 
 				// Although this is a generated on load stage, is based on explicit stage data, therefore it needs to be treated as such
-				EnumAddFlags(CollapsedStepData.StageFlags, EChaosVDSolverStageFlags::ExplicitStage);
+				EnumAddFlags(CollapsedStageData.StageFlags, EChaosVDSolverStageFlags::ExplicitStage);
 
-				ParticlesOnCurrentGeneratedKeyframe.GenerateValueArray(CollapsedStepData.RecordedParticlesData);
+				ParticlesOnCurrentGeneratedKeyframe.GenerateValueArray(CollapsedStageData.RecordedParticlesData);
 
 				ParticlesOnCurrentGeneratedKeyframe.Reset();
 				
-				OutCollapsedFrameData.SolverSteps.Add(MoveTemp(CollapsedStepData));
+				OutCollapsedFrameData.SolverSteps.Add(MoveTemp(CollapsedStageData));
 			}
 		}
 	}
