[ChaosVD] Fix for an issue where the incorrect meshes representing a landscape tile will be shown if the world is loaded and load rapidly while scrubbing.

The change itself is trivial. It is correcting a mistake where we are storing a uint32 value in a uint8 property we use as a key to identify geometry. This CL updates the property to be uint32 as it should ahve been.

diff --git a/Engine/Plugins/ChaosVD/Source/ChaosVD/Private/Components/ChaosVDStaticMeshComponent.h b/Engine/Plugins/ChaosVD/Source/ChaosVD/Private/Components/ChaosVDStaticMeshComponent.h
--- a/Engine/Plugins/ChaosVD/Source/ChaosVD/Private/Components/ChaosVDStaticMeshComponent.h
+++ b/Engine/Plugins/ChaosVD/Source/ChaosVD/Private/Components/ChaosVDStaticMeshComponent.h
@@ -72,7 +72,7 @@ protected:
 	bool UpdateGeometryKey(uint32 NewHandleGeometryKey);
 
 	EChaosVDMeshAttributesFlags MeshComponentAttributeFlags = EChaosVDMeshAttributesFlags::None;
-	uint8 CurrentGeometryKey = 0;
+	uint32 CurrentGeometryKey = 0;
 	bool bIsMeshReady = false;
 	bool bIsOwningParticleSelected = false;
 	FChaosVDMeshReadyDelegate MeshReadyDelegate;
