[ChaosVD] Workaround fix for an issue where visualized landscapes will have the incorrect static mesh if the recorded world is loaded/unleaded quickly.

There is a bug in the ChaosVDStaticMesh component, only used to represent recorded landscape collision in CVD. As the fix will require break binary compatibility, this change just make it so we use Instanced Static Mesh Component for them, which are not affected by the issue.

diff --git a/Engine/Plugins/ChaosVD/Source/ChaosVD/Private/ChaosVDParticleActor.cpp b/Engine/Plugins/ChaosVD/Source/ChaosVD/Private/ChaosVDParticleActor.cpp
--- a/Engine/Plugins/ChaosVD/Source/ChaosVD/Private/ChaosVDParticleActor.cpp
+++ b/Engine/Plugins/ChaosVD/Source/ChaosVD/Private/ChaosVDParticleActor.cpp
@@ -254,7 +254,12 @@ void AChaosVDParticleActor::UpdateGeometry(const Chaos::FConstImplicitObjectPtr&
 	OutExtractedGeometryDataHandles.Reserve(ObjectsToGenerateNum);
 
 	// Heightfields need to be created as Static meshes and use normal Static Mesh components because we need LODs for them due to their high triangle count
-	const bool bHasToUseStaticMeshComponent = Chaos::VisualDebugger::Cvars::bForceStaticMeshComponentUse || FChaosVDGeometryBuilder::DoesImplicitContainType(InImplicitObject, Chaos::ImplicitObjectType::HeightField);
+	//const bool bHasToUseStaticMeshComponent = Chaos::VisualDebugger::Cvars::bForceStaticMeshComponentUse || FChaosVDGeometryBuilder::DoesImplicitContainType(InImplicitObject, Chaos::ImplicitObjectType::HeightField);
+	
+	// There is an issue (UE-229467) with ChaosVDStaticMeshComponent where the geometry key gets stored in an incorrect type, and in consequence, pending update to these components are not cleared out when the component
+    // is returned to the pool. As fixing that will break binary compatibility, for now we make landscapes use ISMC instead
+    constexpr bool bHasToUseStaticMeshComponent = false;
+
 	constexpr int32 LODsToGenerateNum = 3;
 	constexpr int32 LODsToGenerateNumForInstancedStaticMesh = 0;
 
