[ChaosVD] Enabled unsafe type casting as errors in the chaosVDModule, and fixed all the existogn errors

diff --git a/Engine/Plugins/ChaosVD/Source/ChaosVD/ChaosVD.Build.cs b/Engine/Plugins/ChaosVD/Source/ChaosVD/ChaosVD.Build.cs
--- a/Engine/Plugins/ChaosVD/Source/ChaosVD/ChaosVD.Build.cs
+++ b/Engine/Plugins/ChaosVD/Source/ChaosVD/ChaosVD.Build.cs
@@ -53,6 +53,8 @@ public class ChaosVD : ModuleRules
 			}
 			);
 
+		UnsafeTypeCastWarningLevel = WarningLevel.Error;
+
 		SetupModulePhysicsSupport(Target);
 	}
 }
diff --git a/Engine/Plugins/ChaosVD/Source/ChaosVD/Private/Actors/ChaosVDGeometryContainer.cpp b/Engine/Plugins/ChaosVD/Source/ChaosVD/Private/Actors/ChaosVDGeometryContainer.cpp
--- a/Engine/Plugins/ChaosVD/Source/ChaosVD/Private/Actors/ChaosVDGeometryContainer.cpp
+++ b/Engine/Plugins/ChaosVD/Source/ChaosVD/Private/Actors/ChaosVDGeometryContainer.cpp
@@ -20,7 +20,7 @@ void AChaosVDGeometryContainer::CleanUp()
 	GetComponents<UMeshComponent>(ComponentsToDestroy);
 
 	constexpr float AmountOfWork = 1.0f;
-	const float PercentagePerElement = 1.0f / ComponentsToDestroy.Num();
+	const float PercentagePerElement = 1.0f / static_cast<float>(ComponentsToDestroy.Num());
 
 	FScopedSlowTask CleaningDataSlowTask(AmountOfWork, NSLOCTEXT("ChaosVisualDebugger", "CleaningupGeometryData", "Clearing Geoemtry Data ..."));
 	CleaningDataSlowTask.MakeDialog();
diff --git a/Engine/Plugins/ChaosVD/Source/ChaosVD/Private/Actors/ChaosVDSolverInfoActor.cpp b/Engine/Plugins/ChaosVD/Source/ChaosVD/Private/Actors/ChaosVDSolverInfoActor.cpp
--- a/Engine/Plugins/ChaosVD/Source/ChaosVD/Private/Actors/ChaosVDSolverInfoActor.cpp
+++ b/Engine/Plugins/ChaosVD/Source/ChaosVD/Private/Actors/ChaosVDSolverInfoActor.cpp
@@ -206,7 +206,7 @@ void AChaosVDSolverInfoActor::SetIsTemporarilyHiddenInEditor(bool bIsHidden)
 		Super::SetIsTemporarilyHiddenInEditor(bIsHidden);
 
 		constexpr float AmountOfWork = 1.0f;
-		const float PercentagePerElement = 1.0f / SolverParticlesByID.Num();
+		const float PercentagePerElement = 1.0f / static_cast<float>(SolverParticlesByID.Num());
 
 		FScopedSlowTask VisibilityUpdateProgress(AmountOfWork, LOCTEXT("UpdatingParticlesVisisibility", "Updating Particles Visibility ..."));
 		VisibilityUpdateProgress.MakeDialog();
@@ -247,7 +247,7 @@ void AChaosVDSolverInfoActor::Destroyed()
 	if (SolverParticlesByID.Num() > 0)
 	{
 		constexpr float AmountOfWork = 1.0f;
-		const float PercentagePerElement = 1.0f / SolverParticlesByID.Num();
+		const float PercentagePerElement = 1.0f / static_cast<float>(SolverParticlesByID.Num());
 
 		FScopedSlowTask CleaningParticleDataSlowTask(AmountOfWork, LOCTEXT("CleaningParticleDataMessage", "Cleaning Up Particle Data ..."));
 		CleaningParticleDataSlowTask.MakeDialog();
diff --git a/Engine/Plugins/ChaosVD/Source/ChaosVD/Private/ChaosVDGeometryBuilder.cpp b/Engine/Plugins/ChaosVD/Source/ChaosVD/Private/ChaosVDGeometryBuilder.cpp
--- a/Engine/Plugins/ChaosVD/Source/ChaosVD/Private/ChaosVDGeometryBuilder.cpp
+++ b/Engine/Plugins/ChaosVD/Source/ChaosVD/Private/ChaosVDGeometryBuilder.cpp
@@ -269,7 +269,7 @@ void FChaosVDGeometryBuilder::DeInitialize()
 
 	constexpr float MaxAmountOfWork = 1.0f;
 	const int32 WorkRemaining = GeometryBeingGeneratedByKey.Num() + StaticMeshCacheMap.Num();
-	const float PercentagePerElement = 1.0f / WorkRemaining;
+	const float PercentagePerElement = 1.0f / static_cast<float>(WorkRemaining);
 
 	FScopedSlowTask CleaningGeometrySlowTask(MaxAmountOfWork, NSLOCTEXT("ChaosVisualDebugger", "DeInitializeGeometrybuilderSlowTask", "Deinitialiing GeometryBuilder"));
 
@@ -682,8 +682,8 @@ TSharedPtr<UE::Geometry::FMeshShapeGenerator> FChaosVDGeometryBuilder::CreateMes
 			{
 				TSharedPtr<UE::Geometry::FSphereGenerator> SphereGen = MakeShared<UE::Geometry::FSphereGenerator>();
 				SphereGen->Radius = Sphere->GetRadiusf();
-				SphereGen->NumTheta = 25 * SimpleShapesComplexityFactor;
-				SphereGen->NumPhi = 25 * SimpleShapesComplexityFactor;
+				SphereGen->NumTheta = static_cast<int32>(25 * SimpleShapesComplexityFactor);
+				SphereGen->NumPhi =  static_cast<int32>(25 * SimpleShapesComplexityFactor);
 				SphereGen->bPolygroupPerQuad = false;
 
 				return SphereGen;	
@@ -710,8 +710,8 @@ TSharedPtr<UE::Geometry::FMeshShapeGenerator> FChaosVDGeometryBuilder::CreateMes
 				TSharedPtr<UE::Geometry::FCapsuleGenerator> CapsuleGenerator = MakeShared<UE::Geometry::FCapsuleGenerator>();
 				CapsuleGenerator->Radius = FMath::Max(FMathf::ZeroTolerance, Capsule->GetRadiusf());
 				CapsuleGenerator->SegmentLength = FMath::Max(FMathf::ZeroTolerance, Capsule->GetSegment().GetLength());
-				CapsuleGenerator->NumHemisphereArcSteps = 12 * SimpleShapesComplexityFactor;
-				CapsuleGenerator->NumCircleSteps = 12 * SimpleShapesComplexityFactor;
+				CapsuleGenerator->NumHemisphereArcSteps = static_cast<int32>(12 * SimpleShapesComplexityFactor);
+				CapsuleGenerator->NumCircleSteps = static_cast<int32>(12 * SimpleShapesComplexityFactor);
 
 				return CapsuleGenerator;
 			}
diff --git a/Engine/Plugins/ChaosVD/Source/ChaosVD/Private/ChaosVDPlaybackController.cpp b/Engine/Plugins/ChaosVD/Source/ChaosVD/Private/ChaosVDPlaybackController.cpp
--- a/Engine/Plugins/ChaosVD/Source/ChaosVD/Private/ChaosVDPlaybackController.cpp
+++ b/Engine/Plugins/ChaosVD/Source/ChaosVD/Private/ChaosVDPlaybackController.cpp
@@ -863,7 +863,7 @@ int32 FChaosVDPlaybackController::GetFrameRateOverride() const
 void FChaosVDPlaybackController::SetFrameRateOverride(float NewFrameRateOverride)
 {
 	constexpr int32 MinimumFrameRateOverride = 1;
-	CurrentFrameRateOverride = (NewFrameRateOverride >= MinimumFrameRateOverride) ? NewFrameRateOverride : InvalidFrameRateOverride;
+	CurrentFrameRateOverride = (NewFrameRateOverride >= MinimumFrameRateOverride) ? static_cast<int32>(NewFrameRateOverride) : InvalidFrameRateOverride;
 }
 
 float FChaosVDPlaybackController::GetFrameTimeForTrack(EChaosVDTrackType TrackType, int32 TrackID, const TSharedRef<const FChaosVDTrackInfo>& InTrackInfo) const
@@ -884,7 +884,7 @@ float FChaosVDPlaybackController::GetFrameTimeForTrack(EChaosVDTrackType TrackTy
 				{
 					if (const FChaosVDSolverFrameData* FrameData = LoadedRecording->GetSolverFrameData_AssumesLocked(TrackID, InTrackInfo->CurrentFrame))
 					{
-						CurrentTargetFrameTime = FrameData->GetFrameTime();
+						CurrentTargetFrameTime = static_cast<float>(FrameData->GetFrameTime());
 					}
 					break;
 				}
@@ -892,7 +892,7 @@ float FChaosVDPlaybackController::GetFrameTimeForTrack(EChaosVDTrackType TrackTy
 				{
 					if (const FChaosVDGameFrameData* FrameData = LoadedRecording->GetGameFrameData_AssumesLocked(InTrackInfo->CurrentFrame))
 					{
-						CurrentTargetFrameTime = FrameData->GetFrameTime();
+						CurrentTargetFrameTime = static_cast<float>(FrameData->GetFrameTime());
 					}
 
 					break;
diff --git a/Engine/Plugins/ChaosVD/Source/ChaosVD/Private/ChaosVDScene.cpp b/Engine/Plugins/ChaosVD/Source/ChaosVD/Private/ChaosVDScene.cpp
--- a/Engine/Plugins/ChaosVD/Source/ChaosVD/Private/ChaosVDScene.cpp
+++ b/Engine/Plugins/ChaosVD/Source/ChaosVD/Private/ChaosVDScene.cpp
@@ -167,7 +167,7 @@ void FChaosVDScene::UpdateFromRecordedStepData(const int32 SolverID, const FChao
 	
 	{
 		constexpr float AmountOfWork = 1.0f;
-		const float PercentagePerElement = 1.0f / InRecordedStepData.RecordedParticlesData.Num();
+		const float PercentagePerElement = 1.0f / static_cast<float>(InRecordedStepData.RecordedParticlesData.Num());
 
 		const FText ProgressBarTitle = FText::Format(FTextFormat(LOCTEXT("ProcessingParticleData", "Processing Particle Data for {0} Solver with ID {1} ...")), FText::FromName(SolverSceneData->GetSolverName()), FText::AsNumber(SolverID));
 		FScopedSlowTask UpdatingSceneSlowTask(AmountOfWork, ProgressBarTitle);
@@ -410,7 +410,7 @@ void FChaosVDScene::CleanUpScene(EChaosVDSceneCleanUpOptions Options)
 	if (AvailableDataContainerActors.Num() > 0)
 	{
 		constexpr float AmountOfWork = 1.0f;
-		const float PercentagePerElement = 1.0f / AvailableDataContainerActors.Num();
+		const float PercentagePerElement = 1.0f / static_cast<float>(AvailableDataContainerActors.Num());
 
 		FScopedSlowTask CleaningSceneSlowTask(AmountOfWork, LOCTEXT("CleaningupSceneSolverMessage", "Clearing Solver Data ..."));
 		CleaningSceneSlowTask.MakeDialog();
diff --git a/Engine/Plugins/ChaosVD/Source/ChaosVD/Private/DetailsCustomizations/ChaosVDDetailsCustomizationUtils.cpp b/Engine/Plugins/ChaosVD/Source/ChaosVD/Private/DetailsCustomizations/ChaosVDDetailsCustomizationUtils.cpp
--- a/Engine/Plugins/ChaosVD/Source/ChaosVD/Private/DetailsCustomizations/ChaosVDDetailsCustomizationUtils.cpp
+++ b/Engine/Plugins/ChaosVD/Source/ChaosVD/Private/DetailsCustomizations/ChaosVDDetailsCustomizationUtils.cpp
@@ -276,7 +276,7 @@ void FChaosVDDetailsCustomizationUtils::AddWidgetRowForCheckboxValue(TAttribute<
 	.ValueContent()
 	[
 		SNew(SBox)
-		.WidthOverride(50)
+		.WidthOverride(50.0f)
 		.Content()
 		[
 			SNew(SCheckBox)
diff --git a/Engine/Plugins/ChaosVD/Source/ChaosVD/Private/DetailsCustomizations/ChaosVDQueryDataWrappersCustomizationDetails.cpp b/Engine/Plugins/ChaosVD/Source/ChaosVD/Private/DetailsCustomizations/ChaosVDQueryDataWrappersCustomizationDetails.cpp
--- a/Engine/Plugins/ChaosVD/Source/ChaosVD/Private/DetailsCustomizations/ChaosVDQueryDataWrappersCustomizationDetails.cpp
+++ b/Engine/Plugins/ChaosVD/Source/ChaosVD/Private/DetailsCustomizations/ChaosVDQueryDataWrappersCustomizationDetails.cpp
@@ -246,7 +246,7 @@ void FChaosVDCollisionObjectParamsCustomization::CustomizeChildren(TSharedRef<IP
 		[
 			SNew(SBox)
 			.IsEnabled(bIsEditable)
-			.WidthOverride(50)
+			.WidthOverride(50.0f)
 			.Content()
 			[
 				SNew(SCheckBox)
diff --git a/Engine/Plugins/ChaosVD/Source/ChaosVD/Private/DetailsCustomizations/ChaosVDShapeDataCustomization.cpp b/Engine/Plugins/ChaosVD/Source/ChaosVD/Private/DetailsCustomizations/ChaosVDShapeDataCustomization.cpp
--- a/Engine/Plugins/ChaosVD/Source/ChaosVD/Private/DetailsCustomizations/ChaosVDShapeDataCustomization.cpp
+++ b/Engine/Plugins/ChaosVD/Source/ChaosVD/Private/DetailsCustomizations/ChaosVDShapeDataCustomization.cpp
@@ -130,7 +130,7 @@ void FChaosVDShapeDataCustomization::CustomizeChildren(TSharedRef<IPropertyHandl
 	[
 		SNew(SBox)
 		.IsEnabled(bIsEditable)
-		.WidthOverride(50)
+		.WidthOverride(50.0f)
 		.VAlign(VAlign_Center)
 		.Content()
 		[
diff --git a/Engine/Plugins/ChaosVD/Source/ChaosVD/Private/Settings/ChaosVDJointConstraintVisualizationSettings.h b/Engine/Plugins/ChaosVD/Source/ChaosVD/Private/Settings/ChaosVDJointConstraintVisualizationSettings.h
--- a/Engine/Plugins/ChaosVD/Source/ChaosVD/Private/Settings/ChaosVDJointConstraintVisualizationSettings.h
+++ b/Engine/Plugins/ChaosVD/Source/ChaosVD/Private/Settings/ChaosVDJointConstraintVisualizationSettings.h
@@ -46,7 +46,7 @@ public:
 
 	/** Scale to apply to the Linear Impulse vector before draw it. */
 	UPROPERTY(EditAnywhere, config, Category=DebugDraw)
-	float LinearImpulseScale = 0.001;
+	float LinearImpulseScale = 0.001f;
 
 	/** Scale to apply to the Angular Impulse vector before draw it. */
 	UPROPERTY(EditAnywhere, config, Category=DebugDraw)
diff --git a/Engine/Plugins/ChaosVD/Source/ChaosVD/Private/Settings/ChaosVDParticleVisualizationSettings.cpp b/Engine/Plugins/ChaosVD/Source/ChaosVD/Private/Settings/ChaosVDParticleVisualizationSettings.cpp
--- a/Engine/Plugins/ChaosVD/Source/ChaosVD/Private/Settings/ChaosVDParticleVisualizationSettings.cpp
+++ b/Engine/Plugins/ChaosVD/Source/ChaosVD/Private/Settings/ChaosVDParticleVisualizationSettings.cpp
@@ -169,7 +169,7 @@ FColor FChaosDebugDrawColorsByClientServer::GetColorFromState(bool bIsServer, EC
 		return GetColorAtIntensity(bIsServer ? ServerBaseColor : ClientBaseColor, StaticStateIntensity);
 	}
 
-	const float Intensity = IntensityFactor * static_cast<int32>(State);
+	const float Intensity = IntensityFactor * static_cast<float>(State);
 	return GetColorAtIntensity(bIsServer ? ServerBaseColor : ClientBaseColor, Intensity);
 }
 
diff --git a/Engine/Plugins/ChaosVD/Source/ChaosVD/Private/Settings/ChaosVDParticleVisualizationSettings.h b/Engine/Plugins/ChaosVD/Source/ChaosVD/Private/Settings/ChaosVDParticleVisualizationSettings.h
--- a/Engine/Plugins/ChaosVD/Source/ChaosVD/Private/Settings/ChaosVDParticleVisualizationSettings.h
+++ b/Engine/Plugins/ChaosVD/Source/ChaosVD/Private/Settings/ChaosVDParticleVisualizationSettings.h
@@ -258,7 +258,7 @@ public:
 
 	/** Scale to apply to the Linear Impulse vector before draw it. Unit is g.m/s */
 	UPROPERTY(config, EditAnywhere, Category=DebugDraw)
-	float LinearImpulseScale = 0.001;
+	float LinearImpulseScale = 0.001f;
 
 	/** Scale to apply to the Angular Impulse vector before draw it. Unit is g.m2/s */
 	UPROPERTY(config, EditAnywhere, Category=DebugDraw)
diff --git a/Engine/Plugins/ChaosVD/Source/ChaosVD/Private/Visualizers/ChaosVDCharacterGroundConstraintsDataComponentVisualizer.cpp b/Engine/Plugins/ChaosVD/Source/ChaosVD/Private/Visualizers/ChaosVDCharacterGroundConstraintsDataComponentVisualizer.cpp
--- a/Engine/Plugins/ChaosVD/Source/ChaosVD/Private/Visualizers/ChaosVDCharacterGroundConstraintsDataComponentVisualizer.cpp
+++ b/Engine/Plugins/ChaosVD/Source/ChaosVD/Private/Visualizers/ChaosVDCharacterGroundConstraintsDataComponentVisualizer.cpp
@@ -185,8 +185,8 @@ void FChaosVDCharacterGroundConstraintDataComponentVisualizer::DrawConstraint(co
 	const FVector CharacterPos = CharacterParticleData->ParticlePositionRotation.MX;
 	const FVector UpDir = InConstraintData.Settings.VerticalAxis;
 
-	const float GroundDistance = InConstraintData.Data.GroundDistance;
-	const float TargetHeight = InConstraintData.Settings.TargetHeight;
+	const double GroundDistance = InConstraintData.Data.GroundDistance;
+	const double TargetHeight = InConstraintData.Settings.TargetHeight;
 
 	// TODO: The target data is valid for pre-sim positions and the force/torque for post sim
 	// but we don't have a way of differentiating here which state the particle is in so just draw
diff --git a/Engine/Plugins/ChaosVD/Source/ChaosVD/Private/Visualizers/ChaosVDDebugDrawUtils.cpp b/Engine/Plugins/ChaosVD/Source/ChaosVD/Private/Visualizers/ChaosVDDebugDrawUtils.cpp
--- a/Engine/Plugins/ChaosVD/Source/ChaosVD/Private/Visualizers/ChaosVDDebugDrawUtils.cpp
+++ b/Engine/Plugins/ChaosVD/Source/ChaosVD/Private/Visualizers/ChaosVDDebugDrawUtils.cpp
@@ -53,7 +53,7 @@ void FChaosVDDebugDrawUtils::DrawArrowVector(FPrimitiveDrawInterface* PDI, const
 	const float ProportionalArrowSize = MaxTipOfArrowSize * (ArrowLength / MaxVectorSizeForArrow);
 	const float ArrowSize = FMath::Clamp(ProportionalArrowSize, MinTipOfArrowSize, MaxTipOfArrowSize);
 	
-	DrawDirectionalArrow(PDI, ArrowTransformMatrix, Color, ArrowLength, ArrowSize, DepthPriority, Thickness);
+	DrawDirectionalArrow(PDI, ArrowTransformMatrix, Color, ArrowLength, ArrowSize, static_cast<uint8>(DepthPriority), Thickness);
 
 	IncreaseDebugDrawLineCounter();
 
@@ -82,7 +82,7 @@ void FChaosVDDebugDrawUtils::DrawPoint(FPrimitiveDrawInterface* PDI, const FVect
 		return;
 	}
 
-	PDI->DrawPoint(Location, Color, Size, DepthPriority);
+	PDI->DrawPoint(Location, Color, Size, static_cast<uint8>(DepthPriority));
 
 	if (!InDebugText.IsEmpty())
 	{
@@ -138,16 +138,17 @@ void FChaosVDDebugDrawUtils::DrawCircle(FPrimitiveDrawInterface* PDI, const FVec
 
 	// Need at least 2 sides
 	Segments = FMath::Max(Segments, 2);
-	const float	AngleDelta = 2.0f * UE_PI / Segments;
+	const float	AngleDelta = 2.0f * UE_PI / static_cast<float>(Segments);
 	FVector	LastVertex = Origin + XAxis * Radius;
 
-	PDI->AddReserveLines(DepthPriority, Segments, false, Thickness > SMALL_NUMBER);
+	PDI->AddReserveLines(static_cast<uint8>(DepthPriority), Segments, false, Thickness > SMALL_NUMBER);
 
 	for (int32 SideIndex = 0; SideIndex < Segments; SideIndex++)
 	{
-		const FVector Vertex = Origin + (XAxis * FMath::Cos(AngleDelta * (SideIndex + 1)) + YAxis * FMath::Sin(AngleDelta * (SideIndex + 1))) * Radius;
+		const float NextSideIndex = static_cast<float>(SideIndex + 1);
+		const FVector Vertex = Origin + (XAxis * FMath::Cos(AngleDelta * NextSideIndex) + YAxis * FMath::Sin(AngleDelta * NextSideIndex)) * Radius;
 
-		PDI->DrawLine(LastVertex, Vertex, Color, DepthPriority, Thickness, DepthBias, bScreenSpace);
+		PDI->DrawLine(LastVertex, Vertex, Color, static_cast<uint8>(DepthPriority), Thickness, DepthBias, bScreenSpace);
 
 		LastVertex = Vertex;
 	}
@@ -172,7 +173,7 @@ void FChaosVDDebugDrawUtils::DrawBox(FPrimitiveDrawInterface* PDI, const FVector
 
 	constexpr int32 MaxBoxLines = 12;
 	
-	PDI->AddReserveLines(DepthPriority, MaxBoxLines, false, Thickness > SMALL_NUMBER);
+	PDI->AddReserveLines(static_cast<uint8>(DepthPriority), MaxBoxLines, false, Thickness > SMALL_NUMBER);
 
 	// Array for direction offsets for the start/end point of each line
 	static TPair<FVector, FVector> VertexOffsetDirectionFromOrigin[MaxBoxLines] =
@@ -199,7 +200,7 @@ void FChaosVDDebugDrawUtils::DrawBox(FPrimitiveDrawInterface* PDI, const FVector
 		FVector LineStart = InTransform.TransformPosition(InExtents * VertexOffsetDirectionFromOrigin[BoxLineIndex].Key);
 		FVector LineEnd = InTransform.TransformPosition(InExtents * VertexOffsetDirectionFromOrigin[BoxLineIndex].Value);
 
-		PDI->DrawLine(LineStart, LineEnd, InColor, DepthPriority, Thickness,DepthBias, bScreenSpace);
+		PDI->DrawLine(LineStart, LineEnd, InColor, static_cast<uint8>(DepthPriority), Thickness,DepthBias, bScreenSpace);
 	}
 	
 	if (!InDebugText.IsEmpty())
@@ -223,7 +224,7 @@ void FChaosVDDebugDrawUtils::DrawLine(FPrimitiveDrawInterface* PDI, const FVecto
 	constexpr float DepthBias = 0.0f;
 	const bool bScreenSpace = Thickness > 0;
 
-	PDI->DrawLine(InStartPosition, InEndPosition, InColor, DepthPriority, Thickness, DepthBias, bScreenSpace);
+	PDI->DrawLine(InStartPosition, InEndPosition, InColor, static_cast<uint8>(DepthPriority), Thickness, DepthBias, bScreenSpace);
 
 	IncreaseDebugDrawLineCounter();
 
@@ -290,7 +291,7 @@ void FChaosVDDebugDrawUtils::DrawImplicitObject(FPrimitiveDrawInterface* PDI, co
 		GeometryGenerator->AdjustedTransformForImplicit(ImplicitObject, ExtractedTransform);
 		MeshGenerator->Generate();
 
-		PDI->AddReserveLines(DepthPriority, MeshGenerator->Triangles.Num() * 3, false, Thickness > SMALL_NUMBER);
+		PDI->AddReserveLines(static_cast<uint8>(DepthPriority), MeshGenerator->Triangles.Num() * 3, false, Thickness > SMALL_NUMBER);
 
 		for (const UE::Geometry::FIndex3i& Triangle : MeshGenerator->Triangles)
 		{
@@ -317,12 +318,12 @@ void FChaosVDDebugDrawUtils::DrawSphere(FPrimitiveDrawInterface* PDI, const FVec
 	// Need at least 4 segments
 	Segments = FMath::Max(Segments, 4);
 
-	const float AngleInc = 2.f * UE_PI / Segments;
+	const float AngleInc = 2.f * UE_PI / static_cast<float>(Segments);
 	int32 NumSegmentsY = Segments;
 	float Latitude = AngleInc;
 	float SinY1 = 0.0f, CosY1 = 1.0f;
 
-	PDI->AddReserveLines(DepthPriority, NumSegmentsY * Segments * 2, false, Thickness > SMALL_NUMBER);
+	PDI->AddReserveLines(static_cast<uint8>(DepthPriority), NumSegmentsY * Segments * 2, false, Thickness > SMALL_NUMBER);
 	while (NumSegmentsY--)
 	{
 		const float SinY2 = FMath::Sin(Latitude);
diff --git a/Engine/Plugins/ChaosVD/Source/ChaosVD/Private/Visualizers/ChaosVDJointConstraintsDataComponentVisualizer.cpp b/Engine/Plugins/ChaosVD/Source/ChaosVD/Private/Visualizers/ChaosVDJointConstraintsDataComponentVisualizer.cpp
--- a/Engine/Plugins/ChaosVD/Source/ChaosVD/Private/Visualizers/ChaosVDJointConstraintsDataComponentVisualizer.cpp
+++ b/Engine/Plugins/ChaosVD/Source/ChaosVD/Private/Visualizers/ChaosVDJointConstraintsDataComponentVisualizer.cpp
@@ -286,7 +286,7 @@ void FChaosVDJointConstraintsDataComponentVisualizer::DrawJointConstraint(const
 
 	// Create a sphere containing the two particle positions to use as a pseudo bounds to determine if we should draw this joint data
 	FVector DiameterSphereViewVector = ParticleData1->ParticlePositionRotation.MX - ParticleData0->ParticlePositionRotation.MX;
-	float ViewRadius = DiameterSphereViewVector.Size() * 0.5f;
+	float ViewRadius = static_cast<float>(DiameterSphereViewVector.Size() * 0.5f);
 	FVector MiddleViewPoint = ParticleData0->ParticlePositionRotation.MX + DiameterSphereViewVector.GetSafeNormal() * ViewRadius;
 
 	bool bCurrentViewContainsBothLocations = View->ViewFrustum.IntersectSphere(MiddleViewPoint, ViewRadius);
diff --git a/Engine/Plugins/ChaosVD/Source/ChaosVD/Private/Visualizers/ChaosVDSolverCollisionDataComponentVisualizer.cpp b/Engine/Plugins/ChaosVD/Source/ChaosVD/Private/Visualizers/ChaosVDSolverCollisionDataComponentVisualizer.cpp
--- a/Engine/Plugins/ChaosVD/Source/ChaosVD/Private/Visualizers/ChaosVDSolverCollisionDataComponentVisualizer.cpp
+++ b/Engine/Plugins/ChaosVD/Source/ChaosVD/Private/Visualizers/ChaosVDSolverCollisionDataComponentVisualizer.cpp
@@ -297,7 +297,7 @@ void FChaosVDSolverCollisionDataComponentVisualizer::DrawMidPhaseData(const UAct
 				if (EnumHasAnyFlags(VisualizationFlags, EChaosVDCollisionVisualizationFlags::ContactNormal))
 				{
 					FColor NormalColor = ((ManifoldPoint.ContactPoint.ContactType != EChaosVDContactPointType::EdgeEdge) ? PlaneNormalColor : EdgeNormalColor);
-					const int32 Scale = DebugDrawSettings->ContactNormalScale * ContactLenScale;
+					const float Scale = DebugDrawSettings->ContactNormalScale * ContactLenScale;
 					const FText& DebugText = DebugDrawSettings->bShowDebugText ? UEnum::GetDisplayValueAsText(EChaosVDCollisionVisualizationFlags::ContactNormal) : FText::GetEmpty();
 					FChaosVDDebugDrawUtils::DrawArrowVector(PDI, WorldPlaneLocation, WorldPlaneLocation + WorldPlaneNormal * Scale, DebugText, NormalColor, DebugDrawSettings->DepthPriority);
 				}
diff --git a/Engine/Plugins/ChaosVD/Source/ChaosVD/Private/Widgets/SChaosVDPlaybackViewport.cpp b/Engine/Plugins/ChaosVD/Source/ChaosVD/Private/Widgets/SChaosVDPlaybackViewport.cpp
--- a/Engine/Plugins/ChaosVD/Source/ChaosVD/Private/Widgets/SChaosVDPlaybackViewport.cpp
+++ b/Engine/Plugins/ChaosVD/Source/ChaosVD/Private/Widgets/SChaosVDPlaybackViewport.cpp
@@ -353,7 +353,7 @@ void SChaosVDPlaybackViewport::SetCurrentTargetFrameRateOverride(int32 NewTarget
 {
 	if (TSharedPtr<FChaosVDPlaybackController> PlaybackControllerPtr = PlaybackController.Pin())
 	{
-		PlaybackControllerPtr->SetFrameRateOverride(NewTarget);
+		PlaybackControllerPtr->SetFrameRateOverride(static_cast<float>(NewTarget));
 	}
 }
 
diff --git a/Engine/Plugins/ChaosVD/Source/ChaosVD/Private/Widgets/SChaosVDRecordedLogBrowser.h b/Engine/Plugins/ChaosVD/Source/ChaosVD/Private/Widgets/SChaosVDRecordedLogBrowser.h
--- a/Engine/Plugins/ChaosVD/Source/ChaosVD/Private/Widgets/SChaosVDRecordedLogBrowser.h
+++ b/Engine/Plugins/ChaosVD/Source/ChaosVD/Private/Widgets/SChaosVDRecordedLogBrowser.h
@@ -254,7 +254,7 @@ private:
 	TArray<TSharedPtr<FChaosVDLogViewListItem>> SelectedLogItems;
 
 	TArray<TSharedPtr<FChaosVDCachedLogItemEntry>> CachedReadLogData;
-	TMap<int32, TSharedPtr<FChaosVDLogViewListItem>> CachedLogItemsByID;
+	TMap<uint64, TSharedPtr<FChaosVDLogViewListItem>> CachedLogItemsByID;
 	TMap<FName, FCategorizedItemsContainer> CachedLogItemsByCategory;
 
 	TSharedPtr<FChaosVDLogInFlightFilteringTaskWrapper> CurrentFilteringTask;
diff --git a/Engine/Plugins/ChaosVD/Source/ChaosVD/Private/Widgets/SChaosVDTimelineWidget.cpp b/Engine/Plugins/ChaosVD/Source/ChaosVD/Private/Widgets/SChaosVDTimelineWidget.cpp
--- a/Engine/Plugins/ChaosVD/Source/ChaosVD/Private/Widgets/SChaosVDTimelineWidget.cpp
+++ b/Engine/Plugins/ChaosVD/Source/ChaosVD/Private/Widgets/SChaosVDTimelineWidget.cpp
@@ -173,7 +173,7 @@ void SChaosVDTimelineWidget::SetCurrentTimelineFrame(float FrameNumber, EChaosVD
 	{
 		if (EnumHasAnyFlags(Options, EChaosVDSetTimelineFrameFlags::BroadcastChange))
 		{
-			FrameChangedDelegate.ExecuteIfBound(FrameNumber);
+			FrameChangedDelegate.ExecuteIfBound(static_cast<int32>(FrameNumber));
 		}
 	}
 }
@@ -188,7 +188,7 @@ void SChaosVDTimelineWidget::Tick(const FGeometry& AllottedGeometry, const doubl
 		const int32 CurrentMaxFrames = GetCurrentMaxFrames();
 		if (!FMath::IsNearlyEqual(TimelineSlider->GetMinValue(), CurrentMinFrames) || !FMath::IsNearlyEqual(TimelineSlider->GetMaxValue(), CurrentMaxFrames))
 		{
-			TimelineSlider->SetMinAndMaxValues(CurrentMinFrames, CurrentMaxFrames);
+			TimelineSlider->SetMinAndMaxValues(static_cast<float>(CurrentMinFrames), static_cast<float>(CurrentMaxFrames));
 		}
 	}
 }
diff --git a/Engine/Plugins/ChaosVD/Source/ChaosVD/Public/Widgets/SChaosVDEnumFlagsMenu.h b/Engine/Plugins/ChaosVD/Source/ChaosVD/Public/Widgets/SChaosVDEnumFlagsMenu.h
--- a/Engine/Plugins/ChaosVD/Source/ChaosVD/Public/Widgets/SChaosVDEnumFlagsMenu.h
+++ b/Engine/Plugins/ChaosVD/Source/ChaosVD/Public/Widgets/SChaosVDEnumFlagsMenu.h
@@ -98,7 +98,7 @@ void SChaosVDEnumFlagsMenu<TEnum>::Construct(const FArguments& InArgs)
 
 	for (int32 EnumValueIndex = 0; EnumValueIndex < EnumValuesCount ; EnumValueIndex++)
 	{
-		const int32 Value = Enum->GetValueByIndex(EnumValueIndex);
+		const int64 Value =Enum->GetValueByIndex(EnumValueIndex);
 		const int32 Index = EnumValueIndex;
 		const bool bIsHidden = Enum->HasMetaData(TEXT("Hidden"), Index);
 		if (Value >= 0 && !bIsHidden)
