[ChaosVD] Fix deprecation warning on ChaosVD.build.cs

diff --git a/Engine/Plugins/ChaosVD/Source/ChaosVD/ChaosVD.Build.cs b/Engine/Plugins/ChaosVD/Source/ChaosVD/ChaosVD.Build.cs
--- a/Engine/Plugins/ChaosVD/Source/ChaosVD/ChaosVD.Build.cs
+++ b/Engine/Plugins/ChaosVD/Source/ChaosVD/ChaosVD.Build.cs
@@ -53,7 +53,7 @@ public class ChaosVD : ModuleRules
 			}
 			);
 
-		UnsafeTypeCastWarningLevel = WarningLevel.Error;
+		CppCompileWarningSettings.UnsafeTypeCastWarningLevel = WarningLevel.Error;
 
 		SetupModulePhysicsSupport(Target);
 	}
