[ChaosVD] Fix regression where the solver visibility button in a Solver Track would not work properly (It hides the solver onces but it gets out of sync and you can't use the button again)

This also affected the visualization of collision data.

This is fallout of the large refactor made two weeks ago to add support to the new CVD Extension system. The Solver ID value was moved to a parent class, but the getter was not updated, therefore AChaosVDSolverInfoActor was always returning INDEX_NONE (the default value).

(Submitting pre-review as the fix is trivial and the issue is bad within CVD itself)

#tests Tested that after the fix, I can use the visibility button without issues, and collision data is visualized when is available (and the corresponding visualization flag is enabled)

diff --git a/Engine/Plugins/ChaosVD/Source/ChaosVD/Private/Actors/ChaosVDSolverInfoActor.cpp b/Engine/Plugins/ChaosVD/Source/ChaosVD/Private/Actors/ChaosVDSolverInfoActor.cpp
--- a/Engine/Plugins/ChaosVD/Source/ChaosVD/Private/Actors/ChaosVDSolverInfoActor.cpp
+++ b/Engine/Plugins/ChaosVD/Source/ChaosVD/Private/Actors/ChaosVDSolverInfoActor.cpp
@@ -220,7 +220,7 @@ void AChaosVDSolverInfoActor::SetIsTemporarilyHiddenInEditor(bool bIsHidden)
 
 		if (const TSharedPtr<FChaosVDScene> CVDScene = SceneWeakPtr.Pin())
 		{
-			CVDScene->OnSolverVisibilityUpdated().Broadcast(SolverID, !bIsHidden);
+			CVDScene->OnSolverVisibilityUpdated().Broadcast(SolverDataID, !bIsHidden);
 		}
 	}
 }
@@ -292,7 +292,7 @@ FName AChaosVDSolverInfoActor::GetFolderPathForParticleType(EChaosVDParticleType
 	}
 	else
 	{
-		const FStringFormatOrderedArguments Args {SolverName.ToString(), FString::FromInt(SolverID)};
+		const FStringFormatOrderedArguments Args {SolverName.ToString(), FString::FromInt(SolverDataID)};
 		const FName ParticleFolderPath = *FPaths::Combine(FString::Format(TEXT("Solver {0} | ID {1}"), Args), UEnum::GetDisplayValueAsText(ParticleType).ToString());
 
 		FolderPathByParticlePath.Add(ParticleType, ParticleFolderPath);
diff --git a/Engine/Plugins/ChaosVD/Source/ChaosVD/Private/Actors/ChaosVDSolverInfoActor.h b/Engine/Plugins/ChaosVD/Source/ChaosVD/Private/Actors/ChaosVDSolverInfoActor.h
--- a/Engine/Plugins/ChaosVD/Source/ChaosVD/Private/Actors/ChaosVDSolverInfoActor.h
+++ b/Engine/Plugins/ChaosVD/Source/ChaosVD/Private/Actors/ChaosVDSolverInfoActor.h
@@ -30,7 +30,7 @@ public:
 
 	AChaosVDSolverInfoActor();
 
-	int32 GetSolverID() const { return SolverID; }
+	int32 GetSolverID() const { return SolverDataID; }
 
 	void SetSolverName(const FName& InSolverName);
 	const FName& GetSolverName() { return SolverName; }
@@ -86,9 +86,6 @@ protected:
 
 	FName GetFolderPathForParticleType(EChaosVDParticleType ParticleType);
 
-	UPROPERTY(VisibleAnywhere, Category="Solver Data")
-	int32 SolverID = INDEX_NONE;
-
 	UPROPERTY(VisibleAnywhere, Category="Solver Data")
 	FTransform SimulationTransform;
 
