[ChaosVD] Fix for joint (and other) constraints disappearing in auto-generated solver stages during playback in certain conditions

CVD records data in stages within a frame, but in 5.5 we added a way to record data outside a predefined solver stage. In these cases, CVD will auto generate a solver stage on the fly during load and group the data there. This caused a few issues, that were addressed prior the 5.5 release.

For example, there was an issue where particle data was getting incorrectly unloaded in these generated stages, as they didn't have all the data of a normal frame. To address that issue, we added a new flag to identify a stage as explicitly recorded, which worked.

But there is an edge case that was missed, non-particle data visualization is also affected by it, therefore we need to make sure all loaded frames from a recording are marked as explicit, and the code that updates the visualization from recorded constraint data, only is executed for explicitly recorded solver stage data.

#tests I tested that after the fix the visualization fo joints does not dissapear (seamingly) randomly, and that there is not a significant perf impact using a CVD recording of City Sample as a stress test. Lastly I tested recording sessions that uses async physics in different scenarios, ensure the playback is done correct.

diff --git a/Engine/Plugins/ChaosVD/Source/ChaosVD/Private/ChaosVDScene.cpp b/Engine/Plugins/ChaosVD/Source/ChaosVD/Private/ChaosVDScene.cpp
--- a/Engine/Plugins/ChaosVD/Source/ChaosVD/Private/ChaosVDScene.cpp
+++ b/Engine/Plugins/ChaosVD/Source/ChaosVD/Private/ChaosVDScene.cpp
@@ -213,10 +213,14 @@ void FChaosVDScene::UpdateFromRecordedStepData(const int32 SolverID, const FChao
 		}
 	}
 
-	UpdateParticlesCollisionData(InRecordedStepData, SolverID);
+	// Currently only explicitly recorded stages (no autogenerated) have valid constraint data
+	if (EnumHasAnyFlags(InRecordedStepData.StageFlags, EChaosVDSolverStageFlags::ExplicitStage))
+	{
+		UpdateParticlesCollisionData(InRecordedStepData, SolverID);
+
+		UpdateJointConstraintsData(InRecordedStepData, SolverID);
+	}
 
-	UpdateJointConstraintsData(InRecordedStepData, SolverID);
-	
 	const TMap<int32, TObjectPtr<AChaosVDParticleActor>>& AllSolverParticlesByID = SolverSceneData->GetAllParticleActorsByIDMap();
 
 	for (const TPair<int32, TObjectPtr<AChaosVDParticleActor>>& ParticleActorWithID : AllSolverParticlesByID)
diff --git a/Engine/Plugins/ChaosVD/Source/ChaosVD/Private/Trace/ChaosVDTraceAnalyzer.cpp b/Engine/Plugins/ChaosVD/Source/ChaosVD/Private/Trace/ChaosVDTraceAnalyzer.cpp
--- a/Engine/Plugins/ChaosVD/Source/ChaosVD/Private/Trace/ChaosVDTraceAnalyzer.cpp
+++ b/Engine/Plugins/ChaosVD/Source/ChaosVD/Private/Trace/ChaosVDTraceAnalyzer.cpp
@@ -140,6 +140,7 @@ bool FChaosVDTraceAnalyzer::OnEvent(uint16 RouteId, EStyle Style, const FOnEvent
 				EventData.GetString("StepName", DebugNameView);
 				NewSolverStageData.StepName = DebugNameView;
 				EnumAddFlags(NewSolverStageData.StageFlags, EChaosVDSolverStageFlags::Open);
+				EnumAddFlags(NewSolverStageData.StageFlags, EChaosVDSolverStageFlags::ExplicitStage);
 			}
 	
 			break;
