[ChaosVD] Refactor that removes actors as a mean to represent recorded particles.

Given than the current scene outliner and selection system depends on actors, as part of this change we moved over to use the new editor data storage system, TEDS, and its scene outliner.

Changes made in this CL:

CVD Changes:

- Moved all particle logic to its own solver data component (this put the particle update logic in line with newer systems in CVD, and is the first step to migrate it to its own CVD extension).
- Eliminate the uses of actors to represent each particle. They were replaced with a new ustruct type.
- Now that we don't use actors, I split the update logic in the particle representation structs into two sections, one that can be processed outside the GT, and one that needs to be processed in the GT (as interacts with other GT only editor systems, like making updates to the ISMCs).
- Refactored the update loop in the particle data component to do a first past in parallel in all new particle data, to preprocess it and figure out what needs to be done in the GT. Then we process any particle instance that has changes in the GT.
- Added a cvar to remove disable (and in the future remove) the old method we used to unload removed particles. Originally we didn't have particle destroyed events, so we were doing a diff against the list of current particle data.Also updated this removal step to be done in parallel.
- Created the necessary TEDS Data Factories, that will gather any modified data in the scene and update it in TEDS (item labels, icons, visibility, hierarchy data, etc)
- Implemented a bridge to use the Editor typed element selection system with plain ustructs instead of uobjects/component.
- Reworked the old CVD's Scene Outliner mode to interface with the TEDS outliner.
- Made CVD instantiate the TEDS outliner instead of the standard outliner
- Eliminated where possible duplicated calls to the same methods within a same frame.
- Added a cache for evaluated implicit object hashes (similar to the one we use at runtime for recording).
- Updated CVD's pool object components to create new components in a worker thread if the pool does not have a minim objects available (this reduces the stress of having to create them in the GT).

TEDS Changes:
- Added a way to pre-compute all children of any given row. This allow us to avoid having to walk the hierarchy upwards, which in the current implementation is exponential and was causing the outliner takes minutes to populate on a recording of City Sample.
- Fixed an issue where actors with no label set, were showing an empty label instead of a default label.
- Added support for ustructs to the teds visibility widget (we just needed to allow them in the query).
- Added CVD standalone to the supported programs list of the Editor Data Storage Features plugin.

diff --git a/Engine/Plugins/ChaosVD/ChaosVD.uplugin b/Engine/Plugins/ChaosVD/ChaosVD.uplugin
--- a/Engine/Plugins/ChaosVD/ChaosVD.uplugin
+++ b/Engine/Plugins/ChaosVD/ChaosVD.uplugin
@@ -33,6 +33,10 @@
 		}
 	],
 	"Plugins": [
+		{
+			"Name": "EditorDataStorageFeatures",
+			"Enabled": true
+		},
 		{
 			"Name": "GeometryProcessing",
 			"Enabled": true
diff --git a/Engine/Plugins/ChaosVD/Source/ChaosVD/ChaosVD.Build.cs b/Engine/Plugins/ChaosVD/Source/ChaosVD/ChaosVD.Build.cs
--- a/Engine/Plugins/ChaosVD/Source/ChaosVD/ChaosVD.Build.cs
+++ b/Engine/Plugins/ChaosVD/Source/ChaosVD/ChaosVD.Build.cs
@@ -9,7 +9,7 @@ public class ChaosVD : ModuleRules
 		PublicDependencyModuleNames.AddRange(
 			new string[]
 			{
-				"Core",
+				"Core"
 			});
 
 		PrivateDependencyModuleNames.AddRange(
@@ -19,7 +19,7 @@ public class ChaosVD : ModuleRules
 				"CoreUObject",
 				"Engine",
 				"RenderCore",
-				"ChaosVDRuntime", 
+				"ChaosVDRuntime",
 				"TraceServices",
 				"ChaosVDData",
 				"DesktopPlatform",
@@ -49,7 +49,8 @@ public class ChaosVD : ModuleRules
 				"SubobjectDataInterface",
 				"PropertyEditor",
 				"CommonMenuExtensions", 
-				"EditorWidgets"
+				"EditorWidgets", 
+				"TedsOutliner",
 			}
 			);
 
diff --git a/Engine/Plugins/ChaosVD/Source/ChaosVD/Private/Actors/ChaosVDDataContainerBaseActor.cpp b/Engine/Plugins/ChaosVD/Source/ChaosVD/Private/Actors/ChaosVDDataContainerBaseActor.cpp
--- a/Engine/Plugins/ChaosVD/Source/ChaosVD/Private/Actors/ChaosVDDataContainerBaseActor.cpp
+++ b/Engine/Plugins/ChaosVD/Source/ChaosVD/Private/Actors/ChaosVDDataContainerBaseActor.cpp
@@ -40,6 +40,18 @@ void AChaosVDDataContainerBaseActor::UpdateFromNewSolverFrameData(const FChaosVD
 	}
 }
 
+void AChaosVDDataContainerBaseActor::UpdateFromNewSolverStageData(const FChaosVDSolverFrameData& InSolverFrameData, const FChaosVDFrameStageData& InSolverFrameStageData)
+{
+	TInlineComponentArray<UChaosVDSolverDataComponent*> SolverDataComponents(this);
+	for (UChaosVDSolverDataComponent* Component : SolverDataComponents)
+	{
+		if (Component)
+		{
+			Component->UpdateFromNewSolverStageData(InSolverFrameData, InSolverFrameStageData);
+		}
+	}
+}
+
 void AChaosVDDataContainerBaseActor::Destroyed()
 {
 	TInlineComponentArray<UChaosVDSolverDataComponent*> SolverDataComponents(this);
@@ -68,6 +80,20 @@ void AChaosVDDataContainerBaseActor::PostActorCreated()
 	FChaosVDExtensionsManager::Get().OnExtensionRegistered().AddUObject(this, &AChaosVDDataContainerBaseActor::HandlePostInitializationExtensionRegistered);
 }
 
+void AChaosVDDataContainerBaseActor::SetIsTemporarilyHiddenInEditor(bool bIsHidden)
+{
+	TInlineComponentArray<UChaosVDSolverDataComponent*> SolverDataComponents(this);
+	for (UChaosVDSolverDataComponent* Component : SolverDataComponents)
+	{
+		if (Component)
+		{
+			Component->SetVisibility(!bIsHidden);
+		}
+	}
+
+	Super::SetIsTemporarilyHiddenInEditor(bIsHidden);
+}
+
 void AChaosVDDataContainerBaseActor::HandlePostInitializationExtensionRegistered(const TSharedRef<FChaosVDExtension>& NewExtension)
 {
 	TConstArrayView<TSubclassOf<UActorComponent>> CustomDataComponentsClasses = NewExtension->GetSolverDataComponentsClasses();
@@ -98,3 +124,15 @@ void AChaosVDDataContainerBaseActor::SetSolverID(int32 InSolverID)
 		}
 	}
 }
+
+void AChaosVDDataContainerBaseActor::UpdateVisibility(bool bIsVisible)
+{
+	TInlineComponentArray<UChaosVDSolverDataComponent*> SolverDataComponents(this);
+	for (UChaosVDSolverDataComponent* Component : SolverDataComponents)
+	{
+		if (Component)
+		{
+			Component->SetVisibility(bIsVisible);
+		}
+	}
+}
diff --git a/Engine/Plugins/ChaosVD/Source/ChaosVD/Private/Actors/ChaosVDSolverInfoActor.cpp b/Engine/Plugins/ChaosVD/Source/ChaosVD/Private/Actors/ChaosVDSolverInfoActor.cpp
--- a/Engine/Plugins/ChaosVD/Source/ChaosVD/Private/Actors/ChaosVDSolverInfoActor.cpp
+++ b/Engine/Plugins/ChaosVD/Source/ChaosVD/Private/Actors/ChaosVDSolverInfoActor.cpp
@@ -3,9 +3,9 @@
 #include "Actors/ChaosVDSolverInfoActor.h"
 
 #include "ChaosVDModule.h"
-#include "ChaosVDParticleActor.h"
+#include "ChaosVDSceneParticle.h"
 #include "ChaosVDScene.h"
-#include "ChaosVDSettingsManager.h"
+
 #include "EditorActorFolders.h"
 #include "Components/ChaosVDParticleDataComponent.h"
 #include "Components/ChaosVDSceneQueryDataComponent.h"
@@ -13,31 +13,18 @@
 #include "Components/ChaosVDSolverCollisionDataComponent.h"
 #include "Components/ChaosVDSolverJointConstraintDataComponent.h"
 #include "Elements/Framework/TypedElementSelectionSet.h"
-#include "Engine/World.h"
-#include "Misc/ScopedSlowTask.h"
-#include "Settings/ChaosVDParticleVisualizationSettings.h"
 
 #define LOCTEXT_NAMESPACE "ChaosVisualDebugger"
 
 AChaosVDSolverInfoActor::AChaosVDSolverInfoActor()
 {
 	CollisionDataComponent = CreateDefaultSubobject<UChaosVDSolverCollisionDataComponent>(TEXT("SolverCollisionDataComponent"));
-	ParticleDataComponent = CreateDefaultSubobject<UChaosVDParticleDataComponent>(TEXT("ParticleCollisionDataComponent"));
+	ParticleDataComponent = CreateDefaultSubobject<UChaosVDParticleDataComponent>(TEXT("ParticleDataComponent"));
 	JointsDataComponent = CreateDefaultSubobject<UChaosVDSolverJointConstraintDataComponent>(TEXT("JointDataComponent"));
 	CharacterGroundConstraintDataComponent = CreateDefaultSubobject<UChaosVDSolverCharacterGroundConstraintDataComponent>(TEXT("CharacterGroundConstraintDataComponent"));
 	SceneQueryDataComponent = CreateDefaultSubobject<UChaosVDSceneQueryDataComponent>(TEXT("ChaosVDSceneQueryDataComponent"));
 
 	bIsServer = false;
-
-	if (UChaosVDParticleVisualizationSettings* ParticleVisualizationSettings = FChaosVDSettingsManager::Get().GetSettingsObject<UChaosVDParticleVisualizationSettings>())
-	{
-		ParticleVisualizationSettings->OnSettingsChanged().AddUObject(this, &AChaosVDSolverInfoActor::HandleVisibilitySettingsUpdated);
-	}
-	
-	if (UChaosVDParticleVisualizationColorSettings* ColorVisualizationSettings = FChaosVDSettingsManager::Get().GetSettingsObject<UChaosVDParticleVisualizationColorSettings>())
-	{
-		ColorVisualizationSettings->OnSettingsChanged().AddUObject(this, &AChaosVDSolverInfoActor::HandleColorsSettingsUpdated);
-	}
 }
 
 void AChaosVDSolverInfoActor::SetSolverName(const FName& InSolverName)
@@ -46,6 +33,12 @@ void AChaosVDSolverInfoActor::SetSolverName(const FName& InSolverName)
 	SetActorLabel(TEXT("Solver Data Container | ") + InSolverName.ToString());
 }
 
+void AChaosVDSolverInfoActor::UpdateFromNewSolverFrameData(const FChaosVDSolverFrameData& InSolverFrameData)
+{
+	SetSimulationTransform(InSolverFrameData.SimulationTransform);
+	Super::UpdateFromNewSolverFrameData(InSolverFrameData);
+}
+
 void AChaosVDSolverInfoActor::SetScene(TWeakPtr<FChaosVDScene> InScene)
 {
 	Super::SetScene(InScene);
@@ -69,108 +62,20 @@ void AChaosVDSolverInfoActor::SetScene(TWeakPtr<FChaosVDScene> InScene)
 	}
 }
 
-void AChaosVDSolverInfoActor::RegisterParticleActor(int32 ParticleID, AChaosVDParticleActor* ParticleActor)
-{
-	if (!ensure(ParticleActor) || !ensure(ParticleActor->GetParticleData()))
-	{
-		return;
-	}
-
-	if (!SolverParticlesByID.Contains(ParticleID))
-	{
-		SolverParticlesByID.Add(ParticleID, ParticleActor);
-
-#if WITH_EDITOR
-		ApplySolverVisibilityToParticle(ParticleActor, IsTemporarilyHiddenInEditor());
-#endif
-
-	}
-
-	ParticleActor->SetFolderPath(GetFolderPathForParticleType(ParticleActor->GetParticleData()->Type));
-	CreatedFolders.Add(ParticleActor->GetFolder());
-}
-
-AChaosVDParticleActor* AChaosVDSolverInfoActor::GetParticleActor(int32 ParticleID)
-{
-	TObjectPtr<AChaosVDParticleActor>* FoundParticleActor = SolverParticlesByID.Find(ParticleID);
-
-	return FoundParticleActor ? ToRawPtr(*FoundParticleActor) : nullptr;
-}
-
-bool AChaosVDSolverInfoActor::IsParticleSelectedByID(int32 ParticleID)
-{
-	// Currently CVD does not support multi selection, so this should not be slow.
-	// But we might want to find another container for faster search after multi selection support is added
-	return SelectedParticlesID.Contains(ParticleID);
-}
-
-bool AChaosVDSolverInfoActor::SelectParticleByID(int32 ParticleIDToSelect)
-{
-	const TSharedPtr<FChaosVDScene> CVDScene = SceneWeakPtr.Pin();
-	if (!CVDScene)
-	{
-		UE_LOG(LogChaosVDEditor, Error, TEXT("[%s] Tried to select a particle without a valid CVD Scene."), ANSI_TO_TCHAR(__FUNCTION__))
-		return false;
-	}
-
-	AChaosVDParticleActor* ParticleToSelect = GetParticleActor(ParticleIDToSelect);
-	if (!ParticleToSelect)
-	{
-		UE_LOG(LogChaosVDEditor, Error, TEXT("[%s] Particle ID [%d] not found in Solver [%s]"), ANSI_TO_TCHAR(__FUNCTION__), ParticleIDToSelect, *GetSolverName().ToString());
-		return false;
-	}
-
-	CVDScene->SetSelectedObject(ParticleToSelect);
-
-	return true;
-}
-
-void AChaosVDSolverInfoActor::HandleVisibilitySettingsUpdated(UObject* SettingsObject)
-{
-	for (const TPair<int32, TObjectPtr<AChaosVDParticleActor>>& ParticleWithIDPair : SolverParticlesByID)
-	{
-		if (AChaosVDParticleActor* ParticleActor = ParticleWithIDPair.Value)
-		{
-			ParticleActor->UpdateGeometryComponentsVisibility();
-		}
-	}
-}
-
-void AChaosVDSolverInfoActor::HandleColorsSettingsUpdated(UObject* SettingsObject)
+FName AChaosVDSolverInfoActor::GetCustomIconName() const
 {
-	for (const TPair<int32, TObjectPtr<AChaosVDParticleActor>>& ParticleWithIDPair : SolverParticlesByID)
-	{
-		if (AChaosVDParticleActor* ParticleActor = ParticleWithIDPair.Value)
-		{
-			ParticleActor->UpdateGeometryColors();
-		}	
-	}
+	static FName SolverIconName = FName("SolverIcon");
+	return SolverIconName;
 }
 
-void AChaosVDSolverInfoActor::RemoveSolverFolders(UWorld* World)
+TSharedPtr<FChaosVDSceneParticle> AChaosVDSolverInfoActor::GetParticleInstance(int32 ParticleID) const
 {
-	if (!World)
+	if (ensure(ParticleDataComponent))
 	{
-		return;
+		return ParticleDataComponent->GetParticleInstanceByID(ParticleID);
 	}
 
-	TOptional<FFolder> ParentFolder;
-
-	for (const FFolder& Folder : CreatedFolders)
-	{
-		FActorFolders::Get().DeleteFolder(*World, Folder);
-
-		// All folders for the particles from this solver will have the same parent
-		if (!ParentFolder.IsSet())
-		{
-			ParentFolder = Folder.GetParent();
-		}
-	}
-	
-	if (ParentFolder.IsSet())
-	{
-		FActorFolders::Get().DeleteFolder(*World, ParentFolder.GetValue());
-	}
+	return nullptr;
 }
 
 bool AChaosVDSolverInfoActor::IsVisible() const
@@ -178,126 +83,19 @@ bool AChaosVDSolverInfoActor::IsVisible() const
 	return !IsTemporarilyHiddenInEditor();
 }
 
-void AChaosVDSolverInfoActor::ApplySolverVisibilityToParticle(AChaosVDParticleActor* ParticleActor, bool bIsHidden)
-{
-	if (!ParticleActor)
-	{
-		return;
-	}
-
-	if (bIsHidden)
-	{
-		// Note: We should probably add a priority system for the hide requests
-		// For now just clear the HideBySceneOutliner flag when a hide by solver request is done as this has priority
-		ParticleActor->RemoveHiddenFlag(EChaosVDHideParticleFlags::HiddenBySceneOutliner);
-		ParticleActor->AddHiddenFlag(EChaosVDHideParticleFlags::HiddenBySolverVisibility);
-	}
-	else
-	{
-		ParticleActor->RemoveHiddenFlag(EChaosVDHideParticleFlags::HiddenBySolverVisibility);
-	}
-}
-
 #if WITH_EDITOR
 void AChaosVDSolverInfoActor::SetIsTemporarilyHiddenInEditor(bool bIsHidden)
 {
 	if (IsTemporarilyHiddenInEditor() != bIsHidden)
 	{
-		Super::SetIsTemporarilyHiddenInEditor(bIsHidden);
-
-		constexpr float AmountOfWork = 1.0f;
-		const float PercentagePerElement = 1.0f / static_cast<float>(SolverParticlesByID.Num());
-
-		FScopedSlowTask VisibilityUpdateProgress(AmountOfWork, LOCTEXT("UpdatingParticlesVisisibility", "Updating Particles Visibility ..."));
-		VisibilityUpdateProgress.MakeDialog();
-
-		for (const TPair<int32, TObjectPtr<AChaosVDParticleActor>>& ParticleVDInstanceWithID : SolverParticlesByID)
-		{
-			ApplySolverVisibilityToParticle(ParticleVDInstanceWithID.Value, bIsHidden);
-
-			VisibilityUpdateProgress.EnterProgressFrame(PercentagePerElement);
-		}
-
 		if (const TSharedPtr<FChaosVDScene> CVDScene = SceneWeakPtr.Pin())
 		{
 			CVDScene->OnSolverVisibilityUpdated().Broadcast(SolverDataID, !bIsHidden);
 		}
 	}
-}
-#endif
-
-void AChaosVDSolverInfoActor::Destroyed()
-{
-	UWorld* World = GetWorld();
-	if (!World)
-	{
-		return;
-	}
-
-	if (UChaosVDParticleVisualizationSettings* ParticleVisualizationSettings = FChaosVDSettingsManager::Get().GetSettingsObject<UChaosVDParticleVisualizationSettings>())
-	{
-		ParticleVisualizationSettings->OnSettingsChanged().RemoveAll(this);
-	}
-	
-	if (UChaosVDParticleVisualizationColorSettings* ColorVisualizationSettings = FChaosVDSettingsManager::Get().GetSettingsObject<UChaosVDParticleVisualizationColorSettings>())
-	{
-		ColorVisualizationSettings->OnSettingsChanged().RemoveAll(this);
-	}
-
-	if (SolverParticlesByID.Num() > 0)
-	{
-		constexpr float AmountOfWork = 1.0f;
-		const float PercentagePerElement = 1.0f / static_cast<float>(SolverParticlesByID.Num());
-
-		FScopedSlowTask CleaningParticleDataSlowTask(AmountOfWork, LOCTEXT("CleaningParticleDataMessage", "Cleaning Up Particle Data ..."));
-		CleaningParticleDataSlowTask.MakeDialog();
-
-		for (const TPair<int32, TObjectPtr<AChaosVDParticleActor>>& ParticleVDInstanceWithID : SolverParticlesByID)
-		{
-			World->DestroyActor(ParticleVDInstanceWithID.Value);
-
-			CleaningParticleDataSlowTask.EnterProgressFrame(PercentagePerElement);
-		}
-	}
-
-	RemoveSolverFolders(World);
-
-	Super::Destroyed();
-}
-
-void AChaosVDSolverInfoActor::HandlePostSelectionChange(const UTypedElementSelectionSet* ChangesSelectionSet)
-{
-	SelectedParticlesID.Reset();
-
-	TArray<AChaosVDParticleActor*> SelectedParticles = ChangesSelectionSet->GetSelectedObjects<AChaosVDParticleActor>();
-
-	//TODO: Support multi-selection
-	if (SelectedParticles.Num() > 0)
-	{
-		if (AChaosVDParticleActor* SelectedParticle = SelectedParticles[0])
-		{
-			if (TSharedPtr<const FChaosVDParticleDataWrapper> ParticleData = SelectedParticle->GetParticleData())
-			{
-				SelectedParticlesID.Add(ParticleData->ParticleIndex);
-			}
-		}
-	}
-}
-
-FName AChaosVDSolverInfoActor::GetFolderPathForParticleType(EChaosVDParticleType ParticleType)
-{
-	if (const FName* FoundPathPtr = FolderPathByParticlePath.Find(ParticleType))
-	{
-		return *FoundPathPtr;
-	}
-	else
-	{
-		const FStringFormatOrderedArguments Args {SolverName.ToString(), FString::FromInt(SolverDataID)};
-		const FName ParticleFolderPath = *FPaths::Combine(FString::Format(TEXT("Solver {0} | ID {1}"), Args), UEnum::GetDisplayValueAsText(ParticleType).ToString());
 
-		FolderPathByParticlePath.Add(ParticleType, ParticleFolderPath);
-		return ParticleFolderPath;
-	}
+	Super::SetIsTemporarilyHiddenInEditor(bIsHidden);
 }
+#endif
 
 #undef LOCTEXT_NAMESPACE
diff --git a/Engine/Plugins/ChaosVD/Source/ChaosVD/Private/Actors/ChaosVDSolverInfoActor.h b/Engine/Plugins/ChaosVD/Source/ChaosVD/Private/Actors/ChaosVDSolverInfoActor.h
--- a/Engine/Plugins/ChaosVD/Source/ChaosVD/Private/Actors/ChaosVDSolverInfoActor.h
+++ b/Engine/Plugins/ChaosVD/Source/ChaosVD/Private/Actors/ChaosVDSolverInfoActor.h
@@ -3,7 +3,7 @@
 #pragma once
 
 #include "Actors/ChaosVDDataContainerBaseActor.h"
-#include "ChaosVDParticleActor.h"
+#include "ChaosVDSceneParticle.h"
 #include "ChaosVDSceneSelectionObserver.h"
 #include "GameFramework/Actor.h"
 #include "ChaosVDSolverInfoActor.generated.h"
@@ -12,7 +12,6 @@ class UChaosVDGTAccelerationStructuresDataComponent;
 class UChaosVDSolverCharacterGroundConstraintDataComponent;
 class UChaosVDSolverJointConstraintDataComponent;
 struct FChaosVDParticleDataWrapper;
-class AChaosVDParticleActor;
 class UChaosVDParticleDataComponent;
 class UChaosVDSceneQueryDataComponent;
 class UChaosVDSolverCollisionDataComponent;
@@ -30,62 +29,67 @@ public:
 
 	AChaosVDSolverInfoActor();
 
-	int32 GetSolverID() const { return SolverDataID; }
-
 	void SetSolverName(const FName& InSolverName);
-	const FName& GetSolverName() { return SolverName; }
-
-	void SetIsServer(bool bInIsServer) { bIsServer = bInIsServer; }
-	bool GetIsServer() const { return bIsServer; }
+	const FName& GetSolverName()
+	{
+		return SolverName;
+	}
 
-	virtual void SetScene(TWeakPtr<FChaosVDScene> InScene) override;
+	void SetIsServer(bool bInIsServer)
+	{
+		bIsServer = bInIsServer;
+	}
 
-	void SetSimulationTransform(const FTransform& InSimulationTransform) { SimulationTransform = InSimulationTransform; }
-	virtual const FTransform& GetSimulationTransform() const override { return SimulationTransform; }
+	bool GetIsServer() const
+	{
+		return bIsServer;
+	}
 
-	UChaosVDSolverCollisionDataComponent* GetCollisionDataComponent() { return CollisionDataComponent; }
-	UChaosVDParticleDataComponent* GetParticleDataComponent() { return ParticleDataComponent; }
-	UChaosVDSolverJointConstraintDataComponent* GetJointsDataComponent() { return JointsDataComponent; }
-	UChaosVDSolverCharacterGroundConstraintDataComponent* GetCharacterGroundConstraintDataComponent() { return CharacterGroundConstraintDataComponent; }
-	UChaosVDSceneQueryDataComponent* GetSceneQueryDataComponent() const { return SceneQueryDataComponent.Get(); }
+	virtual void UpdateFromNewSolverFrameData(const FChaosVDSolverFrameData& InSolverFrameData) override;
 
-	void RegisterParticleActor(int32 ParticleID, AChaosVDParticleActor* ParticleActor);
+	virtual void SetScene(TWeakPtr<FChaosVDScene> InScene) override;
 
-	AChaosVDParticleActor* GetParticleActor(int32 ParticleID);
-	const TMap<int32, TObjectPtr<AChaosVDParticleActor>>& GetAllParticleActorsByIDMap() { return  SolverParticlesByID; }
+	void SetSimulationTransform(const FTransform& InSimulationTransform)
+	{
+		SimulationTransform = InSimulationTransform;
+	}
+	virtual const FTransform& GetSimulationTransform() const override
+	{
+		return SimulationTransform;
+	}
 
-	const TArray<int32>& GetSelectedParticlesIDs() const { return SelectedParticlesID; }
+	UChaosVDSolverCollisionDataComponent* GetCollisionDataComponent()
+	{
+		return CollisionDataComponent;
+	}
 
-	bool IsParticleSelectedByID(int32 ParticleID);
+	UChaosVDParticleDataComponent* GetParticleDataComponent()
+	{
+		return ParticleDataComponent;
+	}
 
-	bool SelectParticleByID(int32 ParticleIDToSelect);
+	UChaosVDSolverCharacterGroundConstraintDataComponent* GetCharacterGroundConstraintDataComponent()
+	{
+		return CharacterGroundConstraintDataComponent;
+	}
 
-	template <typename TCallback>
-	void VisitSelectedParticleData(TCallback VisitCallback);
+	UChaosVDSceneQueryDataComponent* GetSceneQueryDataComponent() const
+	{
+		return SceneQueryDataComponent.Get();
+	}
 
-	template <typename TCallback>
-	void VisitAllParticleData(TCallback VisitCallback);
+	virtual FName GetCustomIconName() const override;
 
-	void RemoveSolverFolders(UWorld* World);
+	TSharedPtr<FChaosVDSceneParticle> GetParticleInstance(int32 ParticleID) const;
 
 	virtual bool IsVisible() const override;
 
 #if WITH_EDITOR
 	void SetIsTemporarilyHiddenInEditor(bool bIsHidden) override;
 #endif
-	virtual void Destroyed() override;
 
 protected:
 
-	void HandleVisibilitySettingsUpdated(UObject* SettingsObject);
-	void HandleColorsSettingsUpdated(UObject* SettingsObject);
-
-	void ApplySolverVisibilityToParticle(AChaosVDParticleActor* ParticleActor, bool bIsHidden);
-
-	virtual void HandlePostSelectionChange(const UTypedElementSelectionSet* ChangesSelectionSet) override;
-
-	FName GetFolderPathForParticleType(EChaosVDParticleType ParticleType);
-
 	UPROPERTY(VisibleAnywhere, Category="Solver Data")
 	FTransform SimulationTransform;
 
@@ -95,15 +99,6 @@ protected:
 	UPROPERTY()
 	TObjectPtr<UChaosVDSolverCollisionDataComponent> CollisionDataComponent;
 
-	UPROPERTY()
-	TMap<int32, TObjectPtr<AChaosVDParticleActor>> SolverParticlesByID;
-
-	TArray<int32> SelectedParticlesID;
-
-	TSortedMap<EChaosVDParticleType, FName> FolderPathByParticlePath;
-
-	TSet<FFolder> CreatedFolders;
-
 	bool bIsServer;
 
 	UPROPERTY()
@@ -118,41 +113,3 @@ protected:
 	UPROPERTY()
 	TObjectPtr<UChaosVDSceneQueryDataComponent> SceneQueryDataComponent;
 };
-
-template <typename TCallback>
-void AChaosVDSolverInfoActor::VisitSelectedParticleData(TCallback VisitCallback)
-{
-	for (const int32 SelectedParticleID : SelectedParticlesID)
-	{
-		AChaosVDParticleActor* ParticleActor = GetParticleActor(SelectedParticleID);
-		TSharedPtr<const FChaosVDParticleDataWrapper> ParticleDataViewer = ParticleActor ? ParticleActor->GetParticleData() : nullptr;
-		if (!ensure(ParticleDataViewer))
-		{
-			continue;
-		}
-
-		if (!VisitCallback(ParticleDataViewer))
-		{
-			return;
-		}
-	}
-}
-
-template <typename TCallback>
-void AChaosVDSolverInfoActor::VisitAllParticleData(TCallback VisitCallback)
-{
-	for (const TPair<int32, TObjectPtr<AChaosVDParticleActor>>& ParticleWithIDPair : SolverParticlesByID)
-	{
-		AChaosVDParticleActor* ParticleActor = ParticleWithIDPair.Value;
-		TSharedPtr<const FChaosVDParticleDataWrapper> ParticleDataViewer = ParticleActor ? ParticleActor->GetParticleData() : nullptr;
-		if (!ensure(ParticleDataViewer))
-		{
-			continue;
-		}
-
-		if (!VisitCallback(ParticleDataViewer))
-		{
-			return;
-		}
-	}
-}
diff --git a/Engine/Plugins/ChaosVD/Source/ChaosVD/Private/ChaosVDEditorMode.cpp b/Engine/Plugins/ChaosVD/Source/ChaosVD/Private/ChaosVDEditorMode.cpp
--- a/Engine/Plugins/ChaosVD/Source/ChaosVD/Private/ChaosVDEditorMode.cpp
+++ b/Engine/Plugins/ChaosVD/Source/ChaosVD/Private/ChaosVDEditorMode.cpp
@@ -2,7 +2,7 @@
 
 #include "ChaosVDEditorMode.h"
 
-#include "ChaosVDParticleActor.h"
+#include "Actors/ChaosVDSolverInfoActor.h"
 #include "EditorModeManager.h"
 #include "EditorViewportClient.h"
 #include "Elements/Actor/ActorElementEditorViewportInteractionCustomization.h"
@@ -32,7 +32,6 @@ AActor* FChaosVDEdModeWidgetHelper::GetCurrentSelectedActor() const
 	if (USelection* CurrentSelection = Owner->GetSelectedActors())
 	{
 		ensureMsgf(CurrentSelection->Num() < 2, TEXT("CVD does not support multi selection yet"));
-
 		return CurrentSelection->GetTop<AActor>();
 	}
 
@@ -75,8 +74,8 @@ bool FChaosVDEdModeWidgetHelper::UsesTransformWidget() const
 {
 	if (const AActor* SelectedActor = GetCurrentSelectedActor())
 	{
-		// Particle Actors cannot be moved
-		return !SelectedActor->IsA(AChaosVDParticleActor::StaticClass());
+		// CVD Data Actors cannot be moved
+		return !SelectedActor->IsA(AChaosVDSolverInfoActor::StaticClass());
 	}
 
 	return false;
diff --git a/Engine/Plugins/ChaosVD/Source/ChaosVD/Private/ChaosVDEditorModeTools.cpp b/Engine/Plugins/ChaosVD/Source/ChaosVD/Private/ChaosVDEditorModeTools.cpp
--- a/Engine/Plugins/ChaosVD/Source/ChaosVD/Private/ChaosVDEditorModeTools.cpp
+++ b/Engine/Plugins/ChaosVD/Source/ChaosVD/Private/ChaosVDEditorModeTools.cpp
@@ -47,6 +47,16 @@ USelection* FChaosVDEditorModeTools::GetSelectedComponents() const
 	return nullptr;
 }
 
+UTypedElementSelectionSet* FChaosVDEditorModeTools::GetEditorSelectionSet() const
+{
+	if (const TSharedPtr<FChaosVDScene> SceneSharedPtr = ScenePtr.Pin())
+	{
+		return SceneSharedPtr->GetElementSelectionSet();
+	}
+	
+	return nullptr;
+}
+
 UWorld* FChaosVDEditorModeTools::GetWorld() const
 {
 	if (const TSharedPtr<FChaosVDScene> SceneSharedPtr = ScenePtr.Pin())
diff --git a/Engine/Plugins/ChaosVD/Source/ChaosVD/Private/ChaosVDEditorModeTools.h b/Engine/Plugins/ChaosVD/Source/ChaosVD/Private/ChaosVDEditorModeTools.h
--- a/Engine/Plugins/ChaosVD/Source/ChaosVD/Private/ChaosVDEditorModeTools.h
+++ b/Engine/Plugins/ChaosVD/Source/ChaosVD/Private/ChaosVDEditorModeTools.h
@@ -15,6 +15,7 @@ public:
 	virtual USelection* GetSelectedActors() const override;
 	virtual USelection* GetSelectedObjects() const override;
 	virtual USelection* GetSelectedComponents() const override;
+	virtual UTypedElementSelectionSet* GetEditorSelectionSet() const override;
 	virtual UWorld* GetWorld() const override;
 
 	/** Returns the Chaos Visual Debugger Scene */
diff --git a/Engine/Plugins/ChaosVD/Source/ChaosVD/Private/ChaosVDGeometryBuilder.cpp b/Engine/Plugins/ChaosVD/Source/ChaosVD/Private/ChaosVDGeometryBuilder.cpp
--- a/Engine/Plugins/ChaosVD/Source/ChaosVD/Private/ChaosVDGeometryBuilder.cpp
+++ b/Engine/Plugins/ChaosVD/Source/ChaosVD/Private/ChaosVDGeometryBuilder.cpp
@@ -320,14 +320,15 @@ void FChaosVDGeometryBuilder::DeInitialize()
 	MeshComponentsWaitingForMaterial.Reset();
 	MeshComponentsWaitingForGeometry.Reset();
 	StaticMeshCacheMap.Reset();
+	SourceGeometryCache.Reset();
 
 	bInitialized = false;
 }
 
-void FChaosVDGeometryBuilder::CreateMeshesFromImplicitObject(const Chaos::FImplicitObject* InImplicitObject, AActor* Owner, TArray<TSharedPtr<FChaosVDExtractedGeometryDataHandle>>& OutMeshDataHandles,int32 AvailableShapeDataNum , const int32 DesiredLODCount, const Chaos::FRigidTransform3& InTransform, const int32 MeshIndex)
+void FChaosVDGeometryBuilder::CreateMeshesFromImplicitObject(const Chaos::FImplicitObject* InImplicitObject, TArray<TSharedPtr<FChaosVDExtractedGeometryDataHandle>>& OutMeshDataHandles,int32 AvailableShapeDataNum , const int32 DesiredLODCount, const Chaos::FRigidTransform3& InTransform, const int32 MeshIndex)
 {
 	// To start set the leaf and the root to the same ptr. If the object is an union, in the subsequent recursive call the leaf will be set correctly
-	CreateMeshesFromImplicit_Internal(InImplicitObject, InImplicitObject, Owner, OutMeshDataHandles, DesiredLODCount, InTransform, MeshIndex, AvailableShapeDataNum);
+	CreateMeshesFromImplicit_Internal(InImplicitObject, InImplicitObject, OutMeshDataHandles, DesiredLODCount, InTransform, MeshIndex, AvailableShapeDataNum);
 }
 
 void FChaosVDGeometryBuilder::AddReferencedObjects(FReferenceCollector& Collector)
@@ -338,8 +339,6 @@ void FChaosVDGeometryBuilder::AddReferencedObjects(FReferenceCollector& Collecto
 
 bool FChaosVDGeometryBuilder::DoesImplicitContainType(const Chaos::FImplicitObject* InImplicitObject, const Chaos::EImplicitObjectType ImplicitTypeToCheck)
 {
-	TRACE_CPUPROFILER_EVENT_SCOPE(FChaosVDGeometryBuilder::DoesImplicitContainType);
-
 	using namespace Chaos;
 	
 	if (!InImplicitObject)
@@ -383,7 +382,7 @@ bool FChaosVDGeometryBuilder::HasNegativeScale(const Chaos::FRigidTransform3& In
 	return ScaleSignVector.X * ScaleSignVector.Y * ScaleSignVector.Z < 0;
 }
 
-void FChaosVDGeometryBuilder::CreateMeshesFromImplicit_Internal(const Chaos::FImplicitObject* InRootImplicitObject, const Chaos::FImplicitObject* InLeafImplicitObject, AActor* Owner, TArray<TSharedPtr<FChaosVDExtractedGeometryDataHandle>>& OutMeshDataHandles, const int32 DesiredLODCount, const Chaos::FRigidTransform3& InTransform, const int32 ParentShapeInstanceIndex, int32 AvailableShapeDataNum)
+void FChaosVDGeometryBuilder::CreateMeshesFromImplicit_Internal(const Chaos::FImplicitObject* InRootImplicitObject, const Chaos::FImplicitObject* InLeafImplicitObject, TArray<TSharedPtr<FChaosVDExtractedGeometryDataHandle>>& OutMeshDataHandles, const int32 DesiredLODCount, const Chaos::FRigidTransform3& InTransform, const int32 ParentShapeInstanceIndex, int32 AvailableShapeDataNum)
 {
 	TRACE_CPUPROFILER_EVENT_SCOPE(FChaosVDGeometryBuilder::CreateMeshesFromImplicit_Internal);
 
@@ -421,7 +420,7 @@ void FChaosVDGeometryBuilder::CreateMeshesFromImplicit_Internal(const Chaos::FIm
 					}
 				}
 
-				CreateMeshesFromImplicit_Internal(InRootImplicitObject, UnionImplicit.GetReference(), Owner, OutMeshDataHandles, DesiredLODCount, InTransform, CurrentShapeInstanceIndex, AvailableShapeDataNum);	
+				CreateMeshesFromImplicit_Internal(InRootImplicitObject, UnionImplicit.GetReference(), OutMeshDataHandles, DesiredLODCount, InTransform, CurrentShapeInstanceIndex, AvailableShapeDataNum);	
 			}
 		}
 
@@ -433,7 +432,7 @@ void FChaosVDGeometryBuilder::CreateMeshesFromImplicit_Internal(const Chaos::FIm
 		if (const TImplicitObjectTransformed<FReal, 3>* Transformed = InLeafImplicitObject->template GetObject<TImplicitObjectTransformed<FReal, 3>>())
 		{
 			// For transformed objects, the Instance index is the same so we pass it in without changing it
-			CreateMeshesFromImplicit_Internal(InRootImplicitObject, Transformed->GetTransformedObject(), Owner, OutMeshDataHandles, DesiredLODCount, Transformed->GetTransform(), ParentShapeInstanceIndex, AvailableShapeDataNum);
+			CreateMeshesFromImplicit_Internal(InRootImplicitObject, Transformed->GetTransformedObject(), OutMeshDataHandles, DesiredLODCount, Transformed->GetTransform(), ParentShapeInstanceIndex, AvailableShapeDataNum);
 		}
 		
 		return;
@@ -449,7 +448,6 @@ void FChaosVDGeometryBuilder::CreateMeshesFromImplicit_Internal(const Chaos::FIm
 	}
 }
 
-
 bool FChaosVDGeometryBuilder::HasGeometryInCache(uint32 GeometryKey)
 {
 	FReadScopeLock ReadLock(GeometryCacheRWLock);
@@ -579,6 +577,47 @@ void FChaosVDGeometryBuilder::SetMeshComponentMaterial(IChaosVDGeometryComponent
 	}
 }
 
+void FChaosVDGeometryBuilder::HandleNewGeometryData(const Chaos::FConstImplicitObjectPtr& Geometry, const uint32 GeometryID)
+{
+	if (!Geometry)
+	{
+		return;
+	}
+
+	// We use implicit object hashes to tie them to generated static meshes.
+	// Calculating the hash each time we need to create a static mesh for it is too expensive
+	// So we do it here on load as this happens on the Trace Analysis thread.
+	// We intentionally do this with the inner objects as that is what we use to generate meshes
+
+	//TODO: At some point this will slow down too much the trace analysis thread, affecting live debugging and loading times
+	// If we reach that point we should implement a background tasks that primes the hash cache instead
+	
+	Geometry->VisitObjects([this](const Chaos::FImplicitObject* Implicit, const Chaos::FRigidTransform3& Transform, const int32 RootObjectIndex, const int32 ObjectIndex, const int32 LeafObjectIndex)
+	{
+		auto CacheImplicitObjectHashIfNeeded = [this](const Chaos::FImplicitObject* ImplicitObject)
+		{
+			if (!ImplicitObject)
+			{
+				return;
+			}
+
+			if (!SourceGeometryCache.HasGeometryInHashCache(ImplicitObject))
+			{
+				SourceGeometryCache.CacheImplicitObjectHash(ImplicitObject, ImplicitObject->GetTypeHash());
+			}
+		};
+
+		CacheImplicitObjectHashIfNeeded(Implicit);
+	
+		if (ImplicitObjectNeedsUnpacking(Implicit))
+		{
+			Chaos::FRigidTransform3 ExtractedTransform;
+			CacheImplicitObjectHashIfNeeded(UnpackImplicitObject(Implicit, ExtractedTransform));
+		}
+		return true;
+	});
+}
+
 void FChaosVDGeometryBuilder::DestroyMeshComponent(UMeshComponent* MeshComponent)
 {
 	if (IChaosVDGeometryComponent* AsCVDGeometryComponent = Cast<IChaosVDGeometryComponent>(MeshComponent))
@@ -830,19 +869,17 @@ void FChaosVDGeometryBuilder::AdjustedTransformForImplicit(const Chaos::FImplici
 
 TSharedPtr<FChaosVDExtractedGeometryDataHandle> FChaosVDGeometryBuilder::ExtractGeometryDataForImplicit(const Chaos::FImplicitObject* InImplicitObject, const Chaos::FRigidTransform3& InTransform)
 {
-	const uint32 ImplicitObjectHash = InImplicitObject->GetTypeHash();
-
 	Chaos::FRigidTransform3 ExtractedTransform = InTransform;
 	const bool bNeedsUnpack = ImplicitObjectNeedsUnpacking(InImplicitObject);
 	if (const Chaos::FImplicitObject* ImplicitObjectToProcess = bNeedsUnpack ? UnpackImplicitObject(InImplicitObject, ExtractedTransform) : InImplicitObject)
 	{
 		TSharedPtr<FChaosVDExtractedGeometryDataHandle> MeshDataHandle = MakeShared<FChaosVDExtractedGeometryDataHandle>();
 
-		const uint32 GeometryKey = ImplicitObjectToProcess->GetTypeHash();		
+		const uint32 GeometryKey = SourceGeometryCache.GetAndCacheGeometryHash(ImplicitObjectToProcess);		
 		MeshDataHandle->SetGeometryKey(GeometryKey);
 
 		// For the Component data key, we need the hash of the implicit as it is (packed) because we will need to match it when looking for shape data
-		MeshDataHandle->SetDataComponentKey(bNeedsUnpack ? ImplicitObjectHash : GeometryKey);
+		MeshDataHandle->SetDataComponentKey(bNeedsUnpack ? SourceGeometryCache.GetAndCacheGeometryHash(InImplicitObject) : GeometryKey);
 
 		AdjustedTransformForImplicit(ImplicitObjectToProcess, ExtractedTransform);
 		MeshDataHandle->SetTransform(ExtractedTransform);
diff --git a/Engine/Plugins/ChaosVD/Source/ChaosVD/Private/ChaosVDGeometryDataComponent.cpp b/Engine/Plugins/ChaosVD/Source/ChaosVD/Private/ChaosVDGeometryDataComponent.cpp
--- a/Engine/Plugins/ChaosVD/Source/ChaosVD/Private/ChaosVDGeometryDataComponent.cpp
+++ b/Engine/Plugins/ChaosVD/Source/ChaosVD/Private/ChaosVDGeometryDataComponent.cpp
@@ -5,16 +5,20 @@
 #include "Settings/ChaosVDCoreSettings.h"
 #include "ChaosVDGeometryBuilder.h"
 #include "ChaosVDModule.h"
-#include "ChaosVDParticleActor.h"
+#include "ChaosVDSceneParticle.h"
 #include "ChaosVDSettingsManager.h"
 #include "Components/MeshComponent.h"
 #include "Materials/Material.h"
 #include "Materials/MaterialInterface.h"
 #include "Materials/MaterialInstanceDynamic.h"
+#include "TEDS/ChaosVDStructTypedElementData.h"
 #include "Visualizers/ChaosVDParticleDataComponentVisualizer.h"
 
+TObjectPtr<UChaosVDParticleVisualizationColorSettings> FChaosVDGeometryComponentUtils::CachedParticleColorSettings = nullptr;
+TObjectPtr<UChaosVDParticleVisualizationSettings> FChaosVDGeometryComponentUtils::CachedParticleVisualizationSettings = nullptr;
 
-FChaosVDMeshDataInstanceHandle::FChaosVDMeshDataInstanceHandle(int32 InInstanceIndex, UMeshComponent* InMeshComponent, int32 InParticleID, int32 InSolverID)
+FChaosVDMeshDataInstanceHandle::FChaosVDMeshDataInstanceHandle(int32 InInstanceIndex, UMeshComponent* InMeshComponent, int32 InParticleID, int32 InSolverID, const TSharedRef<FChaosVDExtractedGeometryDataHandle>& InSourceGeometryHandle)
+: ExtractedGeometryHandle(InSourceGeometryHandle)
 {
 	InstanceState.MeshComponent = InMeshComponent;
 	InstanceState.MeshInstanceIndex = InInstanceIndex;
@@ -33,16 +37,17 @@ FChaosVDMeshDataInstanceHandle::FChaosVDMeshDataInstanceHandle(int32 InInstanceI
 	{
 		InstanceState.MeshComponentType = EChaosVDMeshComponent::Dynamic;
 	}
+
+	ExtractedGeometryHandle = InSourceGeometryHandle;
+
+	InstanceState.ImplicitObjectInfo.bIsRootObject = ExtractedGeometryHandle->GetRootImplicitObject() == ExtractedGeometryHandle->GetImplicitObject();
+	InstanceState.ImplicitObjectInfo.ShapeInstanceIndex = ExtractedGeometryHandle->GetShapeInstanceIndex();
+	InstanceState.ImplicitObjectInfo.ImplicitObjectType = ExtractedGeometryHandle->GetTypeName();
+	InstanceState.ImplicitObjectInfo.RelativeTransform = ExtractedGeometryHandle->GetRelativeTransform();	
 }
 
 void FChaosVDMeshDataInstanceHandle::SetWorldTransform(const FTransform& InTransform)
 {
-	if (!ExtractedGeometryHandle)
-	{
-		UE_LOG(LogChaosVDEditor, Error, TEXT("[%s] Attempted to update the world transform without a valid geometry handle"), ANSI_TO_TCHAR(__FUNCTION__));
-		return;
-	}
-
 	const FTransform ExtractedRelativeTransform = ExtractedGeometryHandle->GetRelativeTransform();
 
 	InstanceState.CurrentWorldTransform.SetLocation(InTransform.TransformPosition(ExtractedRelativeTransform.GetLocation()));
@@ -55,24 +60,6 @@ void FChaosVDMeshDataInstanceHandle::SetWorldTransform(const FTransform& InTrans
 	}
 }
 
-void FChaosVDMeshDataInstanceHandle::SetGeometryHandle(const TSharedPtr<FChaosVDExtractedGeometryDataHandle>& InHandle)
-{
-	ExtractedGeometryHandle = InHandle;
-
-	using namespace Chaos;
-	if (ExtractedGeometryHandle)
-	{
-		InstanceState.ImplicitObjectInfo.bIsRootObject = ExtractedGeometryHandle->GetRootImplicitObject() == ExtractedGeometryHandle->GetImplicitObject();
-		InstanceState.ImplicitObjectInfo.ShapeInstanceIndex = ExtractedGeometryHandle->GetShapeInstanceIndex();
-		InstanceState.ImplicitObjectInfo.ImplicitObjectType = ExtractedGeometryHandle->GetTypeName();
-		InstanceState.ImplicitObjectInfo.RelativeTransform = ExtractedGeometryHandle->GetRelativeTransform();
-	}
-	else
-	{
-		InstanceState.ImplicitObjectInfo = FChaosVDImplicitObjectBasicView();
-	}
-}
-
 void FChaosVDMeshDataInstanceHandle::SetInstanceColor(const FLinearColor& NewColor)
 {
 	if (InstanceState.CurrentGeometryColor != NewColor)
@@ -215,12 +202,7 @@ void FChaosVDGeometryComponentUtils::UpdateCollisionDataFromShapeArray(const TAr
 		return;
 	}
 
-	const TSharedPtr<FChaosVDExtractedGeometryDataHandle>& ExtractedGeometryHandle = InInstanceHandle->GetGeometryHandle();
-
-	if (!ExtractedGeometryHandle)
-	{
-		return;
-	}
+	const TSharedRef<FChaosVDExtractedGeometryDataHandle>& ExtractedGeometryHandle = InInstanceHandle->GetGeometryHandle();
 
 	const int32 ShapeInstanceIndex = ExtractedGeometryHandle->GetShapeInstanceIndex();
 	if (!InShapeArray.IsValidIndex(ShapeInstanceIndex))
@@ -269,7 +251,7 @@ void FChaosVDGeometryComponentUtils::UpdateMeshColor(const TSharedPtr<FChaosVDMe
 
 void FChaosVDGeometryComponentUtils::UpdateMeshVisibility(const TSharedPtr<FChaosVDMeshDataInstanceHandle>& InInstanceHandle, const FChaosVDParticleDataWrapper& InParticleData, bool bIsActive)
 {
-	if (!InInstanceHandle || !InInstanceHandle->GetGeometryHandle())
+	if (!InInstanceHandle)
 	{
 		return;
 	}
@@ -280,9 +262,9 @@ void FChaosVDGeometryComponentUtils::UpdateMeshVisibility(const TSharedPtr<FChao
 		return;
 	}
 
-	if (const UChaosVDParticleVisualizationSettings* ParticleVisualizationSettings = FChaosVDSettingsManager::Get().GetSettingsObject<UChaosVDParticleVisualizationSettings>())
+	if (const UChaosVDParticleVisualizationSettings* ParticleVisualizationSettings = GetParticleVisualizationSettings())
 	{
-		const EChaosVDGeometryVisibilityFlags CurrentVisibilityFlags = ParticleVisualizationSettings->GetGeometryVisualizationFlags();
+		const EChaosVDGeometryVisibilityFlags CurrentVisibilityFlags = static_cast<EChaosVDGeometryVisibilityFlags>(ParticleVisualizationSettings->GeometryVisibilityFlags);
 		
 		bool bShouldGeometryBeVisible = false;
 
@@ -336,7 +318,7 @@ FLinearColor FChaosVDGeometryComponentUtils::GetGeometryParticleColor(const TSha
 		return ColorToApply;
 	}
 
-	const UChaosVDParticleVisualizationColorSettings* VisualizationSettings = FChaosVDSettingsManager::Get().GetSettingsObject<UChaosVDParticleVisualizationColorSettings>();
+	const UChaosVDParticleVisualizationColorSettings* VisualizationSettings = GetParticleColorSettings();
 	if (!VisualizationSettings)
 	{
 		return ColorToApply;
@@ -383,6 +365,30 @@ FLinearColor FChaosVDGeometryComponentUtils::GetGeometryParticleColor(const TSha
 	return ColorToApply;
 }
 
+const UChaosVDParticleVisualizationColorSettings* FChaosVDGeometryComponentUtils::GetParticleColorSettings()
+{
+	if (CachedParticleColorSettings)
+	{
+		return CachedParticleColorSettings.Get();
+	}
+
+	CachedParticleColorSettings = FChaosVDSettingsManager::Get().GetSettingsObject<UChaosVDParticleVisualizationColorSettings>();
+
+	return CachedParticleColorSettings.Get();
+}
+
+const UChaosVDParticleVisualizationSettings* FChaosVDGeometryComponentUtils::GetParticleVisualizationSettings()
+{
+	if (CachedParticleVisualizationSettings)
+	{
+		return CachedParticleVisualizationSettings.Get();
+	}
+
+	CachedParticleVisualizationSettings = FChaosVDSettingsManager::Get().GetSettingsObject<UChaosVDParticleVisualizationSettings>();
+
+	return CachedParticleVisualizationSettings.Get();
+}
+
 UMaterialInterface* FChaosVDGeometryComponentUtils::GetBaseMaterialForType(EChaosVDMaterialType Type)
 {
 	const UChaosVDCoreSettings* EditorSettings = FChaosVDSettingsManager::Get().GetSettingsObject<UChaosVDCoreSettings>();
@@ -406,13 +412,13 @@ UMaterialInterface* FChaosVDGeometryComponentUtils::GetBaseMaterialForType(EChao
 	}	
 }
 
-void Chaos::VisualDebugger::SelectParticleWithGeometryInstance(const TSharedRef<FChaosVDScene>& InScene, IChaosVDGeometryOwnerInterface* GeometryOwner, const TSharedPtr<FChaosVDMeshDataInstanceHandle>& InMeshDataHandle)
+void Chaos::VisualDebugger::SelectParticleWithGeometryInstance(const TSharedRef<FChaosVDScene>& InScene, FChaosVDSceneParticle* Particle, const TSharedPtr<FChaosVDMeshDataInstanceHandle>& InMeshDataHandle)
 {
 	InScene->SetSelectedObject(nullptr);
 
-	if(GeometryOwner)
+	if (Particle)
 	{
-		GeometryOwner->SetSelectedMeshInstance(InMeshDataHandle);
-		InScene->SetSelectedObject(Cast<UObject>(GeometryOwner));	
+		Particle->SetSelectedMeshInstance(InMeshDataHandle);
+		InScene->SetSelected(VD::TypedElementDataUtil::AcquireTypedElementHandleForStruct(Particle, true));	
 	}
 }
diff --git a/Engine/Plugins/ChaosVD/Source/ChaosVD/Private/ChaosVDGeometryDataComponent.h b/Engine/Plugins/ChaosVD/Source/ChaosVD/Private/ChaosVDGeometryDataComponent.h
--- a/Engine/Plugins/ChaosVD/Source/ChaosVD/Private/ChaosVDGeometryDataComponent.h
+++ b/Engine/Plugins/ChaosVD/Source/ChaosVD/Private/ChaosVDGeometryDataComponent.h
@@ -16,6 +16,9 @@
 
 #include "ChaosVDGeometryDataComponent.generated.h"
 
+class UChaosVDParticleVisualizationSettings;
+class UChaosVDParticleVisualizationColorSettings;
+struct FChaosVDSceneParticle;
 class FChaosVDGeometryBuilder;
 class FChaosVDScene;
 class IChaosVDGeometryComponent;
@@ -131,7 +134,7 @@ struct FChaosVDMeshDataInstanceState
 class FChaosVDMeshDataInstanceHandle : public TSharedFromThis<FChaosVDMeshDataInstanceHandle>
 {
 public:
-	explicit FChaosVDMeshDataInstanceHandle(int32 InInstanceIndex, UMeshComponent* InMeshComponent, int32 InParticleID, int32 InSolverID);
+	explicit FChaosVDMeshDataInstanceHandle(int32 InInstanceIndex, UMeshComponent* InMeshComponent, int32 InParticleID, int32 InSolverID, const TSharedRef<FChaosVDExtractedGeometryDataHandle>& InSourceGeometryHandle);
 
 	/** Returns the Particle ID of the particle owning this mesh instance */
 	int32 GetOwningParticleID() const { return InstanceState.OwningParticleID; }
@@ -147,11 +150,8 @@ public:
 	/** Returns the world transform of the mesh instance this handle represents */
 	const FTransform& GetWorldTransform() const { return InstanceState.CurrentWorldTransform; }
 
-	/** Sets the Geometry Handle used to create the mesh instance this handle represents */
-	void SetGeometryHandle(const TSharedPtr<FChaosVDExtractedGeometryDataHandle>& InHandle);
-
 	/** Returns the geometry handle used to create the mesh instance this handle represents */
-	const TSharedPtr<FChaosVDExtractedGeometryDataHandle>& GetGeometryHandle() const { return ExtractedGeometryHandle; }
+	const TSharedRef<FChaosVDExtractedGeometryDataHandle>& GetGeometryHandle() const { return ExtractedGeometryHandle; }
 
 	/** Applies to the provided color to the mesh instance this handle represents */
 	void SetInstanceColor(const FLinearColor& NewColor);
@@ -195,13 +195,6 @@ public:
 
 	FChaosVDMeshDataInstanceState& GetState() { return InstanceState; }
 
-	friend uint32 GetTypeHash(const FChaosVDMeshDataInstanceHandle& Handle)
-	{
-		const uint32 GeometryHandleHash = Handle.ExtractedGeometryHandle ? GetTypeHash(Handle.ExtractedGeometryHandle->GetGeometryKey()) : 0;
-		const uint32 MeshComponentHandleHash = HashCombine(HashCombine(Handle.GetOwningParticleID(), Handle.GetMeshInstanceIndex()), static_cast<uint32>(Handle.GetMeshComponentType()));
-		return HashCombine(GeometryHandleHash, MeshComponentHandleHash);
-	}
-
 	/** Used only for debugging purposes - It will be set to true if we received new Shape Instance data but the Shape Index for the implicit object we represent is not valid */
 	bool bFailedToUpdateShapeInstanceData = false; 
 
@@ -215,7 +208,7 @@ private:
 
 	FChaosVDMeshDataInstanceState InstanceState;
 
-	TSharedPtr<FChaosVDExtractedGeometryDataHandle> ExtractedGeometryHandle;
+	TSharedRef<FChaosVDExtractedGeometryDataHandle> ExtractedGeometryHandle;
 
 	TWeakPtr<FChaosVDGeometryBuilder> GeometryBuilderInstance = nullptr;
 
@@ -266,7 +259,7 @@ public:
 	 * @param SolverID Solver ID of the particle owning this geometry
 	 * @return CVD Mesh instance handle that provides access to this component and specific instance, allowing manipulation of it
 	 */
-	virtual TSharedPtr<FChaosVDMeshDataInstanceHandle> AddMeshInstance(const FTransform InstanceTransform, bool bIsWorldSpace, const TSharedPtr<FChaosVDExtractedGeometryDataHandle>& InGeometryHandle, int32 ParticleID, int32 SolverID) PURE_VIRTUAL(IChaosVDGeometryDataComponent::AddMeshInstance, return nullptr;);
+	virtual TSharedPtr<FChaosVDMeshDataInstanceHandle> AddMeshInstance(const FTransform InstanceTransform, bool bIsWorldSpace, const TSharedRef<FChaosVDExtractedGeometryDataHandle>& InGeometryHandle, int32 ParticleID, int32 SolverID) PURE_VIRTUAL(IChaosVDGeometryDataComponent::AddMeshInstance, return nullptr;);
 
 	/**
 	 * Adds a new instance to this mesh component, but using an existing Mesh Data Handle instead of creating a new one
@@ -277,7 +270,7 @@ public:
 	 * @param ParticleID Particle ID owning this geometry
 	 * @param SolverID Solver ID of the particle owning this geometry
 	 */
-	virtual void AddMeshInstanceForHandle(TSharedPtr<FChaosVDMeshDataInstanceHandle> MeshDataHandle, const FTransform InstanceTransform, bool bIsWorldSpace, const TSharedPtr<FChaosVDExtractedGeometryDataHandle>& InGeometryHandle, int32 ParticleID, int32 SolverID) PURE_VIRTUAL(IChaosVDGeometryDataComponent::AddMeshInstanceForHandle);
+	virtual void AddMeshInstanceForHandle(TSharedPtr<FChaosVDMeshDataInstanceHandle> MeshDataHandle, const FTransform InstanceTransform, bool bIsWorldSpace, const TSharedRef<FChaosVDExtractedGeometryDataHandle>& InGeometryHandle, int32 ParticleID, int32 SolverID) PURE_VIRTUAL(IChaosVDGeometryDataComponent::AddMeshInstanceForHandle);
 
 	/**
 	 * Removes the instance the provided handle represents
@@ -348,6 +341,12 @@ public:
 private:
 	/** Returns the color that needs to be used to present the provided particle data based on its state and current selected options */
 	static FLinearColor GetGeometryParticleColor(const TSharedPtr<FChaosVDExtractedGeometryDataHandle>& InGeometryHandle, const FChaosVDParticleDataWrapper& InParticleData, bool bIsServer);
+
+	static const UChaosVDParticleVisualizationColorSettings* GetParticleColorSettings();
+	static const UChaosVDParticleVisualizationSettings* GetParticleVisualizationSettings();
+
+	static TObjectPtr<UChaosVDParticleVisualizationColorSettings> CachedParticleColorSettings;
+	static TObjectPtr<UChaosVDParticleVisualizationSettings> CachedParticleVisualizationSettings;
 };
 
 
@@ -365,26 +364,7 @@ EChaosVDMaterialType FChaosVDGeometryComponentUtils::GetMaterialTypeForComponent
 	}
 }
 
-UINTERFACE()
-class UChaosVDGeometryOwnerInterface : public UInterface
-{
-	GENERATED_BODY()
-};
-
-/**
- * Interface with for any CVD object that owns and handles Geometry generated from recorded data.
- */
-class IChaosVDGeometryOwnerInterface
-{
-	GENERATED_BODY()
-public:
-	virtual void SetSelectedMeshInstance(const TWeakPtr<FChaosVDMeshDataInstanceHandle>& InMeshInstanceToSelect) PURE_VIRTUAL(IChaosVDGeometryOwnerInterface::SetSelectedMeshInstance);
-	virtual TWeakPtr<FChaosVDMeshDataInstanceHandle> GetSelectedMeshInstance() const PURE_VIRTUAL(IChaosVDGeometryOwnerInterface::GetSelectedMeshInstance, return nullptr;);
-	virtual TConstArrayView<TSharedPtr<FChaosVDMeshDataInstanceHandle>> GetMeshInstances() const PURE_VIRTUAL(IChaosVDGeometryOwnerInterface::GetMeshInstances, return TConstArrayView<TSharedPtr<FChaosVDMeshDataInstanceHandle>>(););
-	virtual void HandleNewGeometryLoaded(uint32 GeometryID, const Chaos::FConstImplicitObjectPtr& InGeometryData) PURE_VIRTUAL(IChaosVDGeometryOwnerInterface::HandleGeometryUpdated);
-};
-
 namespace Chaos::VisualDebugger
 {
-	void SelectParticleWithGeometryInstance(const TSharedRef<FChaosVDScene>& InScene, IChaosVDGeometryOwnerInterface* GeometryOwner, const TSharedPtr<FChaosVDMeshDataInstanceHandle>& InMeshDataHandle);
+	void SelectParticleWithGeometryInstance(const TSharedRef<FChaosVDScene>& InScene, FChaosVDSceneParticle* Particle, const TSharedPtr<FChaosVDMeshDataInstanceHandle>& InMeshDataHandle);
 }
diff --git a/Engine/Plugins/ChaosVD/Source/ChaosVD/Private/ChaosVDModule.cpp b/Engine/Plugins/ChaosVD/Source/ChaosVD/Private/ChaosVDModule.cpp
--- a/Engine/Plugins/ChaosVD/Source/ChaosVD/Private/ChaosVDModule.cpp
+++ b/Engine/Plugins/ChaosVD/Source/ChaosVD/Private/ChaosVDModule.cpp
@@ -4,7 +4,7 @@
 
 #include "ChaosVDCommands.h"
 #include "ChaosVDEngine.h"
-#include "ChaosVDParticleActorCustomization.h"
+#include "ChaosVDSceneParticleCustomization.h"
 #include "ChaosVDSettingsManager.h"
 #include "ChaosVDStyle.h"
 #include "ChaosVDTabsIDs.h"
diff --git a/Engine/Plugins/ChaosVD/Source/ChaosVD/Private/ChaosVDObjectDetailsTab.cpp b/Engine/Plugins/ChaosVD/Source/ChaosVD/Private/ChaosVDObjectDetailsTab.cpp
--- a/Engine/Plugins/ChaosVD/Source/ChaosVD/Private/ChaosVDObjectDetailsTab.cpp
+++ b/Engine/Plugins/ChaosVD/Source/ChaosVD/Private/ChaosVDObjectDetailsTab.cpp
@@ -3,22 +3,23 @@
 #include "ChaosVDObjectDetailsTab.h"
 
 #include "ChaosVDCollisionDataDetailsTab.h"
-#include "ChaosVDParticleActor.h"
 #include "ChaosVDScene.h"
 #include "ChaosVDStyle.h"
 #include "ChaosVDTabsIDs.h"
 #include "Editor.h"
+#include "Elements/Actor/ActorElementData.h"
 #include "Elements/Framework/TypedElementSelectionSet.h"
+#include "Elements/Framework/EngineElementsLibrary.h"
 #include "Modules/ModuleManager.h"
 #include "PropertyEditorModule.h"
+#include "TEDS/ChaosVDSelectionInterface.h"
+#include "TEDS/ChaosVDStructTypedElementData.h"
 #include "Templates/SharedPointer.h"
 #include "Widgets/Docking/SDockTab.h"
 #include "Widgets/Input/SButton.h"
 #include "Widgets/SChaosVDCollisionDataInspector.h"
 #include "Widgets/SChaosVDDetailsView.h"
 #include "Widgets/SChaosVDMainTab.h"
-#include "Widgets/SToolTip.h"
-#include "Widgets/Text/STextBlock.h"
 
 #define LOCTEXT_NAMESPACE "ChaosVisualDebugger"
 
@@ -29,7 +30,6 @@ void FChaosVDObjectDetailsTab::AddUnsupportedStruct(const UStruct* Struct)
 
 TSharedRef<SDockTab> FChaosVDObjectDetailsTab::HandleTabSpawnRequest(const FSpawnTabArgs& Args)
 {
-
 	TSharedRef<SDockTab> DetailsPanelTab =
 	SNew(SDockTab)
 	.TabRole(ETabRole::MajorTab)
@@ -54,16 +54,7 @@ TSharedRef<SDockTab> FChaosVDObjectDetailsTab::HandleTabSpawnRequest(const FSpaw
 		{
 			DetailsPanelTab->SetContent
 			(
-				SNew(SVerticalBox)
-				+SVerticalBox::Slot()
-				[
-					SAssignNew(DetailsPanelView, SChaosVDDetailsView, MainTabPtr.ToSharedRef())
-				]
-				+SVerticalBox::Slot()
-				.AutoHeight()
-				[
-					GenerateShowCollisionDataButton().ToSharedRef()
-				]
+				SAssignNew(DetailsPanelView, SChaosVDDetailsView, MainTabPtr.ToSharedRef())
 			);
 
 			// If we closed the tab and opened it again with an object already selected, try to restore the selected object view
@@ -76,7 +67,6 @@ TSharedRef<SDockTab> FChaosVDObjectDetailsTab::HandleTabSpawnRequest(const FSpaw
 		{
 			DetailsPanelTab->SetContent(GenerateErrorWidget());
 		}
-
 	}
 	else
 	{
@@ -128,8 +118,29 @@ void FChaosVDObjectDetailsTab::HandleActorsSelection(TArrayView<AActor*> Selecte
 void FChaosVDObjectDetailsTab::HandlePostSelectionChange(const UTypedElementSelectionSet* ChangedSelectionSet)
 {
 	TArray<AActor*> SelectedActors = ChangedSelectionSet->GetSelectedObjects<AActor>();
+	if (SelectedActors.Num() > 0)
+	{
+		HandleActorsSelection(SelectedActors);
+		return;
+	}
+
+	constexpr int32 MaxElements = 1;
+	TArray<FTypedElementHandle, TInlineAllocator<MaxElements>> SelectedParticlesHandles;
+	ChangedSelectionSet->GetSelectedElementHandles(SelectedParticlesHandles, UChaosVDSelectionInterface::StaticClass());
+
+	if (SelectedParticlesHandles.Num() > 0)
+	{
+		// We don't support multi selection yet
+		ensure(SelectedParticlesHandles.Num() == MaxElements);
+
+		using namespace Chaos::VD::TypedElementDataUtil;
+	
+		DetailsPanelView->SetSelectedStruct(GetStructOnScopeDataFromTypedElementHandle(SelectedParticlesHandles[0]));
+		return;
+	}
 
-	HandleActorsSelection(SelectedActors);
+	DetailsPanelView->SetSelectedObject(nullptr);
+	DetailsPanelView->SetSelectedStruct(nullptr);
 }
 
 void FChaosVDObjectDetailsTab::HandleSolverDataSelectionChange(const TSharedPtr<FChaosVDSolverDataSelectionHandle>& SelectionHandle)
@@ -151,90 +162,4 @@ bool FChaosVDObjectDetailsTab::IsSupportedStruct(const TWeakObjectPtr<const UStr
 	return !UnsupportedStructs.Contains(InWeakStructPtr);
 }
 
-EVisibility FChaosVDObjectDetailsTab::GetCollisionDataButtonVisibility() const
-{
-	EVisibility DesiredVisibility = EVisibility::Collapsed;
-	if (CurrentSelectedObject.IsValid())
-	{
-		if (Cast<IChaosVDCollisionDataProviderInterface>(CurrentSelectedObject))
-		{
-			DesiredVisibility = EVisibility::Visible;
-		}
-	}
-	return DesiredVisibility;
-}
-
-bool FChaosVDObjectDetailsTab::GetCollisionDataButtonEnabled() const
-{
-	if (IChaosVDCollisionDataProviderInterface* CollisionDataProvider = Cast<IChaosVDCollisionDataProviderInterface>(CurrentSelectedObject.Get()))
-	{
-		return CollisionDataProvider->HasCollisionData();
-	}
-
-	return false;
-}
-
-TSharedPtr<SWidget> FChaosVDObjectDetailsTab::GenerateShowCollisionDataButton()
-{
-	TSharedPtr<SWidget> ShowCollisionButton = SNew(SHorizontalBox)
-	.Visibility_Raw(this, &FChaosVDObjectDetailsTab::GetCollisionDataButtonVisibility)
-	+SHorizontalBox::Slot()
-	.VAlign(VAlign_Center)
-	.HAlign(HAlign_Center)
-	.Padding(12.0f, 7.0f, 12.0f, 7.0f)
-	.FillWidth(1.0f)
-	[
-		SNew(SButton)
-		.ToolTip(SNew(SToolTip).Text(LOCTEXT("OpenCollisionDataDesc", "Click here to open the collision data for this particle on the collision data inspector.")))
-		.IsEnabled_Raw(this, &FChaosVDObjectDetailsTab::GetCollisionDataButtonEnabled)
-		.ContentPadding(FMargin(0, 5.f, 0, 4.f))
-		.OnClicked_Raw(this, &FChaosVDObjectDetailsTab::ShowCollisionDataForSelectedObject)
-		.Content()
-		[
-			SNew(SHorizontalBox)
-			+SHorizontalBox::Slot()
-			.Padding(FMargin(3, 0, 0, 0))
-			.VAlign(VAlign_Center)
-			.AutoWidth()
-			[
-				SNew(STextBlock)
-				.TextStyle(FAppStyle::Get(), "SmallButtonText")
-				.Text(LOCTEXT("ShowCollisionDataOnInspector","Show Collison Data in Inspector"))
-			]
-		]
-	];
-
-	return ShowCollisionButton;
-}
-
-FReply FChaosVDObjectDetailsTab::ShowCollisionDataForSelectedObject()
-{
-	IChaosVDCollisionDataProviderInterface* CollisionDataProvider = Cast<IChaosVDCollisionDataProviderInterface>(CurrentSelectedObject.Get());
-	if (!CollisionDataProvider)
-	{
-		return FReply::Handled();
-	}
-	
-	TSharedPtr<SChaosVDMainTab> OwningTabPtr = OwningTabWidget.Pin();
-	if (!OwningTabPtr.IsValid())
-	{
-		return FReply::Handled();
-	}
-
-	if (const TSharedPtr<FChaosVDCollisionDataDetailsTab> CollisionDataTab = OwningTabPtr->GetTabSpawnerInstance<FChaosVDCollisionDataDetailsTab>(FChaosVDTabID::CollisionDataDetails).Pin())
-	{
-		if (const TSharedPtr<FTabManager> TabManager = OwningTabPtr->GetTabManager())
-		{
-			TabManager->TryInvokeTab(FChaosVDTabID::CollisionDataDetails);
-
-			if (const TSharedPtr<SChaosVDCollisionDataInspector> CollisionInspector = CollisionDataTab->GetCollisionInspectorInstance().Pin())
-			{
-				CollisionInspector->SetCollisionDataProviderObjectToInspect(CollisionDataProvider);
-			}
-		}
-	}
-	
-	return FReply::Handled();
-}
-
 #undef LOCTEXT_NAMESPACE
diff --git a/Engine/Plugins/ChaosVD/Source/ChaosVD/Private/ChaosVDObjectDetailsTab.h b/Engine/Plugins/ChaosVD/Source/ChaosVD/Private/ChaosVDObjectDetailsTab.h
--- a/Engine/Plugins/ChaosVD/Source/ChaosVD/Private/ChaosVDObjectDetailsTab.h
+++ b/Engine/Plugins/ChaosVD/Source/ChaosVD/Private/ChaosVDObjectDetailsTab.h
@@ -43,13 +43,6 @@ protected:
 
 	bool IsSupportedStruct(const TWeakObjectPtr<const UStruct>& InWeakStructPtr);
 
-	EVisibility GetCollisionDataButtonVisibility() const;
-	bool GetCollisionDataButtonEnabled() const;
-
-	TSharedPtr<SWidget> GenerateShowCollisionDataButton();
-
-	FReply ShowCollisionDataForSelectedObject();
-
 	FDelegateHandle SelectionDelegateHandle;
 	TSharedPtr<SChaosVDDetailsView> DetailsPanelView;
 	TWeakObjectPtr<UObject> CurrentSelectedObject = nullptr;
diff --git a/Engine/Plugins/ChaosVD/Source/ChaosVD/Private/ChaosVDObjectPool.h b/Engine/Plugins/ChaosVD/Source/ChaosVD/Private/ChaosVDObjectPool.h
--- a/Engine/Plugins/ChaosVD/Source/ChaosVD/Private/ChaosVDObjectPool.h
+++ b/Engine/Plugins/ChaosVD/Source/ChaosVD/Private/ChaosVDObjectPool.h
@@ -3,8 +3,10 @@
 #pragma once
 
 #include "ChaosVDModule.h"
+#include "Async/ParallelFor.h"
 #include "Containers/Array.h"
 #include "Components/ActorComponent.h"
+#include "Containers/Ticker.h"
 #include "GameFramework/Actor.h"
 #include "Interfaces/ChaosVDPooledObject.h"
 #include "UObject/Object.h"
@@ -18,7 +20,7 @@ struct FChaosVDObjectPoolCVars
 
 /** Basic Pool system for UObjects */
 template<typename ObjectType>
-class TChaosVDObjectPool : public FGCObject
+class TChaosVDObjectPool : public FGCObject, public FTSTickerObjectBase
 {
 public:
 	TChaosVDObjectPool() = default;
@@ -33,6 +35,8 @@ public:
 		return PoolName;
 	}
 
+	virtual bool Tick(float DeltaTime) override;
+
 	ObjectType* AcquireObject(UObject* Outer, FName Name);
 
 	void DisposeObject(UObject* Object);
@@ -40,9 +44,14 @@ public:
 	TFunction<ObjectType*(UObject*,FName)> ObjectFactoryOverride;
 
 protected:
+	void DestroyUObject(UObject* Object);
+
 	FString PoolName = TEXT("ChaosVDObjectPool");
 	int32 PoolHits = 0;
 	int32 PoolRequests = 0;
+	std::atomic<bool> bGrowingPoolInBackground = false;
+	std::atomic<int32> PoolSize = 0;
+	FRWLock PoolLock;
 	TArray<TObjectPtr<ObjectType>> PooledObjects;
 };
 
@@ -51,6 +60,17 @@ TChaosVDObjectPool<ObjectType>::~TChaosVDObjectPool()
 {
 	float HitMissRatio = PoolRequests > 0 ? ((static_cast<float>(PoolHits) / static_cast<float>(PoolRequests)) * 100.0f) : 0;
 	UE_LOG(LogChaosVDEditor, Log, TEXT("Object pooling Statics for pool [%s] | Hits [%d] | Total Acquire requests [%d] | [%f] percent hit/miss ratio"), *PoolName, PoolHits,  PoolRequests, HitMissRatio);
+
+	if (FChaosVDObjectPoolCVars::bUseObjectPool)
+	{
+		FWriteScopeLock WriteLock(PoolLock);
+		for (TObjectPtr<ObjectType> ObjectToDestroy : PooledObjects)
+		{
+			DestroyUObject(ObjectToDestroy.Get());
+		}
+
+		PooledObjects.Empty();
+	}
 }
 
 template <typename ObjectType>
@@ -59,15 +79,75 @@ void TChaosVDObjectPool<ObjectType>::AddReferencedObjects(FReferenceCollector& C
 	Collector.AddReferencedObjects(PooledObjects);
 }
 
+template <typename ObjectType>
+bool TChaosVDObjectPool<ObjectType>::Tick(float DeltaTime)
+{
+	constexpr int32 MinimumPoolSize = 250;
+	if (FChaosVDObjectPoolCVars::bUseObjectPool && PoolSize < MinimumPoolSize && !bGrowingPoolInBackground)
+	{
+		const int32 ObjectsToCreate = MinimumPoolSize - PoolSize;
+		UE::Tasks::Launch(TEXT("GrowingCVDObjectPool"), [this, ObjectsToCreate]()
+		{
+			bGrowingPoolInBackground = true;
+			TArray<TObjectPtr<ObjectType>, TInlineAllocator<MinimumPoolSize>> CreatedObjects;
+
+			UPackage* TransientPackage = GetTransientPackage();
+
+			for (int32 Count = 0; Count < ObjectsToCreate; Count++)
+			{
+				const FName NewName = MakeUniqueObjectName(TransientPackage, ObjectType::StaticClass());
+
+				ObjectType* CreatedObject = nullptr;
+				if (ObjectFactoryOverride)
+				{
+					CreatedObject = ObjectFactoryOverride(TransientPackage, NewName);
+				}
+				else
+				{
+					CreatedObject = NewObject<ObjectType>(TransientPackage, NewName);
+				}
+
+				CreatedObjects.Emplace(CreatedObject);
+			}
+
+			{
+				FWriteScopeLock WriteLock(PoolLock);
+
+				PooledObjects.Append(CreatedObjects);
+				PoolSize = PooledObjects.Num();
+	
+				// Clear the async flag. It seems to be added to objects created outside the GT and prevents the object from being GC'd
+				// We intend to use these objects only in the GT
+				ParallelFor(CreatedObjects.Num(), [&CreatedObjects](int32 ObjectIndex)
+				{
+					CreatedObjects[ObjectIndex]->ClearInternalFlags(EInternalObjectFlags::Async);
+				});
+			}
+
+			bGrowingPoolInBackground = false;
+		});
+	}
+	
+	return true;
+}
+
 template <typename ObjectType>
 ObjectType* TChaosVDObjectPool<ObjectType>::AcquireObject(UObject* Outer, FName Name)
 {
 	PoolRequests++;
 
 	// If pooling is disabled, just fall through the code that will create the object
-	if (FChaosVDObjectPoolCVars::bUseObjectPool && PooledObjects.Num() > 0)
+	if (FChaosVDObjectPoolCVars::bUseObjectPool && PoolSize > 0)
 	{
-		if (ObjectType* Object = PooledObjects.Pop())
+		ObjectType* Object = nullptr;
+
+		{
+			FWriteScopeLock WriteLock(PoolLock);
+			Object = PooledObjects.Pop();
+			PoolSize = PooledObjects.Num();
+		}
+
+		if (Object)
 		{
 			const FName NewName = MakeUniqueObjectName(Outer, ObjectType::StaticClass(), Name);
 			Object->Rename(*NewName.ToString() , Outer, REN_NonTransactional | REN_DoNotDirty | REN_SkipGeneratedClasses | REN_DontCreateRedirectors);
@@ -110,19 +190,8 @@ void TChaosVDObjectPool<ObjectType>::DisposeObject(UObject* Object)
 	//TODO: Should we provide a way to override how these are destroyed?
 	if (!FChaosVDObjectPoolCVars::bUseObjectPool)
 	{
-		if (UActorComponent* AsActorComponent = Cast<UActorComponent>(Object))
-		{
-			AsActorComponent->DestroyComponent();
-		}
-		else if (AActor* AsActor = Cast<AActor>(Object))
-		{
-			AsActor->Destroy();
-		}
-		else
-		{
-			Object->ConditionalBeginDestroy();
-		}
-		
+		DestroyUObject(Object);
+
 		return;
 	}
 
@@ -135,5 +204,27 @@ void TChaosVDObjectPool<ObjectType>::DisposeObject(UObject* Object)
 	const FName NewName = MakeUniqueObjectName(TransientPackage, ObjectType::StaticClass());
 	Object->Rename(*NewName.ToString(), TransientPackage, REN_NonTransactional | REN_DoNotDirty | REN_SkipGeneratedClasses | REN_DontCreateRedirectors);
 
-	PooledObjects.Emplace(Cast<ObjectType>(Object));
+	{
+		FWriteScopeLock WriteLock(PoolLock);
+		PooledObjects.Emplace(Cast<ObjectType>(Object));
+	}
+
+	PoolSize = PooledObjects.Num();
+}
+
+template <typename ObjectType>
+void TChaosVDObjectPool<ObjectType>::DestroyUObject(UObject* Object)
+{
+	if (UActorComponent* AsActorComponent = Cast<UActorComponent>(Object))
+	{
+		AsActorComponent->DestroyComponent();
+	}
+	else if (AActor* AsActor = Cast<AActor>(Object))
+	{
+		AsActor->Destroy();
+	}
+	else
+	{
+		Object->ConditionalBeginDestroy();
+	}
 }
diff --git a/Engine/Plugins/ChaosVD/Source/ChaosVD/Private/ChaosVDPlaybackController.cpp b/Engine/Plugins/ChaosVD/Source/ChaosVD/Private/ChaosVDPlaybackController.cpp
--- a/Engine/Plugins/ChaosVD/Source/ChaosVD/Private/ChaosVDPlaybackController.cpp
+++ b/Engine/Plugins/ChaosVD/Source/ChaosVD/Private/ChaosVDPlaybackController.cpp
@@ -81,8 +81,6 @@ bool FChaosVDPlaybackController::LoadChaosVDRecordingFromTraceSession(const FCha
 	LoadedRecording->SetIsLive(InSessionDescriptor.bIsLiveSession);
 
 	HandleCurrentRecordingUpdated();
-
-	LoadedRecording->OnGeometryDataLoaded().AddRaw(this, &FChaosVDPlaybackController::EnqueueGeometryDataUpdate);
 	
 	if (TSharedPtr<FChaosVDScene> ScenePtr = SceneToControl.Pin())
 	{
@@ -179,16 +177,11 @@ void FChaosVDPlaybackController::EnqueueTrackInfoUpdate(const TSharedRef<const F
 	TrackInfoUpdateGTQueue.Enqueue(InfoUpdate);
 }
 
-void FChaosVDPlaybackController::EnqueueGeometryDataUpdate(const Chaos::FConstImplicitObjectPtr& NewGeometry, const uint32 GeometryID)
-{	
-	GeometryDataUpdateGTQueue.Enqueue({NewGeometry, GeometryID });
-}
-
 void FChaosVDPlaybackController::PlaySolverStageData(int32 TrackID, const TSharedRef<FChaosVDScene>& InSceneToControlSharedPtr, const FChaosVDSolverFrameData& InSolverFrameData, int32 StageIndex)
 {
 	if (InSolverFrameData.SolverSteps.IsValidIndex(StageIndex))
 	{
-		InSceneToControlSharedPtr->UpdateFromRecordedStepData(TrackID, InSolverFrameData.SolverSteps[StageIndex], InSolverFrameData);
+		InSceneToControlSharedPtr->UpdateFromRecordedSolverStageData(TrackID, InSolverFrameData.SolverSteps[StageIndex], InSolverFrameData);
 	}
 	else
 	{
@@ -709,18 +702,6 @@ bool FChaosVDPlaybackController::Tick(float DeltaTime)
 {
 	TickPlayback(DeltaTime);
 
-	if (!GeometryDataUpdateGTQueue.IsEmpty())
-	{
-		FChaosVDGeometryDataUpdate GeometryDataUpdate;
-		while (GeometryDataUpdateGTQueue.Dequeue(GeometryDataUpdate))
-		{
-			if (const TSharedPtr<FChaosVDScene> ScenePtr = SceneToControl.Pin())
-			{
-				ScenePtr->HandleNewGeometryData(GeometryDataUpdate.NewGeometry, GeometryDataUpdate.GeometryID);
-			}
-		}
-	}
-
 	const TWeakPtr<FChaosVDPlaybackController> ThisWeakPtr = DoesSharedInstanceExist() ? AsWeak() : nullptr;
 	if (!ThisWeakPtr.IsValid())
 	{
diff --git a/Engine/Plugins/ChaosVD/Source/ChaosVD/Private/ChaosVDPlaybackViewportClient.cpp b/Engine/Plugins/ChaosVD/Source/ChaosVD/Private/ChaosVDPlaybackViewportClient.cpp
--- a/Engine/Plugins/ChaosVD/Source/ChaosVD/Private/ChaosVDPlaybackViewportClient.cpp
+++ b/Engine/Plugins/ChaosVD/Source/ChaosVD/Private/ChaosVDPlaybackViewportClient.cpp
@@ -2,10 +2,10 @@
 
 #include "ChaosVDPlaybackViewportClient.h"
 
-#include "ChaosVDParticleActor.h"
+#include "Actors/ChaosVDSolverInfoActor.h"
 #include "ChaosVDScene.h"
 #include "ChaosVDSkySphereInterface.h"
-#include "Components/ChaosVDSceneQueryDataComponent.h"
+#include "Components/InstancedStaticMeshComponent.h"
 #include "ComponentVisualizer.h"
 #include "EditorModeManager.h"
 #include "Engine/DirectionalLight.h"
@@ -14,9 +14,8 @@
 #include "SEditorViewport.h"
 #include "Selection.h"
 #include "UnrealWidget.h"
-#include "Actors/ChaosVDGameFrameInfoActor.h"
-#include "Actors/ChaosVDSolverInfoActor.h"
-#include "Components/InstancedStaticMeshComponent.h"
+#include "TEDS/ChaosVDSelectionInterface.h"
+#include "TEDS/ChaosVDStructTypedElementData.h"
 #include "Visualizers/ChaosVDDebugDrawUtils.h"
 #include "Widgets/SChaosVDMainTab.h"
 
@@ -109,9 +108,9 @@ void FChaosVDPlaybackViewportClient::ProcessClick(FSceneView& View, HHitProxy* H
 		{
 			if (const TSharedPtr<FChaosVDMeshDataInstanceHandle> MeshDataHandle = AsCVDGeometryComponent->GetMeshDataInstanceHandle(MeshInstanceIndex))
 			{
-				if (AChaosVDParticleActor* ClickedActor = ScenePtr->GetParticleActor(MeshDataHandle->GetOwningSolverID(), MeshDataHandle->GetOwningParticleID()))
+				if (TSharedPtr<FChaosVDSceneParticle> ClickedActor = ScenePtr->GetParticleInstance(MeshDataHandle->GetOwningSolverID(), MeshDataHandle->GetOwningParticleID()))
 				{
-					Chaos::VisualDebugger::SelectParticleWithGeometryInstance(ScenePtr.ToSharedRef(), ClickedActor, bIsShiftKeyDown ? MeshDataHandle : nullptr);
+					Chaos::VisualDebugger::SelectParticleWithGeometryInstance(ScenePtr.ToSharedRef(), ClickedActor.Get(), bIsShiftKeyDown ? MeshDataHandle : nullptr);
 					bClickHandled = true;
 				}
 			}
@@ -190,13 +189,21 @@ void FChaosVDPlaybackViewportClient::FocusOnSelectedObject()
 {
 	if (const TSharedPtr<FChaosVDScene> CVDSceneSharedPtr = CVDScene.Pin())
 	{
-		USelection* CurrentSelection = ModeTools->GetSelectedActors();
+		UTypedElementSelectionSet* SelectionSet = CVDSceneSharedPtr->GetElementSelectionSet();
 
 		//TODO: Update this if we add multi selection support
-		if (AActor* SelectedActor = CurrentSelection ? CurrentSelection->GetTop<AActor>() : nullptr)
+		constexpr int32 MaxElements = 1;
+		TArray<FTypedElementHandle, TInlineAllocator<MaxElements>> SelectedParticlesHandles;
+		SelectionSet->GetSelectedElementHandles(SelectedParticlesHandles, UChaosVDSelectionInterface::StaticClass());
+
+		if (SelectedParticlesHandles.Num() > 0)
 		{
-			FBox BoxToTrack = SelectedActor->GetComponentsBoundingBox(false);
-			FocusViewportOnBox(BoxToTrack.ExpandBy(TrackingViewDistance), true);
+			using namespace Chaos::VD::TypedElementDataUtil;
+			if (FChaosVDSceneParticle* Particle = GetStructDataFromTypedElementHandle<FChaosVDSceneParticle>(SelectedParticlesHandles[0]))
+			{
+				FBox BoxToTrack = Particle->GetComponentsBoundingBox();
+				FocusViewportOnBox(BoxToTrack.ExpandBy(TrackingViewDistance), true);
+			}	
 		}
 	}
 }
diff --git a/Engine/Plugins/ChaosVD/Source/ChaosVD/Private/ChaosVDScene.cpp b/Engine/Plugins/ChaosVD/Source/ChaosVD/Private/ChaosVDScene.cpp
--- a/Engine/Plugins/ChaosVD/Source/ChaosVD/Private/ChaosVDScene.cpp
+++ b/Engine/Plugins/ChaosVD/Source/ChaosVD/Private/ChaosVDScene.cpp
@@ -5,23 +5,20 @@
 #include "Actors/ChaosVDSolverInfoActor.h"
 #include "ChaosVDGeometryBuilder.h"
 #include "ChaosVDModule.h"
-#include "ChaosVDParticleActor.h"
+#include "ChaosVDSceneParticle.h"
 #include "Chaos/ImplicitObject.h"
 #include "ChaosVDRecording.h"
 #include "ChaosVDSelectionCustomization.h"
 #include "ChaosVDSettingsManager.h"
 #include "ChaosVDSkySphereInterface.h"
+#include "EngineUtils.h"
 #include "Components/ChaosVDSceneQueryDataComponent.h"
-#include "Components/ChaosVDSolverCharacterGroundConstraintDataComponent.h"
 #include "Components/ChaosVDSolverCollisionDataComponent.h"
-#include "Components/ChaosVDSolverJointConstraintDataComponent.h"
-#include "DataWrappers/ChaosVDParticleDataWrapper.h"
 #include "Engine/DirectionalLight.h"
 #include "Engine/Engine.h"
 #include "Engine/Level.h"
 #include "Engine/StreamableManager.h"
 #include "Engine/World.h"
-#include "Elements/Actor/ActorElementData.h"
 #include "Elements/Framework/EngineElementsLibrary.h"
 #include "Elements/Framework/TypedElementSelectionSet.h"
 #include "Materials/Material.h"
@@ -31,9 +28,13 @@
 #include "Settings/ChaosVDCoreSettings.h"
 #include "UObject/Package.h"
 #include "Actors/ChaosVDGeometryContainer.h"
-#include "Components/LightComponent.h"
+#include "Elements/Columns/TypedElementMiscColumns.h"
+#include "Elements/Common/EditorDataStorageFeatures.h"
+#include "Elements/Framework/TypedElementRegistry.h"
 #include "Engine/PostProcessVolume.h"
 #include "Engine/TextureCube.h"
+#include "TEDS/ChaosVDParticleEditorDataFactory.h"
+#include "TEDS/ChaosVDStructTypedElementData.h"
 
 #define LOCTEXT_NAMESPACE "ChaosVisualDebugger"
 
@@ -41,13 +42,6 @@ FChaosVDScene::FChaosVDScene() = default;
 
 FChaosVDScene::~FChaosVDScene() = default;
 
-namespace ChaosVDSceneUIOptions
-{
-	constexpr float DelayToShowProgressDialogThreshold = 1.0f;
-	constexpr bool bShowCancelButton = false;
-	constexpr bool bAllowInPIE = false;
-}
-
 namespace Chaos::VisualDebugger::Cvars
 {
 	static bool bReInitializeGeometryBuilderOnCleanup = true;
@@ -88,7 +82,6 @@ void FChaosVDScene::Initialize()
 	bIsInitialized = true;
 }
 
-
 void FChaosVDScene::PerformGarbageCollection()
 {
 	FScopedSlowTask CollectingGarbageSlowTask(1, LOCTEXT("CollectingGarbageDataMessage", "Collecting Garbage ..."));
@@ -99,7 +92,6 @@ void FChaosVDScene::PerformGarbageCollection()
 	CollectingGarbageSlowTask.EnterProgressFrame();
 }
 
-
 void FChaosVDScene::DeInitialize()
 {
 	constexpr float AmountOfWork = 1.0f;
@@ -143,7 +135,7 @@ void FChaosVDScene::AddReferencedObjects(FReferenceCollector& Collector)
 	Collector.AddStableReferenceArray(&AvailableDataContainerActors);
 }
 
-void FChaosVDScene::UpdateFromRecordedStepData(const int32 SolverID, const FChaosVDFrameStageData& InRecordedStepData, const FChaosVDSolverFrameData& InFrameData)
+void FChaosVDScene::UpdateFromRecordedSolverStageData(const int32 SolverID, const FChaosVDFrameStageData& InRecordedStepData, const FChaosVDSolverFrameData& InFrameData)
 {
 	AChaosVDSolverInfoActor* SolverSceneData = nullptr;
 	if (AChaosVDSolverInfoActor** SolverSceneDataPtrPtr = SolverDataContainerBySolverID.Find(SolverID))
@@ -161,133 +153,45 @@ void FChaosVDScene::UpdateFromRecordedStepData(const int32 SolverID, const FChao
 	}
 
 	SolverSceneData->SetSimulationTransform(InFrameData.SimulationTransform);
-	
-	TSet<int32> ParticlesIDsInRecordedStepData;
-	ParticlesIDsInRecordedStepData.Reserve(InRecordedStepData.RecordedParticlesData.Num());
-	
-	{
-		constexpr float AmountOfWork = 1.0f;
-		const float PercentagePerElement = 1.0f / static_cast<float>(InRecordedStepData.RecordedParticlesData.Num());
-
-		const FText ProgressBarTitle = FText::Format(FTextFormat(LOCTEXT("ProcessingParticleData", "Processing Particle Data for {0} Solver with ID {1} ...")), FText::FromName(SolverSceneData->GetSolverName()), FText::AsNumber(SolverID));
-		FScopedSlowTask UpdatingSceneSlowTask(AmountOfWork, ProgressBarTitle);
-		UpdatingSceneSlowTask.MakeDialogDelayed(ChaosVDSceneUIOptions::DelayToShowProgressDialogThreshold, ChaosVDSceneUIOptions::bShowCancelButton, ChaosVDSceneUIOptions::bAllowInPIE);
-	
-		// Go over existing Particle VD Instances and update them or create them if needed 
-		for (const TSharedPtr<FChaosVDParticleDataWrapper>& Particle : InRecordedStepData.RecordedParticlesData)
-		{
-			const int32 ParticleVDInstanceID = GetIDForRecordedParticleData(Particle);
-			ParticlesIDsInRecordedStepData.Add(ParticleVDInstanceID);
-
-			if (InRecordedStepData.ParticlesDestroyedIDs.Contains(ParticleVDInstanceID))
-			{
-				// Do not process the particle if it was destroyed in the same step
-				continue;
-			}
-
-			if (AChaosVDParticleActor* ExistingParticleVDInstancePtr = SolverSceneData->GetParticleActor(ParticleVDInstanceID))
-			{
-				// We have new data for this particle, so re-activate the existing actor
-				if (!ExistingParticleVDInstancePtr->IsActive())
-				{
-					ExistingParticleVDInstancePtr->SetIsActive(true);
-				}
-
-				ExistingParticleVDInstancePtr->UpdateFromRecordedParticleData(Particle, InFrameData.SimulationTransform);
-			}
-			else
-			{
-				if (AChaosVDParticleActor* NewParticleVDInstance = SpawnParticleFromRecordedData(Particle, InFrameData))
-				{
-					// TODO: Precalculate the max num of entries we would see in the loaded file, and use that number to pre-allocate this map
-					SolverSceneData->RegisterParticleActor(ParticleVDInstanceID, NewParticleVDInstance);
-				}
-				else
-				{
-					//TODO: Handle this error
-					ensure(false);
-				}
-			}
-			
-			UpdatingSceneSlowTask.EnterProgressFrame(PercentagePerElement);
-		}
-	}
-
-	// Currently only explicitly recorded stages (no autogenerated) have valid constraint data
-	if (EnumHasAnyFlags(InRecordedStepData.StageFlags, EChaosVDSolverStageFlags::ExplicitStage))
-	{
-		UpdateParticlesCollisionData(InRecordedStepData, SolverID);
-
-		UpdateJointConstraintsData(InRecordedStepData, SolverID);
-	}
 
-	const TMap<int32, TObjectPtr<AChaosVDParticleActor>>& AllSolverParticlesByID = SolverSceneData->GetAllParticleActorsByIDMap();
+	SolverSceneData->UpdateFromNewSolverStageData(InFrameData, InRecordedStepData);
 
-	for (const TPair<int32, TObjectPtr<AChaosVDParticleActor>>& ParticleActorWithID : AllSolverParticlesByID)
-	{
-		// If we are playing back a keyframe, the scene should only contain what it is in the recorded data
-		const bool bShouldDestroyParticleAnyway = InFrameData.bIsKeyFrame && EnumHasAnyFlags(InRecordedStepData.StageFlags, EChaosVDSolverStageFlags::ExplicitStage) && !ParticlesIDsInRecordedStepData.Contains(ParticleActorWithID.Key);
-		
-		if (bShouldDestroyParticleAnyway || InFrameData.ParticlesDestroyedIDs.Contains(ParticleActorWithID.Key))
-		{
-			// In large maps moving at high speed (like when moving on a vehicle), level streaming adds/removes hundreds of actors (and therefore particles) constantly.
-			// Destroying particle actors is expensive, specially if we need to spawn them again sooner as we will need to rebuild-them.
-			// So, we deactivate them instead.
-
-			// TODO: We need an actor pool system, so we can keep memory under control as well.
-			if (AChaosVDParticleActor* ActorToDeactivate = ToRawPtr(ParticleActorWithID.Value))
-			{
-				if (IsObjectSelected(ActorToDeactivate))
-				{
-					ClearSelectionAndNotify();
-				}
-
-				ActorToDeactivate->SetIsActive(false);
-			}
-
-		}
-	}
-	
 	OnSceneUpdated().Broadcast();
 }
 
-void FChaosVDScene::UpdateParticlesCollisionData(const FChaosVDFrameStageData& InRecordedStepData, int32 SolverID)
+void FChaosVDScene::AddFromCVDWorldTagToActor(AActor* Actor)
 {
-	if (AChaosVDSolverInfoActor* SolverDataInfoContainer = SolverDataContainerBySolverID.FindChecked(SolverID))
+	using namespace UE::Editor::DataStorage;
+	// Add a selection column in TEDS
+	if (const IEditorDataStorageCompatibilityProvider* Compatibility = GetDataStorageFeature<IEditorDataStorageCompatibilityProvider>(CompatibilityFeatureName))
 	{
-		if (UChaosVDSolverCollisionDataComponent* CollisionDataContainer = SolverDataInfoContainer->GetCollisionDataComponent())
+		RowHandle Row = Compatibility->FindRowWithCompatibleObject(Actor);
+		if (Row != InvalidRowHandle)
 		{
-			CollisionDataContainer->UpdateCollisionData(InRecordedStepData.RecordedMidPhases);
+			if (IEditorDataStorageProvider* DataStorage = GetMutableDataStorageFeature<IEditorDataStorageProvider>(StorageFeatureName))
+			{
+				DataStorage->AddColumn<FTypedElementFromCVDWorldTag>(Row);
+				DataStorage->AddColumn<FTypedElementSyncFromWorldTag>(Row);
+			}
 		}
 	}
 }
 
-void FChaosVDScene::UpdateJointConstraintsData(const FChaosVDFrameStageData& InRecordedStepData, int32 SolverID)
+void FChaosVDScene::SetLoadedRecording(const TSharedPtr<FChaosVDRecording>& NewRecordingInstance)
 {
-	if (AChaosVDSolverInfoActor* SolverDataInfoContainer = SolverDataContainerBySolverID.FindChecked(SolverID))
+	if (LoadedRecording != NewRecordingInstance)
 	{
-		if (UChaosVDSolverJointConstraintDataComponent* JointsDataContainer = SolverDataInfoContainer->GetJointsDataComponent())
+		if (LoadedRecording)
 		{
-			JointsDataContainer->UpdateConstraintData(InRecordedStepData.RecordedJointConstraints);
+			NewRecordingInstance->OnGeometryDataLoaded().RemoveAll(GeometryGenerator.Get());
 		}
-	}
-}
 
-void FChaosVDScene::HandleNewGeometryData(const Chaos::FConstImplicitObjectPtr& GeometryData, const uint32 GeometryID)
-{
-	if (TArray<IChaosVDGeometryOwnerInterface*>* ObjectsWaitingPtr = ObjectsWaitingForGeometry.Find(GeometryID))
-	{
-		TArray<IChaosVDGeometryOwnerInterface*>& ObjectsWaitingRef = *ObjectsWaitingPtr;
-		for (IChaosVDGeometryOwnerInterface* ObjectWaiting : ObjectsWaitingRef)
+		LoadedRecording = NewRecordingInstance;
+
+		if (GeometryGenerator)
 		{
-			if (ObjectWaiting)
-			{
-				ObjectWaiting->HandleNewGeometryLoaded(GeometryID, GeometryData);
-			}
+			LoadedRecording->OnGeometryDataLoaded().AddSP(GeometryGenerator.ToSharedRef(), &FChaosVDGeometryBuilder::HandleNewGeometryData);
 		}
-
-		// Keep the array allocated in case another particle needs to go to the waiting list
-		ObjectsWaitingRef.Reset();
 	}
 }
 
@@ -309,15 +213,16 @@ AChaosVDSolverInfoActor* FChaosVDScene::GetOrCreateSolverInfoActor(int32 SolverI
 	const FName FolderPath = *FString::Format(TEXT("Solver {0} | ID {1}"), Args);
 
 	SolverDataInfo->SetFolderPath(FolderPath);
-
-	SolverDataInfo->SetSolverID(SolverID);
 	SolverDataInfo->SetSolverName(SolverName);
-	SolverDataInfo->SetScene(AsWeak());
 	SolverDataInfo->SetIsServer(bIsServer);
+	SolverDataInfo->SetSolverID(SolverID);
+	SolverDataInfo->SetScene(AsWeak());
 
 	SolverDataContainerBySolverID.Add(SolverID, SolverDataInfo);
 	AvailableDataContainerActors.Add(SolverDataInfo);
 
+	AddFromCVDWorldTagToActor(SolverDataInfo);
+
 	SolverInfoActorCreatedDelegate.Broadcast(SolverDataInfo);
 
 	return SolverDataInfo;
@@ -389,21 +294,22 @@ void FChaosVDScene::HandleEnterNewGameFrame(int32 FrameNumber, const TArray<int3
 	{
 		GameFrameDataContainer->UpdateFromNewGameFrameData(InNewGameFrameData);
 	}
+
+	bPendingUpdateRequest = true;
 }
 
 void FChaosVDScene::HandleEnterNewSolverFrame(int32 FrameNumber, const FChaosVDSolverFrameData& InFrameData)
 {
 	if (AChaosVDSolverInfoActor** SolverDataInfoContainerPtrPtr = SolverDataContainerBySolverID.Find(InFrameData.SolverID))
 	{
-		UChaosVDSolverCharacterGroundConstraintDataComponent* DataContainer = SolverDataInfoContainerPtrPtr ? (*SolverDataInfoContainerPtrPtr)->GetCharacterGroundConstraintDataComponent() : nullptr;
-
-		// TODO: Some times when playback is stopped, we might not have all the solver info actors ready when we go to the first frame
-		// This should not happen. For now I am making a change to avoid the crash and ensure. I created UE-217610 to find the issue and prepare a proper fix
-		if (ensure(DataContainer))
+		AChaosVDSolverInfoActor* SolverDataInfoContainerPtr = *SolverDataInfoContainerPtrPtr;
+		if (SolverDataInfoContainerPtrPtr)
 		{
-			DataContainer->UpdateConstraintData(InFrameData.RecordedCharacterGroundConstraints);
+			SolverDataInfoContainerPtr->UpdateFromNewSolverFrameData(InFrameData);
 		}
 	}
+
+	bPendingUpdateRequest = true;
 }
 
 void FChaosVDScene::CleanUpScene(EChaosVDSceneCleanUpOptions Options)
@@ -457,6 +363,8 @@ void FChaosVDScene::CleanUpScene(EChaosVDSceneCleanUpOptions Options)
 	{
 		PerformGarbageCollection();
 	}
+
+	Chaos::VD::TypedElementDataUtil::CleanUpTypedElementStore();
 }
 
 Chaos::FConstImplicitObjectPtr FChaosVDScene::GetUpdatedGeometry(int32 GeometryID) const
@@ -476,21 +384,11 @@ Chaos::FConstImplicitObjectPtr FChaosVDScene::GetUpdatedGeometry(int32 GeometryI
 	return nullptr;
 }
 
-void FChaosVDScene::AddObjectWaitingForGeometry(uint32 GeometryID, IChaosVDGeometryOwnerInterface* ObjectWaitingForGeometry)
-{
-	if (!ObjectWaitingForGeometry)
-	{
-		return;
-	}
-
-	ObjectsWaitingForGeometry.FindOrAdd(GeometryID).Add(ObjectWaitingForGeometry);
-}
-
-AChaosVDParticleActor* FChaosVDScene::GetParticleActor(int32 SolverID, int32 ParticleID)
+TSharedPtr<FChaosVDSceneParticle> FChaosVDScene::GetParticleInstance(int32 SolverID, int32 ParticleID)
 {
 	if (AChaosVDSolverInfoActor** SolverDataInfo = SolverDataContainerBySolverID.Find(SolverID))
 	{
-		return (*SolverDataInfo)->GetParticleActor(ParticleID);
+		return (*SolverDataInfo)->GetParticleInstance(ParticleID);
 	}
 
 	return nullptr;
@@ -506,49 +404,7 @@ AChaosVDSolverInfoActor* FChaosVDScene::GetSolverInfoActor(int32 SolverID)
 	return nullptr;
 }
 
-bool FChaosVDScene::IsSolverForServer(int32 SolverID) const
-{
-	if (const AChaosVDSolverInfoActor* const * PSolverDataInfo = SolverDataContainerBySolverID.Find(SolverID))
-	{
-		const AChaosVDSolverInfoActor* SolverDataInfo = *PSolverDataInfo;
-		return SolverDataInfo->GetIsServer();
-	}
-
-	return false;
-}
-
-AChaosVDParticleActor* FChaosVDScene::SpawnParticleFromRecordedData(const TSharedPtr<FChaosVDParticleDataWrapper>& InParticleData, const FChaosVDSolverFrameData& InFrameData)
-{
-	using namespace Chaos;
-
-	if (!InParticleData.IsValid())
-	{
-		return nullptr;
-	}
-
-	if (AChaosVDParticleActor* NewActor = PhysicsVDWorld->SpawnActor<AChaosVDParticleActor>())
-	{
-		NewActor->SetIsActive(true);
-		NewActor->SetScene(AsShared());
-		NewActor->SetIsServerParticle(IsSolverForServer(InParticleData->SolverID));
-		NewActor->UpdateFromRecordedParticleData(InParticleData, InFrameData.SimulationTransform);
-
-		// CVD's Outliner mode will update the label based on the particle data without needing to go trough all the code that Set Actor lable goes trough
-		// which can take +0.1 sec per actor
-		ParticleLabelUpdateDelegate.Broadcast(NewActor);
-
-		return NewActor;
-	}
-
-	return nullptr;
-}
-
-int32 FChaosVDScene::GetIDForRecordedParticleData(const TSharedPtr<FChaosVDParticleDataWrapper>& InParticleData) const
-{
-	return InParticleData ? InParticleData->ParticleIndex : INDEX_NONE;
-}
-
-void FChaosVDScene::CreateBaseLights(UWorld* TargetWorld) const
+void FChaosVDScene::CreateBaseLights(UWorld* TargetWorld)
 {
 	if (!TargetWorld)
 	{
@@ -595,6 +451,9 @@ void FChaosVDScene::CreateBaseLights(UWorld* TargetWorld) const
 						AsStaticMeshComponent->WireframeColorOverride = FColor::Black;
 					}
 				}
+
+				AddFromCVDWorldTagToActor(SkySphere);
+				AddFromCVDWorldTagToActor(DirectionalLightActor);
 			}
 		}
 	}
@@ -622,6 +481,8 @@ void FChaosVDScene::CreatePostProcessingVolumes(UWorld* TargetWorld)
 			}
 			
 			PostProcessingVolume->MarkComponentsRenderStateDirty();
+
+			AddFromCVDWorldTagToActor(PostProcessingVolume);
 		}
 	}
 }
@@ -700,30 +561,33 @@ void FChaosVDScene::UpdateSelectionProxiesForActors(TArrayView<AActor*> Selected
 	}
 }
 
+FName FChaosVDScene::GetTEDSSelectionSetName()
+{
+	return TEDSSelectionSetName;
+}
+
 void FChaosVDScene::HandleDeSelectElement(const TTypedElement<ITypedElementSelectionInterface>& InElementSelectionHandle, FTypedElementListRef InSelectionSet, const FTypedElementSelectionOptions& InSelectionOptions)
 {
-	if (AActor* DeselectedActor = ActorElementDataUtil::GetActorFromHandle(InElementSelectionHandle))
+	using namespace Chaos::VD::TypedElementDataUtil;
+	if (FChaosVDSceneParticle* Particle = GetStructDataFromTypedElementHandle<FChaosVDSceneParticle>(InElementSelectionHandle))
 	{
-		if (IChaosVDSelectableObject* SelectionAwareActor = Cast<IChaosVDSelectableObject>(DeselectedActor))
+		if (Particle)
 		{
-			SelectionAwareActor->HandleDeSelected();
+			Particle->HandleDeSelected();
 		}
 	}
-
-	// TODO: Add support for Component and Object Selection Events - This will be needed when we move away from using actors to represent particles
 }
 
 void FChaosVDScene::HandleSelectElement(const TTypedElement<ITypedElementSelectionInterface>& InElementSelectionHandle, FTypedElementListRef InSelectionSet, const FTypedElementSelectionOptions& InSelectionOptions)
 {
-	if (AActor* SelectedActor = ActorElementDataUtil::GetActorFromHandle(InElementSelectionHandle))
+	using namespace Chaos::VD::TypedElementDataUtil;
+	if (FChaosVDSceneParticle* Particle = GetStructDataFromTypedElementHandle<FChaosVDSceneParticle>(InElementSelectionHandle))
 	{
-		if (IChaosVDSelectableObject* SelectionAwareActor = Cast<IChaosVDSelectableObject>(SelectedActor))
+		if (Particle)
 		{
-			SelectionAwareActor->HandleSelected();
+			Particle->HandleSelected();
 		}
 	}
-
-	// TODO: Add support for Component and Object Selection Events - This will be needed when we move away from using actors to represent particles
 }
 
 void FChaosVDScene::ClearSelectionAndNotify()
@@ -737,12 +601,35 @@ void FChaosVDScene::ClearSelectionAndNotify()
 	SelectionSet->NotifyPendingChanges();
 }
 
+void FChaosVDScene::RequestUpdate()
+{
+	bPendingUpdateRequest = true;
+}
+
+bool FChaosVDScene::Tick(float DeltaTime)
+{
+	if (bPendingUpdateRequest)
+	{
+		OnSceneUpdated().Broadcast();
+        bPendingUpdateRequest = false;
+	}
+
+	return true;
+}
+
 void FChaosVDScene::InitializeSelectionSets()
 {
 	SelectionSet = NewObject<UTypedElementSelectionSet>(GetTransientPackage(), NAME_None, RF_Transactional);
+
+	TEDSSelectionSetName = FName(FString::Printf(TEXT("CVDSelectionSet%p"), SelectionSet.Get()));
+	SelectionSet->SetNameForTedsIntegration(TEDSSelectionSetName);
+
 	SelectionSet->AddToRoot();
 
+	using namespace Chaos::VD::TypedElementDataUtil;
+
 	SelectionSet->RegisterInterfaceCustomizationByTypeName(NAME_Actor, MakeUnique<FChaosVDSelectionCustomization>(AsShared()));
+	SelectionSet->RegisterInterfaceCustomizationByTypeName(NAME_CVD_StructDataElement, MakeUnique<FChaosVDSelectionCustomization>(AsShared()));
 	SelectionSet->RegisterInterfaceCustomizationByTypeName(NAME_Components, MakeUnique<FChaosVDSelectionCustomization>(AsShared()));
 	SelectionSet->RegisterInterfaceCustomizationByTypeName(NAME_Object, MakeUnique<FChaosVDSelectionCustomization>(AsShared()));
 
@@ -806,6 +693,27 @@ void FChaosVDScene::SetSelectedObject(UObject* SelectedObject)
 	SelectionSet->NotifyPendingChanges();
 }
 
+void FChaosVDScene::SetSelected(const FTypedElementHandle& InElementHandle)
+{
+	if (!SelectionSet)
+	{
+		return;
+	}
+
+	if (SelectionSet->IsElementSelected(InElementHandle, FTypedElementIsSelectedOptions()))
+	{
+		// Already selected, nothing to do here
+		return;
+	}
+
+	SelectionSet->ClearSelection(FTypedElementSelectionOptions());
+
+	TArray<FTypedElementHandle> NewEditorSelection = { InElementHandle };
+
+	SelectionSet->SetSelection(NewEditorSelection, FTypedElementSelectionOptions());
+	SelectionSet->NotifyPendingChanges();
+}
+
 bool FChaosVDScene::IsObjectSelected(const UObject* Object)
 {
 	if (!SelectionSet)
@@ -818,7 +726,12 @@ bool FChaosVDScene::IsObjectSelected(const UObject* Object)
 		return false;
 	}
 
-	return SelectionSet->IsElementSelected(GetSelectionHandleForObject(Object), FTypedElementIsSelectedOptions());;
+	return SelectionSet->IsElementSelected(GetSelectionHandleForObject(Object), FTypedElementIsSelectedOptions());
+}
+
+bool FChaosVDScene::IsSelected(const FTypedElementHandle& InElementHandle) const
+{
+	return SelectionSet->IsElementSelected(InElementHandle, FTypedElementIsSelectedOptions());
 }
 
 #undef LOCTEXT_NAMESPACE
