[ChaosVD] Refactor that removes actors as a mean to represent recorded particles.

diff --git a/Engine/Plugins/ChaosVD/Source/ChaosVD/Private/ChaosVDSceneParticle.cpp b/Engine/Plugins/ChaosVD/Source/ChaosVD/Private/ChaosVDSceneParticle.cpp
new file mode 100644
--- /dev/null
+++ b/Engine/Plugins/ChaosVD/Source/ChaosVD/Private/ChaosVDSceneParticle.cpp
@@ -0,0 +1,737 @@
+// Copyright Epic Games, Inc. All Rights Reserved.
+
+#include "ChaosVDSceneParticle.h"
+
+#include "ChaosVDGeometryBuilder.h"
+#include "ChaosVDModule.h"
+#include "ChaosVDScene.h"
+#include "Actors/ChaosVDSolverInfoActor.h"
+#include "Components/ChaosVDInstancedStaticMeshComponent.h"
+#include "Components/ChaosVDParticleDataComponent.h"
+#include "Components/ChaosVDSolverCharacterGroundConstraintDataComponent.h"
+#include "Components/ChaosVDSolverCollisionDataComponent.h"
+#include "Components/ChaosVDStaticMeshComponent.h"
+#include "Components/MeshComponent.h"
+#include "Components/SceneComponent.h"
+#include "Components/StaticMeshComponent.h"
+#include "DataWrappers/ChaosVDParticleDataWrapper.h"
+#include "Elements/Columns/TypedElementMiscColumns.h"
+#include "Elements/Columns/TypedElementVisibilityColumns.h"
+#include "Elements/Common/EditorDataStorageFeatures.h"
+#include "Elements/Interfaces/TypedElementDataStorageCompatibilityInterface.h"
+#include "Engine/StaticMesh.h"
+#include "TEDS/ChaosVDParticleEditorDataFactory.h"
+#include "TEDS/ChaosVDStructTypedElementData.h"
+#include "TEDS/ChaosVDTedsUtils.h"
+#include "Visualizers/ChaosVDSolverCollisionDataComponentVisualizer.h"
+
+namespace Chaos::VisualDebugger::Cvars
+{
+	static bool bForceStaticMeshComponentUse = false;
+	static FAutoConsoleVariableRef CVarChaosVDForceStaticMeshComponentUse(
+		TEXT("p.Chaos.VD.Tool.ForceStaticMeshComponentUse"),
+		bForceStaticMeshComponentUse,
+		TEXT("If true, static mesh components will be used instead of Instanced Static mesh components when recreating the geometry for each particle"));
+}
+
+FChaosVDSceneParticle::~FChaosVDSceneParticle()
+{
+	VisitGeometryInstances([](const TSharedRef<FChaosVDMeshDataInstanceHandle>& MeshDataHandle)
+	{
+		if (IChaosVDGeometryComponent* AsGeometryComponent = Cast<IChaosVDGeometryComponent>(MeshDataHandle->GetMeshComponent()))
+		{
+			AsGeometryComponent->RemoveMeshInstance(MeshDataHandle);
+		}
+	});
+
+	MeshDataHandles.Empty();
+
+	SetIsActive(false);
+}
+
+FChaosVDSceneParticle::FChaosVDSceneParticle()
+{
+	static FString UnnamedParticle("UnnamedParticle");
+	SetDisplayName(UnnamedParticle);
+
+	static FName DefaultIconName("RigidBodyIcon");
+	SetIconName(DefaultIconName);
+}
+
+void FChaosVDSceneParticle::UpdateTEDSData()
+{
+	using namespace UE::Editor::DataStorage;
+	SyncTEDSRegistration(GetMutableDataStorageFeature<IEditorDataStorageProvider>(StorageFeatureName), GetMutableDataStorageFeature<IEditorDataStorageCompatibilityProvider>(CompatibilityFeatureName));
+	SyncTEDSTags();
+}
+
+void FChaosVDSceneParticle::SyncTEDSRegistration(IEditorDataStorageProvider* StorageProvider, IEditorDataStorageCompatibilityProvider* DataStorageCompatibility)
+{
+	using namespace Chaos::VD::TedsUtils;
+	if (EnumHasAnyFlags(DirtyFlags, EChaosVDSceneParticleDirtyFlags::Active))
+	{
+		if (StorageProvider)
+		{
+			if (IsActive())
+			{
+				CachedRowHandle = AddObjectToDataStorage(this, StorageProvider, DataStorageCompatibility);
+			}
+			else
+			{
+				RemoveObjectToFromDataStorage(this, DataStorageCompatibility);
+			}
+		}
+		EnumRemoveFlags(DirtyFlags, EChaosVDSceneParticleDirtyFlags::Active);
+	}
+}
+
+void FChaosVDSceneParticle::SyncTEDSTags()
+{
+	if (EnumHasAnyFlags(DirtyFlags, EChaosVDSceneParticleDirtyFlags::TEDS))
+	{
+		using namespace UE::Editor::DataStorage;
+
+		if (CachedRowHandle != InvalidRowHandle)
+		{
+			if (IEditorDataStorageProvider* DataStorage = GetMutableDataStorageFeature<IEditorDataStorageProvider>(StorageFeatureName))
+			{
+				DataStorage->AddColumn<FTypedElementSyncFromWorldTag>(CachedRowHandle);
+			}	
+		}
+	}
+
+	EnumRemoveFlags(DirtyFlags, EChaosVDSceneParticleDirtyFlags::TEDS);
+}
+
+void FChaosVDSceneParticle::UpdateParent(const TSharedPtr<FChaosVDParticleDataWrapper>& InRecordedData)
+{
+	TSharedPtr<FChaosVDScene> ScenePtr = SceneWeakPtr.Pin();
+
+	if (ScenePtr && InRecordedData->ParticleCluster.HasValidData())
+	{
+		if (TSharedPtr<FChaosVDSceneParticle> ParentParticle = ScenePtr->GetParticleInstance(InRecordedData->SolverID, InRecordedData->ParticleCluster.ParentParticleID))
+		{
+			SetParent(ParentParticle);
+		}
+		else if (AChaosVDSolverInfoActor* SolverData = ScenePtr->GetSolverInfoActor(InRecordedData->SolverID))
+		{
+			if (UChaosVDParticleDataComponent* ParticleDataComponent = SolverData->GetParticleDataComponent())
+			{
+				SetParent(ParticleDataComponent->GetParticleContainerByType(InRecordedData->Type));
+			}
+			else
+			{
+				SetParent(nullptr);
+			}
+		}
+		else
+		{
+			SetParent(nullptr);
+		}
+	}
+
+	EnumAddFlags(DirtyFlags, EChaosVDSceneParticleDirtyFlags::TEDS);
+
+	EnumRemoveFlags(DirtyFlags, EChaosVDSceneParticleDirtyFlags::Parent);
+}
+
+void FChaosVDSceneParticle::PreUpdateFromRecordedParticleData(const TSharedPtr<FChaosVDParticleDataWrapper>& InRecordedData, const Chaos::FRigidTransform3& SimulationTransform)
+{
+	if (!ensure(InRecordedData.IsValid()))
+	{
+		return;	
+	}
+	
+	if (!ParticleDataPtr
+		|| ParticleDataPtr->ParticleCluster.HasValidData() != InRecordedData->ParticleCluster.HasValidData()
+		|| ParticleDataPtr->ParticleCluster.ParentParticleID != InRecordedData->ParticleCluster.ParentParticleID)
+	{
+		EnumAddFlags(DirtyFlags, EChaosVDSceneParticleDirtyFlags::Parent);
+	}
+
+	//TODO: Make the simulation transform be cached on the CVD Scene, so we can query from it when needed
+	// Copying it to each particle actor is not efficient
+	CachedSimulationTransform = SimulationTransform;
+
+	if (InRecordedData->ParticlePositionRotation.HasValidData())
+	{
+		const FVector TargetLocation = SimulationTransform.TransformPosition(InRecordedData->ParticlePositionRotation.MX);
+		const FVector CurrentLocation = SimulationTransform.TransformPosition(ParticleDataPtr && ParticleDataPtr->ParticlePositionRotation.HasValidData() ? ParticleDataPtr->ParticlePositionRotation.MX : FVector::ZeroVector);
+
+		const FQuat TargetRotation = SimulationTransform.GetRotation() * InRecordedData->ParticlePositionRotation.MR;
+		const FQuat CurrentRotation = SimulationTransform.GetRotation() * (ParticleDataPtr && ParticleDataPtr->ParticlePositionRotation.HasValidData() ? ParticleDataPtr->ParticlePositionRotation.MR : FQuat::Identity);
+
+		PendingParticleTransform.SetLocation(TargetLocation);
+		PendingParticleTransform.SetRotation(TargetRotation);
+		PendingParticleTransform.SetScale3D(FVector(1.0f,1.0f,1.0f));
+
+		if (CurrentRotation != TargetRotation || CurrentLocation != TargetLocation)
+		{
+			EnumAddFlags(DirtyFlags, EChaosVDSceneParticleDirtyFlags::Transform);
+		}
+	}
+
+	// This is iterating and comparing each element of the array,
+	// We might need to find a faster way of determine if the data changed, but for now this is faster than assuming it changed
+	const bool bShapeDataIsDirty = !ParticleDataPtr || (ParticleDataPtr->CollisionDataPerShape != InRecordedData->CollisionDataPerShape);
+	const bool bDisabledStateChanged = ParticleDataPtr && (ParticleDataPtr->ParticleDynamicsMisc.bDisabled != InRecordedData->ParticleDynamicsMisc.bDisabled);
+	const bool bHasNewGeometry = !ParticleDataPtr || (ParticleDataPtr->GeometryHash != InRecordedData->GeometryHash);
+	const bool bStateChanged = !ParticleDataPtr || (ParticleDataPtr->ParticleDynamicsMisc.MObjectState != InRecordedData->ParticleDynamicsMisc.MObjectState);
+
+	static FString UnnamedParticle("UnnamedParticle");
+	if (!ParticleDataPtr)
+	{
+		SetDisplayName(InRecordedData->DebugName.IsEmpty() ? UnnamedParticle : InRecordedData->DebugName);
+	}
+
+	ParticleDataPtr = InRecordedData;
+
+	if (bHasNewGeometry)
+	{
+		EnumAddFlags(DirtyFlags, EChaosVDSceneParticleDirtyFlags::Geometry);
+	}
+
+	// Now that we have updated particle data, update the Shape data and visibility as needed
+	if (bShapeDataIsDirty || bHasNewGeometry)
+	{
+		EnumAddFlags(DirtyFlags, EChaosVDSceneParticleDirtyFlags::CollisionData);
+	}
+	else if (bDisabledStateChanged)
+	{
+		EnumAddFlags(DirtyFlags, EChaosVDSceneParticleDirtyFlags::Visibility);
+	}
+
+	if (EnumHasAnyFlags(DirtyFlags, EChaosVDSceneParticleDirtyFlags::Visibility) || EnumHasAnyFlags(DirtyFlags, EChaosVDSceneParticleDirtyFlags::Parent))
+	{
+		EnumAddFlags(DirtyFlags, EChaosVDSceneParticleDirtyFlags::TEDS);
+	}
+
+	if (bStateChanged || EnumHasAnyFlags(DirtyFlags, EChaosVDSceneParticleDirtyFlags::Visibility)
+		|| EnumHasAnyFlags(DirtyFlags, EChaosVDSceneParticleDirtyFlags::CollisionData)
+		|| EnumHasAnyFlags(DirtyFlags, EChaosVDSceneParticleDirtyFlags::Geometry))
+	{
+		EnumAddFlags(DirtyFlags, EChaosVDSceneParticleDirtyFlags::Coloring);
+	}
+	
+	EnumAddFlags(DirtyFlags, EChaosVDSceneParticleDirtyFlags::PreUpdatePass);
+}
+
+void FChaosVDSceneParticle::ProcessPendingParticleDataUpdates()
+{
+	if (!ensure(IsInGameThread()))
+	{
+		UE_LOG(LogChaosVDEditor, Error, TEXT("Attempted to update for particle [%s] outside of the game thread!. This is not supported!"), *GetDisplayName())
+		return;
+	}
+
+	if (!ensure(EnumHasAnyFlags(DirtyFlags, EChaosVDSceneParticleDirtyFlags::PreUpdatePass)))
+	{
+		UE_LOG(LogChaosVDEditor, Error, TEXT("Attempted to process update for particle [%s] without doing a pre pass first!. Current particle data is out of date!"), *GetDisplayName())
+		return;
+	}
+
+	if (!ParticleDataPtr)
+	{
+		return;
+	}
+
+	if (EnumHasAnyFlags(DirtyFlags, EChaosVDSceneParticleDirtyFlags::Geometry))
+	{
+		UpdateGeometry(ParticleDataPtr->GeometryHash, EChaosVDActorGeometryUpdateFlags::ForceUpdate);
+	}
+	
+	if (EnumHasAnyFlags(DirtyFlags, EChaosVDSceneParticleDirtyFlags::Transform))
+	{
+		ApplyPendingTransformData();
+	}
+
+	if (EnumHasAnyFlags(DirtyFlags, EChaosVDSceneParticleDirtyFlags::CollisionData))
+	{
+		UpdateShapeDataComponents();
+	}
+
+	if (EnumHasAnyFlags(DirtyFlags, EChaosVDSceneParticleDirtyFlags::Visibility))
+	{
+		UpdateGeometryComponentsVisibility();
+	}
+
+	if (EnumHasAnyFlags(DirtyFlags, EChaosVDSceneParticleDirtyFlags::Coloring))
+	{
+		UpdateGeometryColors();
+	}
+
+	if (EnumHasAnyFlags(DirtyFlags, EChaosVDSceneParticleDirtyFlags::Parent))
+	{
+		UpdateParent(ParticleDataPtr);
+	}
+
+	EnumRemoveFlags(DirtyFlags, EChaosVDSceneParticleDirtyFlags::PreUpdatePass);
+}
+
+void FChaosVDSceneParticle::SetSelectedMeshInstance(const TWeakPtr<FChaosVDMeshDataInstanceHandle>& GeometryInstanceToSelect)
+{
+	if (!ParticleDataPtr)
+	{
+		return;
+	}
+
+	const TSharedPtr<FChaosVDMeshDataInstanceHandle> GeometryInstanceToSelectPtr = GeometryInstanceToSelect.Pin();
+	if (!GeometryInstanceToSelectPtr)
+	{
+		return;
+	}
+
+	if (ensure(ParticleDataPtr->ParticleIndex == GeometryInstanceToSelectPtr->GetOwningParticleID()))
+	{
+		CurrentSelectedGeometryInstance = GeometryInstanceToSelect;
+	}
+}
+
+void FChaosVDSceneParticle::HandleNewGeometryLoaded(uint32 GeometryID, const Chaos::FConstImplicitObjectPtr& InGeometryData)
+{
+	UpdateGeometry(InGeometryData, EChaosVDActorGeometryUpdateFlags::ForceUpdate);
+}
+
+void FChaosVDSceneParticle::HandleDeSelected()
+{
+	CurrentSelectedGeometryInstance = nullptr;
+	UpdateMeshInstancesSelectionState();
+}
+
+void FChaosVDSceneParticle::HandleSelected()
+{
+	UpdateMeshInstancesSelectionState();
+}
+
+bool FChaosVDSceneParticle::IsSelected()
+{
+	// The implementation of this method in UObject, used a global edit callback,
+	// but as we don't use the global editor selection system, we need to re-route it.
+	if (const TSharedPtr<FChaosVDScene> ScenePtr = SceneWeakPtr.Pin())
+	{
+		return ScenePtr->IsSelected(Chaos::VD::TypedElementDataUtil::AcquireTypedElementHandleForStruct(this, true));
+	}
+
+	return false;
+}
+
+void FChaosVDSceneParticle::ProcessUpdatedAndRemovedHandles(TArray<TSharedPtr<FChaosVDExtractedGeometryDataHandle>>& OutExtractedGeometryDataHandles)
+{
+	for (TArray<TSharedRef<FChaosVDMeshDataInstanceHandle>>::TIterator MeshDataHandleRemoveIterator = MeshDataHandles.CreateIterator(); MeshDataHandleRemoveIterator; ++MeshDataHandleRemoveIterator)
+	{
+		TSharedRef<FChaosVDMeshDataInstanceHandle>& ExistingMeshDataHandle = *MeshDataHandleRemoveIterator;
+
+		bool bExists = false;
+
+		// TODO: This search is n2, but I didn't see this as bottleneck. We should check if it is worth adding this to a TSet, or implementing the < operator so we can sort the array and do a binary search
+		// (avoiding the need to allocate a new container) 
+		for (TArray<TSharedPtr<FChaosVDExtractedGeometryDataHandle>>::TIterator HandleRemoveIterator = OutExtractedGeometryDataHandles.CreateIterator(); HandleRemoveIterator; ++HandleRemoveIterator)
+		{
+			const TSharedPtr<FChaosVDExtractedGeometryDataHandle> GeometryDataHandle = *HandleRemoveIterator;
+			const TSharedRef<FChaosVDExtractedGeometryDataHandle> ExistingComponentGeometryDataHandle = ExistingMeshDataHandle->GetGeometryHandle();
+
+			if (GeometryDataHandle && *GeometryDataHandle == *ExistingMeshDataHandle->GetGeometryHandle())
+			{
+				bExists = true;
+
+				// Although the geometry is the same, we need to copy over all the new data on the updated handle
+				// Otherwise the ptr to the root implicit object or the Shape Instance Index will be outdated 
+				*ExistingMeshDataHandle->GetGeometryHandle() = *GeometryDataHandle;
+
+				// If we have a CVD Geometry Component for this handle, just remove it from the list as it means we don't need to re-create it
+				HandleRemoveIterator.RemoveCurrent();
+				break;
+			}
+		}
+
+		if (!bExists)
+		{
+			if (IChaosVDGeometryComponent* AsGeometryComponent = Cast<IChaosVDGeometryComponent>(ExistingMeshDataHandle->GetMeshComponent()))
+			{
+				AsGeometryComponent->RemoveMeshInstance(ExistingMeshDataHandle);
+			}
+
+			MeshDataHandleRemoveIterator.RemoveCurrent();
+		}		
+	}
+}
+
+void FChaosVDSceneParticle::UpdateGeometry(const Chaos::FConstImplicitObjectPtr& InImplicitObject, EChaosVDActorGeometryUpdateFlags OptionsFlags)
+{
+	if (EnumHasAnyFlags(OptionsFlags, EChaosVDActorGeometryUpdateFlags::ForceUpdate))
+	{
+		bIsGeometryDataGenerationStarted = false;
+	}
+
+	if (bIsGeometryDataGenerationStarted)
+	{
+		return;
+	}
+
+	if (!ParticleDataPtr)
+	{
+		return;
+	}
+
+	if (!InImplicitObject.IsValid())
+	{
+		return;
+	}
+
+	const TSharedPtr<FChaosVDScene> ScenePtr = SceneWeakPtr.Pin();
+	if (!ScenePtr.IsValid())
+	{
+		return;
+	}
+
+	const TSharedPtr<FChaosVDGeometryBuilder> GeometryGenerator = ScenePtr->GetGeometryGenerator().Pin();
+	if (!GeometryGenerator.IsValid())
+	{
+		return;
+	}
+
+	const int32 ObjectsToGenerateNum = InImplicitObject->CountLeafObjectsInHierarchyImpl();
+
+	// If the new implicit object is empty, then we can just clear all the mesh components and early out
+	if (ObjectsToGenerateNum == 0)
+	{
+		VisitGeometryInstances([this](const TSharedRef<FChaosVDMeshDataInstanceHandle>& MeshDataHandle)
+		{
+			if (IChaosVDGeometryComponent* AsGeometryComponent = Cast<IChaosVDGeometryComponent>(MeshDataHandle->GetMeshComponent()))
+			{
+				AsGeometryComponent->RemoveMeshInstance(MeshDataHandle);
+			}
+		});
+
+		MeshDataHandles.Reset();
+		return;
+	}
+	
+	TArray<TSharedPtr<FChaosVDExtractedGeometryDataHandle>> OutExtractedGeometryDataHandles;
+	OutExtractedGeometryDataHandles.Reserve(ObjectsToGenerateNum);
+
+	// Heightfields need to be created as Static meshes and use normal Static Mesh components because we need LODs for them due to their high triangle count
+	const bool bHasToUseStaticMeshComponent = Chaos::VisualDebugger::Cvars::bForceStaticMeshComponentUse || FChaosVDGeometryBuilder::DoesImplicitContainType(InImplicitObject, Chaos::ImplicitObjectType::HeightField);
+	constexpr int32 LODsToGenerateNum = 3;
+	constexpr int32 LODsToGenerateNumForInstancedStaticMesh = 0;
+
+	GeometryGenerator->CreateMeshesFromImplicitObject(InImplicitObject, OutExtractedGeometryDataHandles, ParticleDataPtr->CollisionDataPerShape.Num(), bHasToUseStaticMeshComponent ? LODsToGenerateNum : LODsToGenerateNumForInstancedStaticMesh);
+
+	// This should not happen in theory, but there might be some valid situations where it does. Adding an ensure to catch them and then evaluate if it is really an issue (if it is not I will remove the ensure later on). 
+	if (!ensure(ObjectsToGenerateNum == OutExtractedGeometryDataHandles.Num()))
+	{
+		UE_LOG(LogChaosVDEditor, Warning, TEXT("[%s] Geometry objects being generated doesn't match the number of objects in the implicit object | Expected [%d] | Being generated [%d] | Particle Actor [%s]"), ANSI_TO_TCHAR(__FUNCTION__), ObjectsToGenerateNum, OutExtractedGeometryDataHandles.Num(), *GetDisplayName());
+	}
+
+	// Figure out what geometry was removed, and destroy their components as needed. Also, if a geometry is already generated an active, remove it from the geometry to generate list
+	ProcessUpdatedAndRemovedHandles(OutExtractedGeometryDataHandles);
+
+	if (OutExtractedGeometryDataHandles.Num() > 0)
+	{
+		for (const TSharedPtr<FChaosVDExtractedGeometryDataHandle>& ExtractedGeometryDataHandle : OutExtractedGeometryDataHandles)
+		{
+			//TODO: Time Slice component creation
+			TSharedPtr<FChaosVDMeshDataInstanceHandle> MeshDataInstance;
+
+			if (bHasToUseStaticMeshComponent)
+			{
+				MeshDataInstance = GeometryGenerator->CreateMeshDataInstance<UChaosVDStaticMeshComponent>(*ParticleDataPtr.Get(), ExtractedGeometryDataHandle);
+			}
+			else
+			{
+				MeshDataInstance = GeometryGenerator->CreateMeshDataInstance<UChaosVDInstancedStaticMeshComponent>(*ParticleDataPtr.Get(), ExtractedGeometryDataHandle);
+			}
+
+			if (!MeshDataInstance)
+			{
+				UE_LOG(LogChaosVDEditor, Error, TEXT("[%s] Failed To Create mesh data instance for [%s]"), ANSI_TO_TCHAR(__FUNCTION__), *GetDisplayName());
+				continue;
+			}
+
+			const UMeshComponent* CreatedMeshComponent = MeshDataInstance->GetMeshComponent();
+			if (!CreatedMeshComponent)
+			{
+				UE_LOG(LogChaosVDEditor, Error, TEXT("[%s] Failed To Create mesh component for [%s]"), ANSI_TO_TCHAR(__FUNCTION__), *GetDisplayName());
+				continue;
+			}
+
+			// If we have a valid transform data, we need to update our instance with it as the mesh component is not part of this actor (and event if it is, we don't use the actor transform anymore)
+			if (ParticleDataPtr && ParticleDataPtr->ParticlePositionRotation.HasValidData())
+			{
+				const FVector TargetLocation = CachedSimulationTransform.TransformPosition(ParticleDataPtr->ParticlePositionRotation.MX);
+				const FQuat TargetRotation = CachedSimulationTransform.GetRotation() * ParticleDataPtr->ParticlePositionRotation.MR;
+
+				FTransform ParticleTransform;
+				ParticleTransform.SetLocation(TargetLocation);
+				ParticleTransform.SetRotation(TargetRotation);
+
+				MeshDataInstance->SetWorldTransform(ParticleTransform);
+			}
+
+			MeshDataHandles.Add(MeshDataInstance.ToSharedRef());
+		}
+
+		// Ensure that visibility and colorization is up to date after updating this Particle's Geometry
+		EnumAddFlags(DirtyFlags, EChaosVDSceneParticleDirtyFlags::Visibility);
+		EnumAddFlags(DirtyFlags, EChaosVDSceneParticleDirtyFlags::Coloring);
+		EnumRemoveFlags(DirtyFlags, EChaosVDSceneParticleDirtyFlags::Geometry);
+
+		bIsGeometryDataGenerationStarted = true;
+	}
+}
+
+void FChaosVDSceneParticle::UpdateGeometry(uint32 NewGeometryHash, EChaosVDActorGeometryUpdateFlags OptionsFlags)
+{
+	if (const TSharedPtr<FChaosVDScene>& ScenePtr = SceneWeakPtr.Pin())
+	{
+		if (const Chaos::FConstImplicitObjectPtr& Geometry = ScenePtr->GetUpdatedGeometry(NewGeometryHash))
+		{
+			UpdateGeometry(Geometry, OptionsFlags);
+		}
+	}
+}
+
+FBox FChaosVDSceneParticle::GetComponentsBoundingBox() const
+{
+	FBox BoundingBox = FBox(ForceInitToZero);
+	if (ParticleDataPtr)
+	{
+		FBoxSphereBounds::Builder BoundsBuilder;
+
+		for (const TSharedRef<FChaosVDMeshDataInstanceHandle>& MeshDataHandle : MeshDataHandles)
+		{
+			if (const UStaticMeshComponent* StaticMeshComponent = Cast<UStaticMeshComponent>(MeshDataHandle->GetMeshComponent()))
+			{
+				if (const UStaticMesh* Mesh = StaticMeshComponent->GetStaticMesh())
+				{
+					BoundsBuilder += Mesh->GetBounds().TransformBy(MeshDataHandle->GetWorldTransform());
+				}
+			}
+		}
+
+		// Geometry might not be generated yet, so we need a placeholder box so the Focus on object feature works
+		if (!BoundsBuilder.IsValid())
+		{
+			return BoundingBox.ExpandBy(10.0f).MoveTo(ParticleDataPtr->ParticlePositionRotation.MX);
+		}
+		else
+		{
+			BoundingBox = FBoxSphereBounds(BoundsBuilder).GetBox();
+		}
+	}
+
+	return BoundingBox;
+}
+
+TConstArrayView<TSharedPtr<FChaosVDParticlePairMidPhase>> FChaosVDSceneParticle::GetCollisionData()
+{
+	if (const TArray<TSharedPtr<FChaosVDParticlePairMidPhase>>* MidPhases = GetCollisionMidPhasesArray())
+	{
+		return *MidPhases;
+	}
+	return TConstArrayView<TSharedPtr<FChaosVDParticlePairMidPhase>>();
+}
+
+bool FChaosVDSceneParticle::HasCollisionData()
+{
+	if (const TArray<TSharedPtr<FChaosVDParticlePairMidPhase>>* MidPhases = GetCollisionMidPhasesArray())
+	{
+		return MidPhases->Num() > 0;
+	}
+
+	return false;
+}
+
+FName FChaosVDSceneParticle::GetProviderName()
+{
+	return ParticleDataPtr ? FName(ParticleDataPtr->DebugName) : NAME_None;
+}
+
+void FChaosVDSceneParticle::UpdateMeshInstancesSelectionState()
+{
+	TSharedPtr<FChaosVDMeshDataInstanceHandle> CurrentSelectedGeometry = CurrentSelectedGeometryInstance.Pin();
+	const bool bIsOwningParticleSelectedInEditor = IsSelected();
+	VisitGeometryInstances([this, bIsOwningParticleSelectedInEditor, CurrentSelectedGeometry](const TSharedRef<FChaosVDMeshDataInstanceHandle>& MeshDataHandle)
+	{
+		const bool bShouldSelectInstance = bIsOwningParticleSelectedInEditor ? (CurrentSelectedGeometry ? CurrentSelectedGeometryInstance == MeshDataHandle : true) : false;
+
+		MeshDataHandle->SetIsSelected(bShouldSelectInstance);
+	});
+}
+
+void FChaosVDSceneParticle::GetCharacterGroundConstraintData(TArray<TSharedPtr<FChaosVDCharacterGroundConstraint>>& OutConstraintsFound)
+{
+	if (const TArray<TSharedPtr<FChaosVDConstraintDataWrapperBase>>* Constraints = GetCharacterGroundConstraintArray())
+	{
+		OutConstraintsFound.Reserve(Constraints->Num());
+
+		for (const TSharedPtr<FChaosVDConstraintDataWrapperBase>& Constraint : *Constraints)
+		{
+			OutConstraintsFound.Add(StaticCastSharedPtr<FChaosVDCharacterGroundConstraint>(Constraint));
+		}
+	}
+}
+
+bool FChaosVDSceneParticle::HasCharacterGroundConstraintData()
+{
+	if (const TArray<TSharedPtr<FChaosVDConstraintDataWrapperBase>>* Constraints = GetCharacterGroundConstraintArray())
+	{
+		return Constraints->Num() > 0;
+	}
+
+	return false;
+}
+
+const TArray<TSharedPtr<FChaosVDParticlePairMidPhase>>* FChaosVDSceneParticle::GetCollisionMidPhasesArray() const
+{
+	if (!ParticleDataPtr.IsValid())
+	{
+		return nullptr;
+	}
+
+	const TSharedPtr<FChaosVDScene> ScenePtr = SceneWeakPtr.Pin();
+	if (!ScenePtr.IsValid())
+	{
+		return nullptr;
+	}
+
+	if (AChaosVDSolverInfoActor* SolverInfoActor = ScenePtr->GetSolverInfoActor(ParticleDataPtr->SolverID))
+	{
+		if (const UChaosVDSolverCollisionDataComponent* CollisionDataComponent = SolverInfoActor->GetCollisionDataComponent())
+		{
+			return CollisionDataComponent->GetMidPhasesForParticle(ParticleDataPtr->ParticleIndex, EChaosVDParticlePairSlot::Any);
+		}
+	}
+
+	return nullptr;
+}
+
+const TArray<TSharedPtr<FChaosVDConstraintDataWrapperBase>>* FChaosVDSceneParticle::GetCharacterGroundConstraintArray() const
+{
+	if (!ParticleDataPtr.IsValid())
+	{
+		return nullptr;
+	}
+
+	const TSharedPtr<FChaosVDScene> ScenePtr = SceneWeakPtr.Pin();
+	if (!ScenePtr.IsValid())
+	{
+		return nullptr;
+	}
+
+	if (AChaosVDSolverInfoActor* SolverInfoActor = ScenePtr->GetSolverInfoActor(ParticleDataPtr->SolverID))
+	{
+		if (const UChaosVDSolverCharacterGroundConstraintDataComponent* ConstraintDataComponent = SolverInfoActor->GetCharacterGroundConstraintDataComponent())
+		{
+			return ConstraintDataComponent->GetConstraintsForParticle(ParticleDataPtr->ParticleIndex, EChaosVDParticlePairSlot::Primary);
+		}
+	}
+
+	return nullptr;
+}
+
+void FChaosVDSceneParticle::UpdateShapeDataComponents()
+{
+	VisitGeometryInstances([this](const TSharedRef<FChaosVDMeshDataInstanceHandle>& MeshDataHandle)
+	{
+		if (ParticleDataPtr)
+		{
+			FChaosVDGeometryComponentUtils::UpdateCollisionDataFromShapeArray(ParticleDataPtr->CollisionDataPerShape, MeshDataHandle);
+		}
+	});
+
+	EnumAddFlags(DirtyFlags, EChaosVDSceneParticleDirtyFlags::Visibility);
+	EnumRemoveFlags(DirtyFlags, EChaosVDSceneParticleDirtyFlags::CollisionData);
+}
+
+void FChaosVDSceneParticle::ApplyPendingTransformData()
+{
+	VisitGeometryInstances([this](const TSharedRef<FChaosVDMeshDataInstanceHandle>& MeshDataHandle)
+	{
+		MeshDataHandle->SetWorldTransform(PendingParticleTransform);
+	});
+
+	EnumRemoveFlags(DirtyFlags, EChaosVDSceneParticleDirtyFlags::Transform);
+}
+
+void FChaosVDSceneParticle::SetParent(const TSharedPtr<FChaosVDBaseSceneObject>& NewParent)
+{
+	FChaosVDBaseSceneObject::SetParent(NewParent);
+
+	EnumAddFlags(DirtyFlags, EChaosVDSceneParticleDirtyFlags::TEDS);
+}
+
+void FChaosVDSceneParticle::UpdateGeometryComponentsVisibility()
+{
+	VisitGeometryInstances([this](const TSharedRef<FChaosVDMeshDataInstanceHandle>& MeshDataHandle)
+	{
+		if (ParticleDataPtr)
+		{
+			FChaosVDGeometryComponentUtils::UpdateMeshVisibility(MeshDataHandle, *ParticleDataPtr.Get(), IsVisible());
+		}
+	});
+	
+	EnumRemoveFlags(DirtyFlags, EChaosVDSceneParticleDirtyFlags::Visibility);
+	
+	if (TSharedPtr<FChaosVDScene> ScenePtr = SceneWeakPtr.Pin())
+	{
+		ScenePtr->RequestUpdate();
+	}
+}
+
+void FChaosVDSceneParticle::UpdateGeometryColors()
+{
+	VisitGeometryInstances([this](const TSharedRef<FChaosVDMeshDataInstanceHandle>& MeshDataHandle)
+	{
+		if (ParticleDataPtr)
+		{
+			FChaosVDGeometryComponentUtils::UpdateMeshColor(MeshDataHandle, *ParticleDataPtr.Get(), GetIsServerParticle());
+		}
+	});
+	
+	EnumRemoveFlags(DirtyFlags, EChaosVDSceneParticleDirtyFlags::Coloring);
+}
+
+void FChaosVDSceneParticle::SetIsActive(bool bNewActive)
+{
+	if (bIsActive != bNewActive)
+	{
+		bIsActive = bNewActive;
+
+		if (bNewActive)
+		{
+			RemoveHiddenFlag(EChaosVDHideParticleFlags::HiddenByActiveState);
+		}
+		else
+		{
+			AddHiddenFlag(EChaosVDHideParticleFlags::HiddenByActiveState);
+		}
+
+		EnumAddFlags(DirtyFlags, EChaosVDSceneParticleDirtyFlags::TEDS);
+		EnumAddFlags(DirtyFlags, EChaosVDSceneParticleDirtyFlags::Active);
+	}
+}
+
+void FChaosVDSceneParticle::AddHiddenFlag(EChaosVDHideParticleFlags Flag)
+{
+	bool bOldIsVisible = IsVisible();
+
+	EnumAddFlags(HideParticleFlags, Flag);
+
+	if (bOldIsVisible != IsVisible())
+	{
+		EnumAddFlags(DirtyFlags, EChaosVDSceneParticleDirtyFlags::Visibility);
+	}
+}
+
+void FChaosVDSceneParticle::RemoveHiddenFlag(EChaosVDHideParticleFlags Flag)
+{
+	bool bOldIsVisible = IsVisible();
+
+	EnumRemoveFlags(HideParticleFlags, Flag);
+
+	if (bOldIsVisible != IsVisible())
+	{
+		EnumAddFlags(DirtyFlags, EChaosVDSceneParticleDirtyFlags::Visibility);
+	}
+}
diff --git a/Engine/Plugins/ChaosVD/Source/ChaosVD/Private/ChaosVDStyle.cpp b/Engine/Plugins/ChaosVD/Source/ChaosVD/Private/ChaosVDStyle.cpp
--- a/Engine/Plugins/ChaosVD/Source/ChaosVD/Private/ChaosVDStyle.cpp
+++ b/Engine/Plugins/ChaosVD/Source/ChaosVD/Private/ChaosVDStyle.cpp
@@ -73,6 +73,9 @@ TSharedRef< FSlateStyleSet > FChaosVDStyle::Create()
 	Style->Set("LinkedIcon", new IMAGE_BRUSH_SVG(TEXT("Linked"), Icon16x16));
 	Style->Set("UnLinkedIcon", new IMAGE_BRUSH_SVG(TEXT("UnLinked"), Icon16x16));
 
+	Style->Set("RigidBodyIcon", new IMAGE_BRUSH(TEXT("RigidBodyIcon_16x"), Icon16x16));
+	Style->Set("SolverIcon", new IMAGE_BRUSH(TEXT("SolverIcon_16x"), Icon16x16));
+
 	return Style;
 }
 
diff --git a/Engine/Plugins/ChaosVD/Source/ChaosVD/Private/ChaosVDWorldOutlinerMode.cpp b/Engine/Plugins/ChaosVD/Source/ChaosVD/Private/ChaosVDWorldOutlinerMode.cpp
--- a/Engine/Plugins/ChaosVD/Source/ChaosVD/Private/ChaosVDWorldOutlinerMode.cpp
+++ b/Engine/Plugins/ChaosVD/Source/ChaosVD/Private/ChaosVDWorldOutlinerMode.cpp
@@ -4,137 +4,116 @@
 
 #include "ActorTreeItem.h"
 #include "ChaosVDModule.h"
-#include "ChaosVDParticleActor.h"
+#include "ChaosVDSceneParticle.h"
 #include "ChaosVDPlaybackController.h"
 #include "ChaosVDScene.h"
+#include "TedsOutlinerItem.h"
+#include "Elements/Actor/ActorElementData.h"
+#include "Elements/Columns/TypedElementCompatibilityColumns.h"
+#include "Elements/Columns/TypedElementTypeInfoColumns.h"
+#include "Elements/Common/EditorDataStorageFeatures.h"
+#include "Elements/Framework/EngineElementsLibrary.h"
 #include "Elements/Framework/TypedElementSelectionSet.h"
+#include "Elements/Interfaces/TypedElementDataStorageCompatibilityInterface.h"
+#include "TEDS/ChaosVDSelectionInterface.h"
+#include "TEDS/ChaosVDStructTypedElementData.h"
+#include "TEDS/ChaosVDTedsUtils.h"
 
-namespace Chaos::VisualDebugger::Cvars
-{
-static bool bQueueAndCombineSceneOutlinerEvent = true;
-static FAutoConsoleVariableRef CVarChaosVDQueueAndCombineSceneOutlinerEvents(
-	TEXT("p.Chaos.VD.Tool.QueueAndCombineSceneOutlinerEvents"),
-	bQueueAndCombineSceneOutlinerEvent,
-	TEXT("If set to true, scene outliner events will be queued and sent once per frame. If there was a unprocessed event for an item, the las queued event will replace it"));
-
-static bool bPurgeInvalidOutlinerItemsBeforeBroadcast = true;
-static FAutoConsoleVariableRef CVarChaosVDPurgeInvalidOutlinerItemsBeforeBroadcast(
-	TEXT("p.Chaos.VD.Tool.PurgeInvalidOutlinerItemsBeforeBroadcast"),
-	bPurgeInvalidOutlinerItemsBeforeBroadcast,
-	TEXT("If set to true, scene outliner events will evaluated and any invalid outliner event in them will be removed before broadcasting the hierarchy change."));
-}
-
-const FSceneOutlinerTreeItemType FChaosVDActorTreeItem::Type(&FActorTreeItem::Type);
-
-#define LOCTEXT_NAMESPACE "ChaosVisualDebugger"
+struct FTypedElementScriptStructTypeInfoColumn;
 
-const TSharedRef<SWidget> FChaosVDSceneOutlinerGutter::ConstructRowWidget(FSceneOutlinerTreeItemRef TreeItem, const STableRow<FSceneOutlinerTreeItemPtr>& Row)
+namespace Chaos::VD::SceneOutlinerUtils::Private
 {
-	if (TreeItem->ShouldShowVisibilityState())
+	FChaosVDSceneParticle* GetParticleInstanceFromOutlinerItem(const TSharedRef<ISceneOutlinerTreeItem>& TreeItemRef)
 	{
-		return SNew(SHorizontalBox)
-			+SHorizontalBox::Slot()
-			.AutoWidth()
-			.VAlign(VAlign_Center)
-			[
-				SNew(SVisibilityWidget, SharedThis(this), WeakOutliner, TreeItem, &Row)
-				.IsEnabled_Raw(this, &FChaosVDSceneOutlinerGutter::IsEnabled, TreeItem->AsWeak())
-				.ToolTipText_Raw(this, &FChaosVDSceneOutlinerGutter::GetVisibilityTooltip, TreeItem->AsWeak())
-			];
+		using namespace UE::Editor::Outliner;
+		using namespace UE::Editor::DataStorage;
+		if (FTedsOutlinerTreeItem* TEDSItem = TreeItemRef->CastTo<FTedsOutlinerTreeItem>())
+		{
+			IEditorDataStorageProvider* Storage = GetMutableDataStorageFeature<IEditorDataStorageProvider>(StorageFeatureName);
+			if (TEDSItem->IsValid() && Storage)
+			{
+				const FTypedElementExternalObjectColumn* RawObjectColumn = Storage->GetColumn<FTypedElementExternalObjectColumn>(TEDSItem->GetRowHandle());
+				const FTypedElementScriptStructTypeInfoColumn* TypeInfoColumn = Storage->GetColumn<FTypedElementScriptStructTypeInfoColumn>(TEDSItem->GetRowHandle());
+				if (RawObjectColumn && TypeInfoColumn)
+				{
+					if (RawObjectColumn->Object && TypeInfoColumn->TypeInfo == FChaosVDSceneParticle::StaticStruct())
+					{
+						return static_cast<FChaosVDSceneParticle*>(RawObjectColumn->Object);
+					}
+				}
+			}
+		}
+	
+		return nullptr;
 	}
-	return SNullWidget::NullWidget;
-}
-
-FText FChaosVDSceneOutlinerGutter::GetVisibilityTooltip(TWeakPtr<ISceneOutlinerTreeItem> WeakTreeItem) const
-{
-
-	return IsEnabled(WeakTreeItem) ? LOCTEXT("SceneOutlinerVisibilityToggleTooltip", "Toggles the visibility of this object in the level editor.") :
-										LOCTEXT("SceneOutlinerVisibilityToggleDisabkedTooltip", "Visibility of this object is being controlled by another visibility setting");
-}
 
-bool FChaosVDSceneOutlinerGutter::IsEnabled(TWeakPtr<ISceneOutlinerTreeItem> WeakTreeItem) const
-{
-	bool bIsEnabled = true;
-	if (const TSharedPtr<ISceneOutlinerTreeItem> TreeItem = WeakTreeItem.Pin())
+	AActor* GetActorFromOutlinerItem(const TSharedRef<ISceneOutlinerTreeItem>& TreeItemRef)
 	{
-		if (const FActorTreeItem* ActorItem = TreeItem->CastTo<FActorTreeItem>())
+		using namespace UE::Editor::Outliner;
+		using namespace UE::Editor::DataStorage;
+
+		if (FTedsOutlinerTreeItem* TEDSItem = TreeItemRef->CastTo<FTedsOutlinerTreeItem>())
 		{
-			if (const AChaosVDParticleActor* CVDActor = Cast<AChaosVDParticleActor>(ActorItem->Actor.Get()))
+			IEditorDataStorageProvider* Storage = GetMutableDataStorageFeature<IEditorDataStorageProvider>(StorageFeatureName);
+			if (TEDSItem->IsValid() && Storage)
 			{
-				bIsEnabled = CVDActor->GetHideFlags() == EChaosVDHideParticleFlags::HiddenBySceneOutliner || CVDActor->GetHideFlags() == EChaosVDHideParticleFlags::None;
+				const FTypedElementUObjectColumn* RawObjectColumn = Storage->GetColumn<FTypedElementUObjectColumn>(TEDSItem->GetRowHandle());
+				return RawObjectColumn ? Cast<AActor>(RawObjectColumn->Object) : nullptr;
 			}
 		}
+	
+		return nullptr;
 	}
 
-	return bIsEnabled;
-}
-
-bool FChaosVDActorTreeItem::GetVisibility() const
-{
-	if (const AChaosVDParticleActor* CVDActor = Cast<AChaosVDParticleActor>(Actor.Get()))
-	{
-		return CVDActor->IsVisible();
-	}
-	else
+	AActor* GetActorFromOutlinerItem(const TWeakPtr<ISceneOutlinerTreeItem>& WeakTreeItem)
 	{
-		return FActorTreeItem::GetVisibility();
+		if (const TSharedPtr<ISceneOutlinerTreeItem> TreeItem = WeakTreeItem.Pin())
+		{
+			return GetActorFromOutlinerItem(TreeItem.ToSharedRef());
+		}
+	
+		return nullptr;
 	}
-}
 
-void FChaosVDActorTreeItem::OnVisibilityChanged(const bool bNewVisibility)
-{
-	if (AChaosVDParticleActor* CVDActor = Cast<AChaosVDParticleActor>(Actor.Get()))
+	FChaosVDSceneParticle* GetParticleInstanceFromOutlinerItem(const TWeakPtr<ISceneOutlinerTreeItem>& WeakTreeItem)
 	{
-		if (bNewVisibility)
-		{
-			CVDActor->RemoveHiddenFlag(EChaosVDHideParticleFlags::HiddenBySceneOutliner);
-		}
-		else
+		if (const TSharedPtr<ISceneOutlinerTreeItem> TreeItem = WeakTreeItem.Pin())
 		{
-			CVDActor->AddHiddenFlag(EChaosVDHideParticleFlags::HiddenBySceneOutliner);
+			return GetParticleInstanceFromOutlinerItem(TreeItem.ToSharedRef());
 		}
-	}
-	else
-	{
-		FActorTreeItem::OnVisibilityChanged(bNewVisibility);
+		return nullptr;
 	}
 }
 
-void FChaosVDActorTreeItem::UpdateDisplayString()
+#define LOCTEXT_NAMESPACE "ChaosVisualDebugger"
+
+bool FChaosVDParticleOutlinerSelector::operator()(const TWeakPtr<ISceneOutlinerTreeItem>& Item, FChaosVDSceneParticle*& ParticlePtrOut) const
 {
-	if (AChaosVDParticleActor* CVDActor = Cast<AChaosVDParticleActor>(Actor.Get()))
-	{
-		if (TSharedPtr<const FChaosVDParticleDataWrapper> ParticleData = CVDActor->GetParticleData())
-		{
-			const bool bHasDebugName = !ParticleData->DebugName.IsEmpty();
-			DisplayString = bHasDebugName ? ParticleData->DebugName : TEXT("Unnamed Particle - ID : ") + FString::FromInt(ParticleData->ParticleIndex);
-		}
-	}
-	else
+	if (FChaosVDSceneParticle* ParticleInstance = Chaos::VD::SceneOutlinerUtils::Private::GetParticleInstanceFromOutlinerItem(Item))
 	{
-		FActorTreeItem::UpdateDisplayString();
+		ParticlePtrOut = ParticleInstance;
+		return true;
 	}
+
+	return false;
 }
 
-TUniquePtr<FChaosVDOutlinerHierarchy> FChaosVDOutlinerHierarchy::Create(ISceneOutlinerMode* Mode, const TWeakObjectPtr<UWorld>& World)
+bool FChaosVDActorOutlinerSelector::operator()(const TWeakPtr<ISceneOutlinerTreeItem>& Item, AActor*& ActorPtrOut) const
 {
-	FChaosVDOutlinerHierarchy* Hierarchy = new FChaosVDOutlinerHierarchy(Mode, World);
-
-	Create_Internal(Hierarchy, World);
+	if (AActor* Actor = Chaos::VD::SceneOutlinerUtils::Private::GetActorFromOutlinerItem(Item))
+	{
+		ActorPtrOut = Actor;
+		return true;
+	}
 
-	return TUniquePtr<FChaosVDOutlinerHierarchy>(Hierarchy);
+	return false;
 }
 
 
-FSceneOutlinerTreeItemPtr FChaosVDOutlinerHierarchy::CreateItemForActor(AActor* InActor, bool bForce) const
-{
-	return Mode->CreateItemFor<FChaosVDActorTreeItem>(InActor, bForce);
-}
-
-FChaosVDWorldOutlinerMode::FChaosVDWorldOutlinerMode(const FActorModeParams& InModeParams, TWeakPtr<FChaosVDScene> InScene, TWeakPtr<FChaosVDPlaybackController> InPlaybackController)
-	: FActorMode(InModeParams),
-	CVDScene(InScene),
-	PlaybackController(InPlaybackController)
+FChaosVDWorldOutlinerMode::FChaosVDWorldOutlinerMode(const UE::Editor::Outliner::FTedsOutlinerParams& InModeParams, TWeakPtr<FChaosVDScene> InScene, TWeakPtr<FChaosVDPlaybackController> InPlaybackController)
+	: FTedsOutlinerMode(InModeParams)
+	, CVDScene(InScene)
+	, PlaybackController(InPlaybackController)
 {
 	TSharedPtr<FChaosVDScene> ScenePtr = CVDScene.Pin();
 	if (!ensure(ScenePtr.IsValid()))
@@ -142,19 +121,11 @@ FChaosVDWorldOutlinerMode::FChaosVDWorldOutlinerMode(const FActorModeParams& InM
 		return;
 	}
 
-	ScenePtr->OnActorActiveStateChanged().AddRaw(this, &FChaosVDWorldOutlinerMode::HandleActorActiveStateChanged);
-	ScenePtr->OnActorLabelChanged().AddRaw(this, &FChaosVDWorldOutlinerMode::HandleActorLabelChanged);
-
 	RegisterSelectionSetObject(ScenePtr->GetElementSelectionSet());
 }
 
 FChaosVDWorldOutlinerMode::~FChaosVDWorldOutlinerMode()
 {
-	if (TSharedPtr<FChaosVDScene> ScenePtr = CVDScene.Pin())
-	{
-		ScenePtr->OnActorActiveStateChanged().RemoveAll(this);
-		ScenePtr->OnActorLabelChanged().RemoveAll(this);
-	}
 }
 
 void FChaosVDWorldOutlinerMode::OnItemSelectionChanged(FSceneOutlinerTreeItemPtr Item, ESelectInfo::Type SelectionType, const FSceneOutlinerItemSelection& Selection)
@@ -170,18 +141,24 @@ void FChaosVDWorldOutlinerMode::OnItemSelectionChanged(FSceneOutlinerTreeItemPtr
 		return;
 	}
 
-	TArray<AActor*> OutlinerSelectedActors = Selection.GetData<AActor*>(SceneOutliner::FActorSelector());
+	using namespace Chaos::VD::TypedElementDataUtil;
 
-	if (const UTypedElementSelectionSet* SelectionSet = GetSelectionSetObject())
+	TArray<FChaosVDSceneParticle*> OutlinerSelectedParticles = Selection.GetData<FChaosVDSceneParticle*>(FChaosVDParticleOutlinerSelector());
+	if (!OutlinerSelectedParticles.IsEmpty())
 	{
-		TArray<AActor*> SceneSelectedActors = SelectionSet->GetSelectedObjects<AActor>();
-		ensureMsgf(SceneSelectedActors.Num() < 2, TEXT("Multi Selection is not supported, but [%d] Actors are selected... Choosing the first one"), SceneSelectedActors.Num());
-
-		AActor* SelectedActor = OutlinerSelectedActors.IsEmpty() ? nullptr : OutlinerSelectedActors[0];
+		ScenePtr->SetSelected(AcquireTypedElementHandleForStruct(OutlinerSelectedParticles[0], true));
+		return;
+	}
 
-		// If the actor ptr is null, the selection will be cleared
-		ScenePtr->SetSelectedObject(SelectedActor);
+	TArray<AActor*> OutlinerSelectedActors = Selection.GetData<AActor*>(FChaosVDActorOutlinerSelector());
+	
+	if (!OutlinerSelectedActors.IsEmpty())
+	{
+		ScenePtr->SetSelected(UEngineElementsLibrary::AcquireEditorActorElementHandle(OutlinerSelectedActors[0]));
+		return;
 	}
+
+	ScenePtr->SetSelected(FTypedElementHandle());
 }
 
 void FChaosVDWorldOutlinerMode::OnItemDoubleClick(FSceneOutlinerTreeItemPtr Item)
@@ -192,75 +169,20 @@ void FChaosVDWorldOutlinerMode::OnItemDoubleClick(FSceneOutlinerTreeItemPtr Item
 		return;
 	}
 
-	if (const FActorTreeItem* ActorItem = Item->CastTo<FActorTreeItem>())
+	if (FChaosVDSceneParticle* ParticleInstance = Chaos::VD::SceneOutlinerUtils::Private::GetParticleInstanceFromOutlinerItem(Item))
 	{
-		if (AActor* Actor = ActorItem->Actor.Get())
-		{
-			ScenePtr->OnFocusRequest().Broadcast(Actor->GetComponentsBoundingBox(false));
-		}
+		ScenePtr->OnFocusRequest().Broadcast(ParticleInstance->GetComponentsBoundingBox());
 	}
-}
-
-void FChaosVDWorldOutlinerMode::ProcessPendingHierarchyEvents()
-{
-	const double StartTimeSeconds = FPlatformTime::Seconds();
-	double CurrentTimeSpentSeconds = 0.0;
-	int32 CurrentEventProcessedNum = 0;
-
-	constexpr double MaxUpdateBudgetSeconds = 0.002;
-
-	for (TMap<FSceneOutlinerTreeItemID, FSceneOutlinerHierarchyChangedData>::TIterator RemoveIterator = PendingOutlinerEventsMap.CreateIterator(); RemoveIterator; ++RemoveIterator)
+	else if (AActor* ActorItem = Chaos::VD::SceneOutlinerUtils::Private::GetActorFromOutlinerItem(Item))
 	{
-		if (CurrentTimeSpentSeconds > MaxUpdateBudgetSeconds)
-		{
-			return;
-		}
-
-		// Only check the budget every 5 tasks as Getting the current time is a syscall and it is not free
-		if (CurrentEventProcessedNum % 5 == 0)
-		{
-			CurrentTimeSpentSeconds += FPlatformTime::Seconds() - StartTimeSeconds;
-		}
-
-		FSceneOutlinerHierarchyChangedData& HierarchyChangedData = RemoveIterator->Value;
-
-		if (Chaos::VisualDebugger::Cvars::bPurgeInvalidOutlinerItemsBeforeBroadcast && HierarchyChangedData.Type == FSceneOutlinerHierarchyChangedData::Added)
-		{ 
-			HierarchyChangedData.Items.RemoveAllSwap([](const FSceneOutlinerTreeItemPtr& SceneOutlinerItemPtr)
-			{
-				return !SceneOutlinerItemPtr || !SceneOutlinerItemPtr->IsValid();
-			});
-		}
-
-		Hierarchy->OnHierarchyChanged().Broadcast(HierarchyChangedData);
-		RemoveIterator.RemoveCurrent();
-		CurrentEventProcessedNum++;
+		ScenePtr->OnFocusRequest().Broadcast(ActorItem->GetComponentsBoundingBox(false));
 	}
 }
 
-bool FChaosVDWorldOutlinerMode::Tick(float DeltaTime)
-{
-	ProcessPendingHierarchyEvents();
-	return true;
-}
-
-TUniquePtr<ISceneOutlinerHierarchy> FChaosVDWorldOutlinerMode::CreateHierarchy()
-{
-	TUniquePtr<FChaosVDOutlinerHierarchy> ActorHierarchy = FChaosVDOutlinerHierarchy::Create(this, RepresentingWorld);
-
-	ActorHierarchy->SetShowingComponents(!bHideComponents);
-	ActorHierarchy->SetShowingOnlyActorWithValidComponents(!bHideComponents && bHideActorWithNoComponent);
-	ActorHierarchy->SetShowingLevelInstances(!bHideLevelInstanceHierarchy);
-	ActorHierarchy->SetShowingUnloadedActors(!bHideUnloadedActors);
-	ActorHierarchy->SetShowingEmptyFolders(!bHideEmptyFolders);
-	
-	return ActorHierarchy;
-}
-
 bool FChaosVDWorldOutlinerMode::CanInteract(const ISceneOutlinerTreeItem& Item) const
 {
 	// This option is not supported in CVD yet
-	ensure(!bCanInteractWithSelectableActorsOnly);
+	//ensure(!bCanInteractWithSelectableActorsOnly);
 
 	return true;
 }
@@ -276,90 +198,31 @@ bool FChaosVDWorldOutlinerMode::CanPopulate() const
 	return true;
 }
 
-void FChaosVDWorldOutlinerMode::EnqueueAndCombineHierarchyEvent(const FSceneOutlinerTreeItemID& ItemID, const FSceneOutlinerHierarchyChangedData& EnventToProcess)
-{
-	if (FSceneOutlinerHierarchyChangedData* EventData = PendingOutlinerEventsMap.Find(ItemID))
-	{
-		*EventData = EnventToProcess;
-	}
-	else
-	{
-		PendingOutlinerEventsMap.Add(ItemID, EnventToProcess);
-	}	
-}
-void FChaosVDWorldOutlinerMode::HandleActorLabelChanged(AChaosVDParticleActor* ChangedActor)
-{
-	if (!ensure(ChangedActor))
-	{
-		return;
-	}
-
-	if (IsActorDisplayable(ChangedActor) && RepresentingWorld.Get() == ChangedActor->GetWorld())
-	{
-		// Force create the item otherwise the outliner may not be notified of a change to the item if it is filtered out
-		if (FSceneOutlinerTreeItemPtr Item = CreateItemFor<FChaosVDActorTreeItem>(ChangedActor, true))
-		{
-			SceneOutliner->OnItemLabelChanged(Item);
-		}
-	}
-}
-
-void FChaosVDWorldOutlinerMode::HandleActorActiveStateChanged(AChaosVDParticleActor* ChangedActor)
-{
-	if (!ChangedActor)
-	{
-		return;
-	}
-
-	if (Hierarchy.IsValid())
-	{
-		FSceneOutlinerHierarchyChangedData EventData;
-
-		if (ChangedActor->IsActive())
-		{
-			EventData.Type = FSceneOutlinerHierarchyChangedData::Added;
-			EventData.Items.Emplace(CreateItemFor<FChaosVDActorTreeItem>(ChangedActor));
-		}
-		else
-		{
-			EventData.ItemIDs.Emplace(ChangedActor);
-			EventData.Type = FSceneOutlinerHierarchyChangedData::Removed;
-		}
-
-		// There is currently a bug in the Scene Outliner where if opposite events happen multiple times within the same tick, the last ones get dropped
-		// (UE-193877). As our current use case is fairly simple, as a workaround we can just queue the events and process them once per frame
-		// Only taking into account only the last requested event for each item.
-		// Keeping this behind a cvar enabled by default so when the Scene Outliner bug is fixed, we can test it easily.
-		if (Chaos::VisualDebugger::Cvars::bQueueAndCombineSceneOutlinerEvent)
-		{
-			EnqueueAndCombineHierarchyEvent(ChangedActor, EventData);
-		}
-		else
-		{
-			Hierarchy->OnHierarchyChanged().Broadcast(EventData);
-		}
-	}
-}
-
 void FChaosVDWorldOutlinerMode::HandlePostSelectionChange(const UTypedElementSelectionSet* ChangesSelectionSet)
 {
-	TArray<AActor*> SelectedActors = ChangesSelectionSet->GetSelectedObjects<AActor>();
+	TArray<FTypedElementHandle> SelectedParticlesHandles = ChangesSelectionSet->GetSelectedElementHandles(UChaosVDSelectionInterface::StaticClass());
 
-	if (SelectedActors.Num() > 0)
+	if (SelectedParticlesHandles.Num() > 0)
 	{
 		// We don't support multi selection yet
-		ensure(SelectedActors.Num() == 1);
-		AActor* SelectedActor = SelectedActors[0];
+		ensure(SelectedParticlesHandles.Num() == 1);
 		
-		if (FSceneOutlinerTreeItemPtr TreeItem = SceneOutliner->GetTreeItem(SelectedActor, false))
+		using namespace UE::Editor::DataStorage;
+		using namespace Chaos::VD::TypedElementDataUtil;
+		
+		RowHandle RowHandle = InvalidRowHandle;
+	
+		if (const FChaosVDSceneParticle* Particle = GetStructDataFromTypedElementHandle<FChaosVDSceneParticle>(SelectedParticlesHandles[0]))
 		{
-			SceneOutliner->ScrollItemIntoView(TreeItem);
-			SceneOutliner->SetItemSelection(TreeItem, true, ESelectInfo::Direct);
+			RowHandle = Particle->GetTedsRowHandle();
 		}
-		else
+		else if (const IEditorDataStorageCompatibilityProvider* Compatibility = GetDataStorageFeature<IEditorDataStorageCompatibilityProvider>(CompatibilityFeatureName))
 		{
-			UE_LOG(LogChaosVDEditor, Verbose, TEXT("Selected actor is not in the outliner. It might be filtered out"))	
+			RowHandle = Compatibility->FindRowWithCompatibleObject(ActorElementDataUtil::GetActorFromHandle(SelectedParticlesHandles[0]));
 		}
+
+		// Add a selection column in TEDS
+		TedsOutlinerImpl->SetSelection({ RowHandle });
 	}
 }
 
diff --git a/Engine/Plugins/ChaosVD/Source/ChaosVD/Private/ChaosVDWorldOutlinerMode.h b/Engine/Plugins/ChaosVD/Source/ChaosVD/Private/ChaosVDWorldOutlinerMode.h
--- a/Engine/Plugins/ChaosVD/Source/ChaosVD/Private/ChaosVDWorldOutlinerMode.h
+++ b/Engine/Plugins/ChaosVD/Source/ChaosVD/Private/ChaosVDWorldOutlinerMode.h
@@ -7,106 +7,78 @@
 #include "ChaosVDScene.h"
 #include "ChaosVDSceneSelectionObserver.h"
 #include "SceneOutlinerGutter.h"
+#include "TedsOutlinerMode.h"
 #include "Containers/Map.h"
 #include "Templates/SharedPointer.h"
 
 class FChaosVDPlaybackController;
 class FChaosVDScene;
 
-/** Outliner Gutter used to override the behaviour of the visibility widget in the Scene outliner  */
-class FChaosVDSceneOutlinerGutter final : public FSceneOutlinerGutter
+/** Functor which can be used to get actors from a selection including component parents */
+struct FChaosVDParticleOutlinerSelector
 {
-public:
-	explicit FChaosVDSceneOutlinerGutter(ISceneOutliner& Outliner)
-		: FSceneOutlinerGutter(Outliner)
+	FChaosVDParticleOutlinerSelector()
 	{
 	}
 
-	virtual const TSharedRef<SWidget> ConstructRowWidget(FSceneOutlinerTreeItemRef TreeItem, const STableRow<FSceneOutlinerTreeItemPtr>& Row) override;
-	
-	FText GetVisibilityTooltip(TWeakPtr<ISceneOutlinerTreeItem> WeakTreeItem) const;
-
-	bool IsEnabled(TWeakPtr<ISceneOutlinerTreeItem> WeakTreeItem) const;
+	bool operator()(const TWeakPtr<ISceneOutlinerTreeItem>& Item, FChaosVDSceneParticle*& ParticlePtrOut) const;
 };
 
-/** Scene outliner Actor Tree Item used to override the behaviour of the visibility changes for CVD Actors in the Scene outliner  */
-class FChaosVDActorTreeItem final : public FActorTreeItem
+struct FChaosVDActorOutlinerSelector
 {
-public:
-	
-	FChaosVDActorTreeItem(AActor* InActor) : FActorTreeItem(InActor)
-	{
-		UpdateDisplayString();
-	}
-
-	FChaosVDActorTreeItem(const FSceneOutlinerTreeItemType& TypeIn, AActor* InActor)
-		: FActorTreeItem(TypeIn, InActor)
+	FChaosVDActorOutlinerSelector()
 	{
-		UpdateDisplayString();
 	}
 
-	virtual bool GetVisibility() const override;
-
-	virtual void OnVisibilityChanged(const bool bNewVisibility) override;
-
-	static const FSceneOutlinerTreeItemType Type;
-
-	virtual void UpdateDisplayString() override;
+	bool operator()(const TWeakPtr<ISceneOutlinerTreeItem>& Item, AActor*& ActorPtrOut) const;
 };
 
-/** Actor Hierarchy used to override the what Tree items will be used for actors in CVD's Scene outliner  */
-class FChaosVDOutlinerHierarchy final : public FActorHierarchy
-{
-public:
-
-	static TUniquePtr<FChaosVDOutlinerHierarchy> Create(ISceneOutlinerMode* Mode, const TWeakObjectPtr<UWorld>& World);
-
-protected:
-
-	FChaosVDOutlinerHierarchy(ISceneOutlinerMode* Mode, const TWeakObjectPtr<UWorld>& Worlds)
-	: FActorHierarchy(Mode, Worlds)
-	{
-	}
-	virtual FSceneOutlinerTreeItemPtr CreateItemForActor(AActor* InActor, bool bForce) const override;
-};
 
 /**
  * Scene outliner mode used to represent a CVD (Chaos Visual Debugger) world
  * It has a more limited view compared to the normal outliner, hiding features we don't support,
  * and it is integrated with the CVD local selection system
  */
-class FChaosVDWorldOutlinerMode : public FActorMode, public FChaosVDSceneSelectionObserver, public FTSTickerObjectBase
+class FChaosVDWorldOutlinerMode : public UE::Editor::Outliner::FTedsOutlinerMode, public FChaosVDSceneSelectionObserver
 {
 public:
 
-	FChaosVDWorldOutlinerMode(const FActorModeParams& InModeParams, TWeakPtr<FChaosVDScene> InScene, TWeakPtr<FChaosVDPlaybackController> InPlaybackController);
+	FChaosVDWorldOutlinerMode(const UE::Editor::Outliner::FTedsOutlinerParams& InModeParams, TWeakPtr<FChaosVDScene> InScene, TWeakPtr<FChaosVDPlaybackController> InPlaybackController);
 
 	virtual ~FChaosVDWorldOutlinerMode() override;
 
 	virtual void OnItemSelectionChanged(FSceneOutlinerTreeItemPtr Item, ESelectInfo::Type SelectionType, const FSceneOutlinerItemSelection& Selection) override;
 	virtual void OnItemDoubleClick(FSceneOutlinerTreeItemPtr Item) override;
-	void ProcessPendingHierarchyEvents();
-
 
 	virtual void CreateViewContent(FMenuBuilder& MenuBuilder) override
 	{
 		// Intentionally not implemented, as we don't support the built in menu to switch worlds
 	}
 
-	virtual bool ShouldShowFolders() const override { return true;}
-
-	virtual bool Tick(float DeltaTime) override;
-
-	virtual TUniquePtr<ISceneOutlinerHierarchy> CreateHierarchy() override;
+	virtual bool ShouldShowFolders() const override
+	{
+		return true;
+	}
 
 	virtual bool CanInteract(const ISceneOutlinerTreeItem& Item) const override;
 	virtual bool CanPopulate() const override;
 
-private:
 
-	void EnqueueAndCombineHierarchyEvent(const FSceneOutlinerTreeItemID& ItemID, const FSceneOutlinerHierarchyChangedData& EnventToProcess);
-	void HandleActorLabelChanged(AChaosVDParticleActor* ChangedActor);
-	void HandleActorActiveStateChanged(AChaosVDParticleActor* ChangedActor);
+	virtual ESelectionMode::Type GetSelectionMode() const override
+	{
+		return ESelectionMode::Single;
+	}
+	virtual bool CanSupportDragAndDrop() const override
+	{
+		return false;
+	}
+
+	virtual TSharedPtr<SWidget> CreateContextMenu() override
+	{
+		return SNullWidget::NullWidget;
+	}
+
+private:
 
 	virtual void HandlePostSelectionChange(const UTypedElementSelectionSet* ChangesSelectionSet) override;
 
diff --git a/Engine/Plugins/ChaosVD/Source/ChaosVD/Private/ChaosVDWorldOutlinerTab.cpp b/Engine/Plugins/ChaosVD/Source/ChaosVD/Private/ChaosVDWorldOutlinerTab.cpp
--- a/Engine/Plugins/ChaosVD/Source/ChaosVD/Private/ChaosVDWorldOutlinerTab.cpp
+++ b/Engine/Plugins/ChaosVD/Source/ChaosVD/Private/ChaosVDWorldOutlinerTab.cpp
@@ -9,8 +9,18 @@
 #include "Modules/ModuleManager.h"
 #include "SceneOutlinerModule.h"
 #include "SceneOutlinerPublicTypes.h"
+#include "TedsOutlinerImpl.h"
+#include "TedsOutlinerItem.h"
+#include "Compatibility/SceneOutlinerTedsBridge.h"
+
+#include "Elements/Columns/TypedElementCompatibilityColumns.h"
+#include "Elements/Columns/TypedElementLabelColumns.h"
+#include "Elements/Columns/TypedElementVisibilityColumns.h"
+#include "Elements/Common/EditorDataStorageFeatures.h"
+#include "Elements/Framework/TypedElementRegistry.h"
+#include "TEDS/ChaosVDParentDataStorageFactory.h"
 #include "Widgets/Docking/SDockTab.h"
-#include "Engine/World.h"
+#include "TEDS/ChaosVDParticleEditorDataFactory.h"
 #include "Widgets/SChaosVDMainTab.h"
 
 #define LOCTEXT_NAMESPACE "ChaosVisualDebugger"
@@ -23,33 +33,107 @@ void FChaosVDWorldOutlinerTab::CreateWorldOutlinerWidget()
 		PlaybackController = MainTabPtr->GetChaosVDEngineInstance()->GetPlaybackController();
 	}
 
-	FSceneOutlinerInitializationOptions InitOptions;
-	InitOptions.bShowTransient = true;
+	using namespace UE::Editor::DataStorage::Queries;
+	using namespace UE::Editor::Outliner;
 
-	InitOptions.OutlinerIdentifier = TEXT("ChaosVDOutliner");
+	FQueryDescription OutlinerQueryDescription =
+		Select()
+			.ReadOnly<FTypedElementLabelColumn>()
+		.Where()
+			.Any<FChaosVDObjectDataTag, FTypedElementActorTag>()
+			.All<FTypedElementFromCVDWorldTag>()
+		.Compile();
+
+	FSceneOutlinerInitializationOptions InitOptions;
+	InitOptions.bShowHeaderRow = true;
 	InitOptions.FilterBarOptions.bHasFilterBar = true;
 	InitOptions.FilterBarOptions.bUseSharedSettings = false;
+	InitOptions.bShowTransient = true;
+	InitOptions.OutlinerIdentifier = "CVDTEDSOutliner";
+
+	FTedsOutlinerParams Params(nullptr);
+	Params.QueryDescription = OutlinerQueryDescription;
+	Params.bUseDefaultTedsFilters = false;
+	Params.bShowRowHandleColumn = false;
+
+	const FTedsOutlinerHierarchyData::FGetParentRowHandle RowHandleGetter = FTedsOutlinerHierarchyData::FGetParentRowHandle::CreateLambda([](void* InColumnData)
+	{
+		if(const FTableRowParentColumn* ParentColumn = static_cast<FTableRowParentColumn *>(InColumnData))
+		{
+			return ParentColumn->Parent;
+		}
+
+		return InvalidRowHandle;
+	});
+
+	const FTedsOutlinerHierarchyData::FGetChildrenRowsHandles ChildrenRowHandlesGetter = FTedsOutlinerHierarchyData::FGetChildrenRowsHandles::CreateLambda([](void* InColumnData)
+	{
+		if (FChaosVDTableRowParentColumn* ParentColumn = static_cast<FChaosVDTableRowParentColumn *>(InColumnData))
+		{
+			ParentColumn->Children.Reset(ParentColumn->ChildrenSet.Num());
+			for (RowHandle Handle : ParentColumn->ChildrenSet)
+			{
+				ParentColumn->Children.Emplace(Handle);
+			}
+			return TArrayView<RowHandle>(ParentColumn->Children);
+		}
+
+		return TArrayView<RowHandle>();
+	});
 
-	InitOptions.UseDefaultColumns();
-	InitOptions.ColumnMap.Add(FSceneOutlinerBuiltInColumnTypes::Gutter(), FSceneOutlinerColumnInfo(ESceneOutlinerColumnVisibility::Visible, 0, FCreateSceneOutlinerColumn(), true, TOptional<float>(), FSceneOutlinerBuiltInColumnTypes::Gutter_Localized()));
-	InitOptions.ColumnMap.Add(FSceneOutlinerBuiltInColumnTypes::ActorInfo(), FSceneOutlinerColumnInfo(ESceneOutlinerColumnVisibility::Visible, 20, FCreateSceneOutlinerColumn(), true, TOptional<float>(), FSceneOutlinerBuiltInColumnTypes::ActorInfo_Localized()));
-	InitOptions.ColumnMap.Add(FChaosVDSceneOutlinerGutter::GetID(), FSceneOutlinerColumnInfo(ESceneOutlinerColumnVisibility::Visible, 0, FCreateSceneOutlinerColumn::CreateLambda([](ISceneOutliner& InSceneOutliner) { return MakeShareable(new FChaosVDSceneOutlinerGutter(InSceneOutliner)); })));
+	const FTedsOutlinerHierarchyData::FSetParentRowHandle RowHandleSetter = FTedsOutlinerHierarchyData::FSetParentRowHandle::CreateLambda([](void* InColumnData,
+		RowHandle InRowHandle)
+		{
+			if (FChaosVDTableRowParentColumn* ParentColumn = static_cast<FChaosVDTableRowParentColumn*>(InColumnData))
+			{
+				ParentColumn->ParentObject = InRowHandle;
+			}
+		});
 		
-	FCreateSceneOutlinerMode ModeFactory = FCreateSceneOutlinerMode::CreateLambda([this, PlaybackController](SSceneOutliner* Outliner)
+	Params.HierarchyData = FTedsOutlinerHierarchyData(FChaosVDTableRowParentColumn::StaticStruct(), RowHandleGetter, RowHandleSetter, ChildrenRowHandlesGetter);
+
+	Params.SelectionSetOverride = FName("CVDSelection");
+
+	InitOptions.ModeFactory = FCreateSceneOutlinerMode::CreateLambda([this, PlaybackController, &Params, &InitOptions](SSceneOutliner* Outliner)
 	{
-		FActorModeParams ModeParams(Outliner);
-		ModeParams.SpecifiedWorldToDisplay = GetChaosVDWorld();
-		ModeParams.bHideEmptyFolders = true;
-		ModeParams.bCanInteractWithSelectableActorsOnly = false;
+		Params.SceneOutliner = Outliner;
 
 		// The mode is deleted by the Outliner when it is destroyed 
-		return new FChaosVDWorldOutlinerMode(ModeParams, GetChaosVDScene(), PlaybackController);
+		return new FChaosVDWorldOutlinerMode(Params, GetChaosVDScene(), PlaybackController);
 	});
-	InitOptions.ModeFactory = ModeFactory;
+
+	InitOptions.ColumnMap.Add(FSceneOutlinerBuiltInColumnTypes::Label(), FSceneOutlinerColumnInfo(ESceneOutlinerColumnVisibility::Visible, 10));
+	
+	TSharedRef<ISceneOutliner> TedsOutliner = SNew(SSceneOutliner, InitOptions);
+
+	using namespace UE::Editor::DataStorage;
+	IEditorDataStorageProvider* Storage = GetMutableDataStorageFeature<IEditorDataStorageProvider>(StorageFeatureName);
+
+	static QueryHandle ColumnQuery = Storage->RegisterQuery(
+	Select()
+		.ReadOnly<FVisibleInEditorColumn>()
+	.Compile());
+
+	FSceneOutlinerTedsQueryBinder::GetInstance().RegisterTreeItemIDDealiaser(TedsOutliner, FTreeItemIDDealiaser::CreateLambda([](RowHandle RowHandle)
+	{
+		return FSceneOutlinerTreeItemID(RowHandle);
+	}));
+
+	FSceneOutlinerTedsQueryBinder::GetInstance().AssignQuery(ColumnQuery, TedsOutliner, Params.CellWidgetPurposes);
 	
-	const FSceneOutlinerModule& SceneOutlinerModule = FModuleManager::Get().LoadModuleChecked<FSceneOutlinerModule>("SceneOutliner");
+	SceneOutlinerWidget = TedsOutliner;
 
-	SceneOutlinerWidget = SceneOutlinerModule.CreateSceneOutliner(InitOptions);
+	// TODO: There is an issue where the actors created with CVD's world, do not appear until the outliner does a new query (by searching something for example)
+	// It seems to be a timing issue with TEDS. For now, lest just hack it by rebuilding the outliner in next tick. There should only be ~5 actors.
+	DeferredOutlinerUpdateHandle = FTSTicker::GetCoreTicker().AddTicker(FTickerDelegate::CreateLambda([WeakOutliner = SceneOutlinerWidget.ToWeakPtr()](float DeltaTime)
+	{
+		if (TSharedPtr<ISceneOutliner> OutlinerPtr = WeakOutliner.Pin())
+		{
+			OutlinerPtr->FullRefresh();
+		}
+
+		return false;
+	}));
 }
 
 TSharedRef<SDockTab> FChaosVDWorldOutlinerTab::HandleTabSpawnRequest(const FSpawnTabArgs& Args)
@@ -76,6 +160,9 @@ TSharedRef<SDockTab> FChaosVDWorldOutlinerTab::HandleTabSpawnRequest(const FSpaw
 
 void FChaosVDWorldOutlinerTab::HandleTabClosed(TSharedRef<SDockTab> InTabClosed)
 {
+	FTSTicker::GetCoreTicker().RemoveTicker(DeferredOutlinerUpdateHandle);
+	DeferredOutlinerUpdateHandle = FTSTicker::FDelegateHandle();
+
 	FChaosVDTabSpawnerBase::HandleTabClosed(InTabClosed);
 
 	SceneOutlinerWidget.Reset();
diff --git a/Engine/Plugins/ChaosVD/Source/ChaosVD/Private/ChaosVDWorldOutlinerTab.h b/Engine/Plugins/ChaosVD/Source/ChaosVD/Private/ChaosVDWorldOutlinerTab.h
--- a/Engine/Plugins/ChaosVD/Source/ChaosVD/Private/ChaosVDWorldOutlinerTab.h
+++ b/Engine/Plugins/ChaosVD/Source/ChaosVD/Private/ChaosVDWorldOutlinerTab.h
@@ -37,4 +37,6 @@ private:
 	void CreateWorldOutlinerWidget();
 
 	TSharedPtr<ISceneOutliner> SceneOutlinerWidget;
+
+	FTSTicker::FDelegateHandle DeferredOutlinerUpdateHandle;
 };
diff --git a/Engine/Plugins/ChaosVD/Source/ChaosVD/Private/Components/ChaosVDInstancedStaticMeshComponent.cpp b/Engine/Plugins/ChaosVD/Source/ChaosVD/Private/Components/ChaosVDInstancedStaticMeshComponent.cpp
--- a/Engine/Plugins/ChaosVD/Source/ChaosVD/Private/Components/ChaosVDInstancedStaticMeshComponent.cpp
+++ b/Engine/Plugins/ChaosVD/Source/ChaosVD/Private/Components/ChaosVDInstancedStaticMeshComponent.cpp
@@ -31,10 +31,12 @@ void UChaosVDInstancedStaticMeshComponent::UpdateInstanceVisibility(const TShare
 	}
 
 	constexpr bool bIsWorldSpaceTransform = true;
-	constexpr bool bMarkRenderDirty = true;
+	constexpr bool bMarkRenderDirty = false;
 	constexpr bool bTeleport = true;
 
 	UpdateInstanceTransform(InstanceIndex, Transform, bIsWorldSpaceTransform, bMarkRenderDirty, bTeleport);
+
+	bHadTransformsUpdate = true;
 }
 
 void UChaosVDInstancedStaticMeshComponent::SetIsSelected(const TSharedPtr<FChaosVDMeshDataInstanceHandle>& InInstanceHandle, bool bIsSelected)
@@ -117,10 +119,12 @@ void UChaosVDInstancedStaticMeshComponent::UpdateInstanceWorldTransform(const TS
 	}
 
 	constexpr bool bIsWorldSpaceTransform = true;
-	constexpr bool bMarkRenderDirty = true;
+	constexpr bool bMarkRenderDirty = false;
 	constexpr bool bTeleport = true;
 
 	UpdateInstanceTransform(InstanceIndex, NewTransform, bIsWorldSpaceTransform, bMarkRenderDirty, bTeleport);
+
+	bHadTransformsUpdate = true;
 }
 
 TArrayView<TSharedPtr<FChaosVDMeshDataInstanceHandle>> UChaosVDInstancedStaticMeshComponent::GetMeshDataInstanceHandles()
@@ -186,7 +190,7 @@ void UChaosVDInstancedStaticMeshComponent::OnDisposed()
 		UnregisterComponent();
 	}
 
-	if (AActor* Owner =GetOwner())
+	if (AActor* Owner = GetOwner())
 	{
 		Owner->RemoveOwnedComponent(this);
 	}
@@ -209,14 +213,13 @@ bool UChaosVDInstancedStaticMeshComponent::UpdateGeometryKey(const uint32 NewHan
 	return true;
 }
 
-TSharedPtr<FChaosVDMeshDataInstanceHandle> UChaosVDInstancedStaticMeshComponent::AddMeshInstance(const FTransform InstanceTransform, bool bIsWorldSpace, const TSharedPtr<FChaosVDExtractedGeometryDataHandle>& InGeometryHandle, int32 ParticleID, int32 SolverID)
+TSharedPtr<FChaosVDMeshDataInstanceHandle> UChaosVDInstancedStaticMeshComponent::AddMeshInstance(const FTransform InstanceTransform, bool bIsWorldSpace, const TSharedRef<FChaosVDExtractedGeometryDataHandle>& InGeometryHandle, int32 ParticleID, int32 SolverID)
 {
 	int32 InstanceIndex = AddInstance(InstanceTransform, bIsWorldSpace);
 
 	check(InstanceIndex != INDEX_NONE);
 
-	const TSharedPtr<FChaosVDMeshDataInstanceHandle> InstanceHandle = MakeShared<FChaosVDMeshDataInstanceHandle>(InstanceIndex, this, ParticleID, SolverID);
-	InstanceHandle->SetGeometryHandle(InGeometryHandle);
+	const TSharedPtr<FChaosVDMeshDataInstanceHandle> InstanceHandle = MakeShared<FChaosVDMeshDataInstanceHandle>(InstanceIndex, this, ParticleID, SolverID, InGeometryHandle);
 
 	const uint32 NewHandleGeometryKey = InGeometryHandle->GetGeometryKey();
 
@@ -230,7 +233,7 @@ TSharedPtr<FChaosVDMeshDataInstanceHandle> UChaosVDInstancedStaticMeshComponent:
 	return InstanceHandle;
 }
 
-void UChaosVDInstancedStaticMeshComponent::AddMeshInstanceForHandle(TSharedPtr<FChaosVDMeshDataInstanceHandle> MeshDataHandle, const FTransform InstanceTransform, bool bIsWorldSpace, const TSharedPtr<FChaosVDExtractedGeometryDataHandle>& InGeometryHandle, int32 ParticleID, int32 SolverID)
+void UChaosVDInstancedStaticMeshComponent::AddMeshInstanceForHandle(TSharedPtr<FChaosVDMeshDataInstanceHandle> MeshDataHandle, const FTransform InstanceTransform, bool bIsWorldSpace, const TSharedRef<FChaosVDExtractedGeometryDataHandle>& InGeometryHandle, int32 ParticleID, int32 SolverID)
 {
 	const int32 InstanceIndex = AddInstance(InstanceTransform, bIsWorldSpace);
 
@@ -313,6 +316,17 @@ bool UChaosVDInstancedStaticMeshComponent::IsNavigationRelevant() const
 	return false;
 }
 
+void UChaosVDInstancedStaticMeshComponent::TickComponent(float DeltaTime, enum ELevelTick TickType, FActorComponentTickFunction* ThisTickFunction)
+{
+	Super::TickComponent(DeltaTime, TickType, ThisTickFunction);
+
+	if (bHadTransformsUpdate)
+	{
+		MarkRenderStateDirty();
+		bHadTransformsUpdate = false;
+	}
+}
+
 uint32 UChaosVDInstancedStaticMeshComponent::GetGeometryKey() const
 {
 	return CurrentGeometryKey;
diff --git a/Engine/Plugins/ChaosVD/Source/ChaosVD/Private/Components/ChaosVDInstancedStaticMeshComponent.h b/Engine/Plugins/ChaosVD/Source/ChaosVD/Private/Components/ChaosVDInstancedStaticMeshComponent.h
--- a/Engine/Plugins/ChaosVD/Source/ChaosVD/Private/Components/ChaosVDInstancedStaticMeshComponent.h
+++ b/Engine/Plugins/ChaosVD/Source/ChaosVD/Private/Components/ChaosVDInstancedStaticMeshComponent.h
@@ -23,6 +23,8 @@ public:
 		bOverrideWireframeColor = true;
 		WireframeColorOverride = FColor::White;
 	}
+
+	virtual void TickComponent(float DeltaTime, enum ELevelTick TickType, FActorComponentTickFunction* ThisTickFunction) override;
 	
 	// BEGIN IChaosVDGeometryDataComponent Interface
 	
@@ -60,8 +62,8 @@ public:
 	virtual EChaosVDMaterialType GetMaterialType() const;
 	// END IChaosVDGeometryDataComponent Interface
 
-	virtual TSharedPtr<FChaosVDMeshDataInstanceHandle> AddMeshInstance(const FTransform InstanceTransform, bool bIsWorldSpace, const TSharedPtr<FChaosVDExtractedGeometryDataHandle>& InGeometryHandle, int32 ParticleID, int32 SolverID) override;
-	virtual void AddMeshInstanceForHandle(TSharedPtr<FChaosVDMeshDataInstanceHandle> MeshDataHandle, const FTransform InstanceTransform, bool bIsWorldSpace, const TSharedPtr<FChaosVDExtractedGeometryDataHandle>& InGeometryHandle, int32 ParticleID, int32 SolverID) override;
+	virtual TSharedPtr<FChaosVDMeshDataInstanceHandle> AddMeshInstance(const FTransform InstanceTransform, bool bIsWorldSpace, const TSharedRef<FChaosVDExtractedGeometryDataHandle>& InGeometryHandle, int32 ParticleID, int32 SolverID) override;
+	virtual void AddMeshInstanceForHandle(TSharedPtr<FChaosVDMeshDataInstanceHandle> MeshDataHandle, const FTransform InstanceTransform, bool bIsWorldSpace, const TSharedRef<FChaosVDExtractedGeometryDataHandle>& InGeometryHandle, int32 ParticleID, int32 SolverID) override;
 	virtual void RemoveMeshInstance(TSharedPtr<FChaosVDMeshDataInstanceHandle> InHandleToRemove) override;
 
 	virtual bool Modify(bool bAlwaysMarkDirty) override;
@@ -95,4 +97,6 @@ protected:
 	TArray<TSharedPtr<FChaosVDMeshDataInstanceHandle>> CurrentInstanceHandles;
 	
 	TWeakPtr<FChaosVDGeometryBuilder> GeometryBuilderWeakPtr;
+
+	bool bHadTransformsUpdate = false;
 };
diff --git a/Engine/Plugins/ChaosVD/Source/ChaosVD/Private/Components/ChaosVDParticleDataComponent.h b/Engine/Plugins/ChaosVD/Source/ChaosVD/Private/Components/ChaosVDParticleDataComponent.h
--- a/Engine/Plugins/ChaosVD/Source/ChaosVD/Private/Components/ChaosVDParticleDataComponent.h
+++ b/Engine/Plugins/ChaosVD/Source/ChaosVD/Private/Components/ChaosVDParticleDataComponent.h
@@ -2,9 +2,15 @@
 
 #pragma once
 
-#include "Components/ActorComponent.h"
+#include "ChaosVDScene.h"
+#include "ChaosVDSceneParticle.h"
+#include "Components/ChaosVDSolverDataComponent.h"
+#include "TEDS/ChaosVDSelectionInterface.h"
+#include "TEDS/ChaosVDStructTypedElementData.h"
 #include "ChaosVDParticleDataComponent.generated.h"
 
+struct FChaosVDParticleDataWrapper;
+struct FChaosVDSceneParticle;
 /**
  * Component that will reference all particle data for a specific solver for the current frame.
  * Currently it is empty and only used to interface with the component visualizer system for particle data,
@@ -12,7 +18,128 @@
  * (similarly to how the Scene Query Data and Collision Data is referenced in their respective components)
  */
 UCLASS()
-class UChaosVDParticleDataComponent : public UActorComponent
+class UChaosVDParticleDataComponent : public UChaosVDSolverDataComponent
 {
+public:
+
+	UChaosVDParticleDataComponent();
+
+private:
 	GENERATED_BODY()
+
+public:
+
+	virtual void ClearData() override;
+
+	virtual void UpdateFromNewSolverStageData(const FChaosVDSolverFrameData& InSolverFrameData, const FChaosVDFrameStageData& InSolverFrameStageData) override;
+	virtual void UpdateFromSolverFrameData(const FChaosVDSolverFrameData& InSolverFrameData) override;
+	
+	template <class TCallback>
+	void VisitSelectedParticleData(TCallback VisitCallback) const;
+
+	template <class TCallback>
+	void VisitAllParticleData(TCallback VisitCallback) const;
+
+	virtual void SetVisibility(bool bNewIsVisible) override;
+
+	TSharedPtr<FChaosVDSceneParticle> GetParticleInstanceByID(int32 ParticleID) const;
+	TSharedPtr<FChaosVDSceneParticle> GetParticleInstanceByID_AssumesLocked(int32 ParticleID) const;
+
+	TSharedPtr<FChaosVDBaseSceneObject> GetParticleContainerByType(EChaosVDParticleType ParticleType);
+	TSharedPtr<FChaosVDBaseSceneObject> GetParticleContainerByType_AssumesLocked(EChaosVDParticleType ParticleType);
+
+	const TSortedMap<EChaosVDParticleType, TSharedPtr<FChaosVDBaseSceneObject>>& GetParticleSceneContainersByType()
+	{
+		return ParticleSceneContainersByType;
+	}
+
+	FChaosVDSceneParticle* GetSelectedParticle() const;
+
+protected:
+
+	int32 GetMaxElementCountForBatchTEDSUpdate_AssumesLocked() const;
+
+	void BatchTEDSSyncFromWorldTagUpdate(TArray<UE::Editor::DataStorage::RowHandle, TInlineAllocator<256, TMemStackAllocator<>>>& ParticlesPendingTEDSTagSync);
+	void BatchTEDSSyncFromWorldTagUpdate(TArray<UE::Editor::DataStorage::RowHandle, TInlineAllocator<256, TMemStackAllocator<>>>& ParticlesPendingTEDSTagSync, 	const IEditorDataStorageCompatibilityProvider* CompatibilityStorage, IEditorDataStorageProvider* DataStorage);
+
+	void ProcessRemovedParticles_AssumesLocked(const TSharedRef<FChaosVDScene>& SceneRef, const FChaosVDSolverFrameData& InSolverFrameData, const FChaosVDFrameStageData& InSolverFrameStageData, const TSet<int32>& ParticlesUpdatedIDsForKeyFrameDiff);
+
+	void BatchApplySolverVisibilityToParticle(bool bNewIsVisible);
+
+	void HandleVisibilitySettingsUpdated(UObject* SettingsObject);
+	void HandleColorsSettingsUpdated(UObject* SettingsObject);
+
+	bool IsServerData();
+	const FString& GetSolverName();
+
+	/** Creates an ChaosVDParticle actor for the Provided recorded Particle Data */
+	TSharedRef<FChaosVDSceneParticle> CreateSceneParticle_AssumesLocked(const TSharedRef<FChaosVDParticleDataWrapper>& InParticleData, const FChaosVDSolverFrameData& InFrameData);
+
+	TOptional<bool> bCachedIsServerData;
+	TOptional<FString> bCachedSolverName;
+
+	TSortedMap<EChaosVDParticleType, TSharedPtr<FChaosVDBaseSceneObject>> ParticleSceneContainersByType;
+
+	TSortedMap<EChaosVDParticleType, FName> FolderPathByParticlePath;
+
+	TMap<int32, TSharedRef<FChaosVDSceneParticle>> SolverParticlesByID;
+
+	mutable FRWLock ParticleSceneObjectsLock;
 };
+
+template <typename TCallback>
+void UChaosVDParticleDataComponent::VisitSelectedParticleData(TCallback VisitCallback) const
+{
+	TSharedPtr<FChaosVDScene> ScenePtr = SceneWeakPtr.Pin();
+	if (!ScenePtr)
+	{
+		return;
+	}
+	using namespace Chaos::VD::TypedElementDataUtil;
+
+	TArray<FTypedElementHandle, TInlineAllocator<8>> SelectedParticlesHandles;
+	ScenePtr->GetElementSelectionSet()->GetSelectedElementHandles(SelectedParticlesHandles, UChaosVDSelectionInterface::StaticClass());
+
+	for (const FTypedElementHandle& SelectedParticleHandle : SelectedParticlesHandles)
+	{
+		if (FChaosVDSceneParticle* ParticleInstance = GetStructDataFromTypedElementHandle<FChaosVDSceneParticle>(SelectedParticleHandle))
+		{
+			TSharedPtr<const FChaosVDParticleDataWrapper> ParticleDataViewer = ParticleInstance ? ParticleInstance->GetParticleData() : nullptr;
+	
+			if (!ensure(ParticleDataViewer))
+			{
+				continue;
+			}
+
+			if (ParticleDataViewer->SolverID != SolverID)
+			{
+				continue;
+			}
+
+			if (!VisitCallback(ParticleDataViewer))
+			{
+				return;
+			}
+		}
+	}
+}
+
+template <typename TCallback>
+void UChaosVDParticleDataComponent::VisitAllParticleData(TCallback VisitCallback) const
+{
+	FReadScopeLock ReadLock(ParticleSceneObjectsLock);
+
+	for (const TPair<int32, TSharedRef<FChaosVDSceneParticle>>& ParticleWithIDPair : SolverParticlesByID)
+	{
+		TSharedPtr<const FChaosVDParticleDataWrapper> ParticleDataViewer = ParticleWithIDPair.Value->GetParticleData();
+		if (!ensure(ParticleDataViewer))
+		{
+			continue;
+		}
+
+		if (!VisitCallback(ParticleDataViewer))
+		{
+			return;
+		}
+	}
+}
diff --git a/Engine/Plugins/ChaosVD/Source/ChaosVD/Private/Components/ChaosVDSolverCharacterGroundConstraintDataComponent.cpp b/Engine/Plugins/ChaosVD/Source/ChaosVD/Private/Components/ChaosVDSolverCharacterGroundConstraintDataComponent.cpp
--- a/Engine/Plugins/ChaosVD/Source/ChaosVD/Private/Components/ChaosVDSolverCharacterGroundConstraintDataComponent.cpp
+++ b/Engine/Plugins/ChaosVD/Source/ChaosVD/Private/Components/ChaosVDSolverCharacterGroundConstraintDataComponent.cpp
@@ -3,9 +3,11 @@
 #include "Components/ChaosVDSolverCharacterGroundConstraintDataComponent.h"
 
 #include "ChaosVDCharacterGroundConstraintDataProviderInterface.h"
+#include "ChaosVDParticleDataComponent.h"
 #include "ChaosVDScene.h"
 #include "ChaosVDSettingsManager.h"
 #include "Selection.h"
+#include "Actors/ChaosVDSolverInfoActor.h"
 #include "Settings/ChaosVDCharacterConstraintsVisualizationSettings.h"
 
 void UChaosVDSolverCharacterGroundConstraintDataComponent::HandleSceneUpdated()
@@ -16,7 +18,6 @@ void UChaosVDSolverCharacterGroundConstraintDataComponent::HandleSceneUpdated()
 		{
 			return;
 		}
-
 	}
 
 	const TSharedPtr<FChaosVDScene> ScenePtr = SceneWeakPtr.Pin();
@@ -24,9 +25,15 @@ void UChaosVDSolverCharacterGroundConstraintDataComponent::HandleSceneUpdated()
 	{
 		return;
 	}
-	
-	USelection* ActorSelection = ScenePtr->GetActorSelectionObject();
-	if (!ActorSelection)
+
+	AChaosVDSolverInfoActor* OwnerSolverData = Cast<AChaosVDSolverInfoActor>(GetOwner());
+	if (!OwnerSolverData)
+	{
+		return;
+	}
+
+	UChaosVDParticleDataComponent* ParticleDataComponent = OwnerSolverData->GetParticleDataComponent();
+	if (!ParticleDataComponent)
 	{
 		return;
 	}
@@ -37,16 +44,13 @@ void UChaosVDSolverCharacterGroundConstraintDataComponent::HandleSceneUpdated()
 		return;
 	}
 
-	if (ActorSelection->Num() > 0)
+	if (FChaosVDSceneParticle* ParticleInstance = ParticleDataComponent->GetSelectedParticle())
 	{
-		if (IChaosVDCharacterGroundConstraintDataProviderInterface* DataProvider = Cast<IChaosVDCharacterGroundConstraintDataProviderInterface>(ActorSelection->GetSelectedObject(0)))
+		if (ParticleInstance->HasCharacterGroundConstraintData())
 		{
-			if (DataProvider->HasCharacterGroundConstraintData())
-			{
-				TArray<TSharedPtr<FChaosVDCharacterGroundConstraint>> FoundConstraintData;
-				DataProvider->GetCharacterGroundConstraintData(FoundConstraintData);
-				SolverDataSelection->SelectData(SolverDataSelection->MakeSelectionHandle(FoundConstraintData[0]));
-			}
+			TArray<TSharedPtr<FChaosVDCharacterGroundConstraint>> FoundConstraintData;
+			ParticleInstance->GetCharacterGroundConstraintData(FoundConstraintData);
+			SolverDataSelection->SelectData(SolverDataSelection->MakeSelectionHandle(FoundConstraintData[0]));
 		}
 	}
 }
@@ -76,3 +80,9 @@ void UChaosVDSolverCharacterGroundConstraintDataComponent::SetScene(const TWeakP
 	
 	ScenePtr->OnSceneUpdated().AddUObject(this, &UChaosVDSolverCharacterGroundConstraintDataComponent::HandleSceneUpdated);
 }
+
+void UChaosVDSolverCharacterGroundConstraintDataComponent::UpdateFromSolverFrameData(const FChaosVDSolverFrameData& InSolverFrameData)
+{
+	UpdateConstraintData(InSolverFrameData.RecordedCharacterGroundConstraints);
+}
+
diff --git a/Engine/Plugins/ChaosVD/Source/ChaosVD/Private/Components/ChaosVDSolverCharacterGroundConstraintDataComponent.h b/Engine/Plugins/ChaosVD/Source/ChaosVD/Private/Components/ChaosVDSolverCharacterGroundConstraintDataComponent.h
--- a/Engine/Plugins/ChaosVD/Source/ChaosVD/Private/Components/ChaosVDSolverCharacterGroundConstraintDataComponent.h
+++ b/Engine/Plugins/ChaosVD/Source/ChaosVD/Private/Components/ChaosVDSolverCharacterGroundConstraintDataComponent.h
@@ -12,6 +12,8 @@ class CHAOSVD_API UChaosVDSolverCharacterGroundConstraintDataComponent : public
 public:
 	virtual void SetScene(const TWeakPtr<FChaosVDScene>& InSceneWeakPtr) override;
 
+	virtual void UpdateFromSolverFrameData(const FChaosVDSolverFrameData& InSolverFrameData) override;
+
 	void HandleSceneUpdated();
 
 	virtual void BeginDestroy() override;
diff --git a/Engine/Plugins/ChaosVD/Source/ChaosVD/Private/Components/ChaosVDSolverCollisionDataComponent.cpp b/Engine/Plugins/ChaosVD/Source/ChaosVD/Private/Components/ChaosVDSolverCollisionDataComponent.cpp
--- a/Engine/Plugins/ChaosVD/Source/ChaosVD/Private/Components/ChaosVDSolverCollisionDataComponent.cpp
+++ b/Engine/Plugins/ChaosVD/Source/ChaosVD/Private/Components/ChaosVDSolverCollisionDataComponent.cpp
@@ -2,6 +2,7 @@
 
 #include "Components/ChaosVDSolverCollisionDataComponent.h"
 
+#include "ChaosVDRecording.h"
 #include "DataWrappers/ChaosVDCollisionDataWrappers.h"
 
 UChaosVDSolverCollisionDataComponent::UChaosVDSolverCollisionDataComponent()
@@ -35,6 +36,14 @@ void UChaosVDSolverCollisionDataComponent::UpdateCollisionData(const TArray<TSha
 	}
 }
 
+void UChaosVDSolverCollisionDataComponent::UpdateFromNewSolverStageData(const FChaosVDSolverFrameData& InSolverFrameData, const FChaosVDFrameStageData& InSolverFrameStageData)
+{
+	if (EnumHasAnyFlags(InSolverFrameStageData.StageFlags, EChaosVDSolverStageFlags::ExplicitStage))
+	{
+		UpdateCollisionData(InSolverFrameStageData.RecordedMidPhases);
+	}
+}
+
 const TArray<TSharedPtr<FChaosVDParticlePairMidPhase>>* UChaosVDSolverCollisionDataComponent::GetMidPhasesForParticle(int32 ParticleID, EChaosVDParticlePairSlot Options) const
 {
 	return Chaos::VisualDebugger::Utils::GetDataFromParticlePairMaps<FChaosVDMidPhaseByParticleMap, TSharedPtr<FChaosVDParticlePairMidPhase>>(MidPhasesByParticleID0, MidPhasesByParticleID1, ParticleID, Options);
diff --git a/Engine/Plugins/ChaosVD/Source/ChaosVD/Private/Components/ChaosVDSolverCollisionDataComponent.h b/Engine/Plugins/ChaosVD/Source/ChaosVD/Private/Components/ChaosVDSolverCollisionDataComponent.h
--- a/Engine/Plugins/ChaosVD/Source/ChaosVD/Private/Components/ChaosVDSolverCollisionDataComponent.h
+++ b/Engine/Plugins/ChaosVD/Source/ChaosVD/Private/Components/ChaosVDSolverCollisionDataComponent.h
@@ -26,15 +26,21 @@ class UChaosVDSolverCollisionDataComponent : public UChaosVDSolverDataComponent
 public:
 	UChaosVDSolverCollisionDataComponent();
 
-	void UpdateCollisionData(const TArray<TSharedPtr<FChaosVDParticlePairMidPhase>>& InMidPhaseData);
+	virtual void UpdateFromNewSolverStageData(const FChaosVDSolverFrameData& InSolverFrameData, const FChaosVDFrameStageData& InSolverFrameStageData) override;
+
+	const TArray<TSharedPtr<FChaosVDParticlePairMidPhase>>& GetMidPhases() const
+	{
+		return AllMidPhases;
+	}
 
-	const TArray<TSharedPtr<FChaosVDParticlePairMidPhase>>& GetMidPhases() const { return AllMidPhases; }
 	const TArray<TSharedPtr<FChaosVDParticlePairMidPhase>>* GetMidPhasesForParticle(int32 ParticleID, EChaosVDParticlePairSlot Options) const;
 	const TArray<FChaosVDConstraint*>* GetConstraintsForParticle(int32 ParticleID, EChaosVDParticlePairSlot Options) const;
 
 	virtual void ClearData() override;
 
 protected:
+	
+	void UpdateCollisionData(const TArray<TSharedPtr<FChaosVDParticlePairMidPhase>>& InMidPhaseData);
 
 	TArray<TSharedPtr<FChaosVDParticlePairMidPhase>> AllMidPhases;
 	FChaosVDMidPhaseByParticleMap MidPhasesByParticleID0;
diff --git a/Engine/Plugins/ChaosVD/Source/ChaosVD/Private/Components/ChaosVDSolverDataComponent.cpp b/Engine/Plugins/ChaosVD/Source/ChaosVD/Private/Components/ChaosVDSolverDataComponent.cpp
--- a/Engine/Plugins/ChaosVD/Source/ChaosVD/Private/Components/ChaosVDSolverDataComponent.cpp
+++ b/Engine/Plugins/ChaosVD/Source/ChaosVD/Private/Components/ChaosVDSolverDataComponent.cpp
@@ -2,8 +2,12 @@
 
 #include "Components/ChaosVDSolverDataComponent.h"
 
-
 void UChaosVDSolverDataComponent::SetScene(const TWeakPtr<FChaosVDScene>& InSceneWeakPtr)
 {
 	SceneWeakPtr = InSceneWeakPtr;
 }
+
+void UChaosVDSolverDataComponent::SetVisibility(bool bNewIsVisible)
+{
+	bIsVisible = bNewIsVisible;
+}
diff --git a/Engine/Plugins/ChaosVD/Source/ChaosVD/Private/Components/ChaosVDSolverJointConstraintDataComponent.cpp b/Engine/Plugins/ChaosVD/Source/ChaosVD/Private/Components/ChaosVDSolverJointConstraintDataComponent.cpp
--- a/Engine/Plugins/ChaosVD/Source/ChaosVD/Private/Components/ChaosVDSolverJointConstraintDataComponent.cpp
+++ b/Engine/Plugins/ChaosVD/Source/ChaosVD/Private/Components/ChaosVDSolverJointConstraintDataComponent.cpp
@@ -1,3 +1,13 @@
 ﻿// Copyright Epic Games, Inc. All Rights Reserved.
 
 #include "Components/ChaosVDSolverJointConstraintDataComponent.h"
+
+#include "ChaosVDRecording.h"
+
+void UChaosVDSolverJointConstraintDataComponent::UpdateFromNewSolverStageData(const FChaosVDSolverFrameData& InSolverFrameData, const FChaosVDFrameStageData& InSolverFrameStageData)
+{
+	if (EnumHasAnyFlags(InSolverFrameStageData.StageFlags, EChaosVDSolverStageFlags::ExplicitStage))
+	{
+		UpdateConstraintData(InSolverFrameStageData.RecordedJointConstraints);
+	}
+}
diff --git a/Engine/Plugins/ChaosVD/Source/ChaosVD/Private/Components/ChaosVDSolverJointConstraintDataComponent.h b/Engine/Plugins/ChaosVD/Source/ChaosVD/Private/Components/ChaosVDSolverJointConstraintDataComponent.h
--- a/Engine/Plugins/ChaosVD/Source/ChaosVD/Private/Components/ChaosVDSolverJointConstraintDataComponent.h
+++ b/Engine/Plugins/ChaosVD/Source/ChaosVD/Private/Components/ChaosVDSolverJointConstraintDataComponent.h
@@ -10,4 +10,6 @@ UCLASS()
 class CHAOSVD_API UChaosVDSolverJointConstraintDataComponent : public UChaosVDConstraintDataComponent
 {
 	GENERATED_BODY()
+
+	virtual void UpdateFromNewSolverStageData(const FChaosVDSolverFrameData& InSolverFrameData, const FChaosVDFrameStageData& InSolverFrameStageData) override;
 };
diff --git a/Engine/Plugins/ChaosVD/Source/ChaosVD/Private/Components/ChaosVDStaticMeshComponent.cpp b/Engine/Plugins/ChaosVD/Source/ChaosVD/Private/Components/ChaosVDStaticMeshComponent.cpp
--- a/Engine/Plugins/ChaosVD/Source/ChaosVD/Private/Components/ChaosVDStaticMeshComponent.cpp
+++ b/Engine/Plugins/ChaosVD/Source/ChaosVD/Private/Components/ChaosVDStaticMeshComponent.cpp
@@ -108,13 +108,12 @@ void UChaosVDStaticMeshComponent::Reset()
 	CurrentGeometryKey = 0;
 }
 
-TSharedPtr<FChaosVDMeshDataInstanceHandle> UChaosVDStaticMeshComponent::AddMeshInstance(const FTransform InstanceTransform, bool bIsWorldSpace, const TSharedPtr<FChaosVDExtractedGeometryDataHandle>& InGeometryHandle, int32 ParticleID, int32 SolverID)
+TSharedPtr<FChaosVDMeshDataInstanceHandle> UChaosVDStaticMeshComponent::AddMeshInstance(const FTransform InstanceTransform, bool bIsWorldSpace, const TSharedRef<FChaosVDExtractedGeometryDataHandle>& InGeometryHandle, int32 ParticleID, int32 SolverID)
 {
 	// Static meshes only have one instance
 	constexpr int32 StaticMeshInstanceIndex = 0;
 
-	CurrentMeshDataHandle = MakeShared<FChaosVDMeshDataInstanceHandle>(StaticMeshInstanceIndex, this, ParticleID, SolverID);
-	CurrentMeshDataHandle->SetGeometryHandle(InGeometryHandle);
+	CurrentMeshDataHandle = MakeShared<FChaosVDMeshDataInstanceHandle>(StaticMeshInstanceIndex, this, ParticleID, SolverID, InGeometryHandle);
 
 	const uint32 NewHandleGeometryKey = InGeometryHandle->GetGeometryKey();
 
@@ -128,7 +127,7 @@ TSharedPtr<FChaosVDMeshDataInstanceHandle> UChaosVDStaticMeshComponent::AddMeshI
 	return CurrentMeshDataHandle;
 }
 
-void UChaosVDStaticMeshComponent::AddMeshInstanceForHandle(TSharedPtr<FChaosVDMeshDataInstanceHandle> MeshDataHandle, const FTransform InstanceTransform, bool bIsWorldSpace, const TSharedPtr<FChaosVDExtractedGeometryDataHandle>& InGeometryHandle, int32 ParticleID, int32 SolverID)
+void UChaosVDStaticMeshComponent::AddMeshInstanceForHandle(TSharedPtr<FChaosVDMeshDataInstanceHandle> MeshDataHandle, const FTransform InstanceTransform, bool bIsWorldSpace, const TSharedRef<FChaosVDExtractedGeometryDataHandle>& InGeometryHandle, int32 ParticleID, int32 SolverID)
 {
 	// Static meshes only have one instance
 	constexpr int32 StaticMeshInstanceIndex = 0;
diff --git a/Engine/Plugins/ChaosVD/Source/ChaosVD/Private/Components/ChaosVDStaticMeshComponent.h b/Engine/Plugins/ChaosVD/Source/ChaosVD/Private/Components/ChaosVDStaticMeshComponent.h
--- a/Engine/Plugins/ChaosVD/Source/ChaosVD/Private/Components/ChaosVDStaticMeshComponent.h
+++ b/Engine/Plugins/ChaosVD/Source/ChaosVD/Private/Components/ChaosVDStaticMeshComponent.h
@@ -52,8 +52,8 @@ public:
 	virtual void Initialize() override;
 	virtual void Reset() override;
 
-	virtual TSharedPtr<FChaosVDMeshDataInstanceHandle> AddMeshInstance(const FTransform InstanceTransform, bool bIsWorldSpace, const TSharedPtr<FChaosVDExtractedGeometryDataHandle>& InGeometryHandle, int32 ParticleID, int32 SolverID) override;
-	virtual void AddMeshInstanceForHandle(TSharedPtr<FChaosVDMeshDataInstanceHandle> MeshDataHandle, const FTransform InstanceTransform, bool bIsWorldSpace, const TSharedPtr<FChaosVDExtractedGeometryDataHandle>& InGeometryHandle, int32 ParticleID, int32 SolverID) override;
+	virtual TSharedPtr<FChaosVDMeshDataInstanceHandle> AddMeshInstance(const FTransform InstanceTransform, bool bIsWorldSpace, const TSharedRef<FChaosVDExtractedGeometryDataHandle>& InGeometryHandle, int32 ParticleID, int32 SolverID) override;
+	virtual void AddMeshInstanceForHandle(TSharedPtr<FChaosVDMeshDataInstanceHandle> MeshDataHandle, const FTransform InstanceTransform, bool bIsWorldSpace, const TSharedRef<FChaosVDExtractedGeometryDataHandle>& InGeometryHandle, int32 ParticleID, int32 SolverID) override;
 	virtual void RemoveMeshInstance(TSharedPtr<FChaosVDMeshDataInstanceHandle> InHandleToRemove) override;
 
 	virtual void SetGeometryBuilder(TWeakPtr<FChaosVDGeometryBuilder> GeometryBuilder) override;
diff --git a/Engine/Plugins/ChaosVD/Source/ChaosVD/Private/Settings/ChaosVDParticleVisualizationSettings.h b/Engine/Plugins/ChaosVD/Source/ChaosVD/Private/Settings/ChaosVDParticleVisualizationSettings.h
--- a/Engine/Plugins/ChaosVD/Source/ChaosVD/Private/Settings/ChaosVDParticleVisualizationSettings.h
+++ b/Engine/Plugins/ChaosVD/Source/ChaosVD/Private/Settings/ChaosVDParticleVisualizationSettings.h
@@ -292,8 +292,6 @@ public:
 	static void SetGeometryVisualizationFlags(EChaosVDGeometryVisibilityFlags Flags);
 	static EChaosVDGeometryVisibilityFlags GetGeometryVisualizationFlags();
 
-private:
-
 	/** Set of flags to enable/disable visibility of specific types of geometry/particles */
 	UPROPERTY(config, meta = (Bitmask, BitmaskEnum = "/Script/ChaosVD.EChaosVDGeometryVisibilityFlags"))
 	uint32 GeometryVisibilityFlags = static_cast<uint32>(EChaosVDGeometryVisibilityFlags::Simulated | EChaosVDGeometryVisibilityFlags::Simple | EChaosVDGeometryVisibilityFlags::ShowHeightfields | EChaosVDGeometryVisibilityFlags::Complex | EChaosVDGeometryVisibilityFlags::Query);
diff --git a/Engine/Plugins/ChaosVD/Source/ChaosVD/Private/TEDS/ChaosVDDataStorageVisibilityQueries.h b/Engine/Plugins/ChaosVD/Source/ChaosVD/Private/TEDS/ChaosVDDataStorageVisibilityQueries.h
new file mode 100644
--- /dev/null
+++ b/Engine/Plugins/ChaosVD/Source/ChaosVD/Private/TEDS/ChaosVDDataStorageVisibilityQueries.h
@@ -0,0 +1,34 @@
+// Copyright Epic Games, Inc. All Rights Reserved.
+
+#pragma once
+
+#include "Elements/Interfaces/TypedElementDataStorageFactory.h"
+#include "ChaosVDDataStorageVisibilityQueries.generated.h"
+
+UCLASS()
+class UChaosVDDataStorageVisibilityQueries : public UEditorDataStorageFactory
+{
+	GENERATED_BODY()
+
+public:
+	~UChaosVDDataStorageVisibilityQueries() override = default;
+
+	void RegisterQueries(IEditorDataStorageProvider& DataStorage) override;
+
+private:
+
+	/**
+	 * Adds the Visibility columns to new actors that do not have one already.
+	 */
+	void RegisterParticleAddVisibilityColumn(IEditorDataStorageProvider& DataStorage) const;
+
+	/**
+	 * Takes the visibility set on an actor and copies it to the Data Storage if they differ.
+	 */
+	void RegisterParticleVisibilityToColumnQuery(IEditorDataStorageProvider& DataStorage) const;
+	/**
+	 * Takes the visibility stored in the Data Storage and copies it to the actor's visibility if the FTypedElementSyncBackToWorldTag
+	 * has been set and the visibility differ.
+	 */
+	void RegisterVisibilityColumnToParticleQuery(IEditorDataStorageProvider& DataStorage) const;
+};
diff --git a/Engine/Plugins/ChaosVD/Source/ChaosVD/Private/TEDS/ChaosVDSelectionInterface.cpp b/Engine/Plugins/ChaosVD/Source/ChaosVD/Private/TEDS/ChaosVDSelectionInterface.cpp
new file mode 100644
--- /dev/null
+++ b/Engine/Plugins/ChaosVD/Source/ChaosVD/Private/TEDS/ChaosVDSelectionInterface.cpp
@@ -0,0 +1,4 @@
+// Copyright Epic Games, Inc. All Rights Reserved.
+
+#include "ChaosVDSelectionInterface.h"
+
diff --git a/Engine/Plugins/ChaosVD/Source/ChaosVD/Private/TEDS/ChaosVDTedsUtils.h b/Engine/Plugins/ChaosVD/Source/ChaosVD/Private/TEDS/ChaosVDTedsUtils.h
new file mode 100644
--- /dev/null
+++ b/Engine/Plugins/ChaosVD/Source/ChaosVD/Private/TEDS/ChaosVDTedsUtils.h
@@ -0,0 +1,125 @@
+// Copyright Epic Games, Inc. All Rights Reserved.
+
+#pragma once
+#include "Elements/Common/EditorDataStorageFeatures.h"
+#include "Elements/Interfaces/TypedElementDataStorageCompatibilityInterface.h"
+
+/** Set of utility methods to interact with TEDS using CVD types */
+namespace Chaos::VD::TedsUtils
+{
+	template<typename ColumnType, typename ObjectType>
+	static void AddColumnToObject(ObjectType* Object)
+	{
+		using namespace UE::Editor::DataStorage;
+
+		if (const IEditorDataStorageCompatibilityProvider* Compatibility = GetDataStorageFeature<IEditorDataStorageCompatibilityProvider>(CompatibilityFeatureName))
+		{
+			RowHandle Row = Compatibility->FindRowWithCompatibleObject(Object);
+			if (IEditorDataStorageProvider* DataStorage = GetMutableDataStorageFeature<IEditorDataStorageProvider>(StorageFeatureName))
+			{
+				if (DataStorage->IsRowAvailable(Row))
+				{
+					DataStorage->AddColumn<ColumnType>(Row);
+				}
+				
+			}
+		}
+	}
+
+	template<typename ColumnType>
+	static void AddColumnToHandle(UE::Editor::DataStorage::RowHandle Handle)
+	{
+		using namespace UE::Editor::DataStorage;
+
+		if (IEditorDataStorageProvider* DataStorage = GetMutableDataStorageFeature<IEditorDataStorageProvider>(StorageFeatureName))
+		{
+			if (DataStorage->IsRowAvailable(Handle))
+			{
+				DataStorage->AddColumn<ColumnType>(Handle);
+			}
+		}
+	}
+
+	template<typename ColumnType>
+	static void BatchAddOrRemoveColumn(TConstArrayView<UE::Editor::DataStorage::RowHandle> Rows)
+	{
+		using namespace UE::Editor::DataStorage;
+
+		if (IEditorDataStorageProvider* DataStorage = GetMutableDataStorageFeature<IEditorDataStorageProvider>(StorageFeatureName))
+		{
+			DataStorage->BatchAddRemoveColumns(Rows, { ColumnType::StaticStruct(), 1}, {});
+		}
+	}
+
+	template<typename ColumnType>
+	static void RemoveColumnToHandle(UE::Editor::DataStorage::RowHandle Handle)
+	{
+		using namespace UE::Editor::DataStorage;
+
+		if (IEditorDataStorageProvider* DataStorage = GetMutableDataStorageFeature<IEditorDataStorageProvider>(StorageFeatureName))
+		{
+			if (DataStorage->IsRowAvailable(Handle))
+			{
+				DataStorage->RemoveColumn<ColumnType>(Handle);
+			}
+		}
+	}
+
+	template<typename ColumnType, typename ObjectType>
+	static void RemoveColumnFromObject(ObjectType* Object)
+	{
+		using namespace UE::Editor::DataStorage;
+
+		if (const IEditorDataStorageCompatibilityProvider* Compatibility = GetDataStorageFeature<IEditorDataStorageCompatibilityProvider>(CompatibilityFeatureName))
+		{
+			RowHandle Row = Compatibility->FindRowWithCompatibleObject(Object);
+
+			if (IEditorDataStorageProvider* DataStorage = GetMutableDataStorageFeature<IEditorDataStorageProvider>(StorageFeatureName))
+			{
+				if (DataStorage->IsRowAvailable(Row))
+				{
+					DataStorage->RemoveColumn<ColumnType>(Row);
+				}
+			}
+		}
+	}
+
+	template<typename ObjectType>
+	static UE::Editor::DataStorage::RowHandle AddObjectToDataStorage(ObjectType* Object, IEditorDataStorageProvider* DataStorageInterface, IEditorDataStorageCompatibilityProvider* DataStorageCompatibility)
+	{
+		using namespace UE::Editor::DataStorage;
+		if (DataStorageCompatibility != nullptr && DataStorageInterface != nullptr)
+		{
+			return DataStorageCompatibility->AddCompatibleObject(Object);
+		}
+
+		return InvalidRowHandle;
+	}
+
+	template<typename ObjectType>
+	static UE::Editor::DataStorage::RowHandle AddObjectToDataStorage(ObjectType* Object)
+	{
+		using namespace UE::Editor::DataStorage;
+		IEditorDataStorageProvider* DataStorageInterface = GetMutableDataStorageFeature<IEditorDataStorageProvider>(StorageFeatureName);
+		IEditorDataStorageCompatibilityProvider* DataStorageCompatibility = GetMutableDataStorageFeature<IEditorDataStorageCompatibilityProvider>(CompatibilityFeatureName);
+
+		return AddObjectToDataStorage(Object, DataStorageInterface, DataStorageCompatibility);	
+	}
+
+	template<typename ObjectType>
+	static void RemoveObjectToFromDataStorage(ObjectType* Object, IEditorDataStorageCompatibilityProvider* DataStorageCompatibility)
+	{
+		if (DataStorageCompatibility != nullptr)
+		{
+			DataStorageCompatibility->RemoveCompatibleObject(Object);
+		}
+	}
+
+	template<typename ObjectType>
+	static void RemoveObjectToFromDataStorage(ObjectType* Object)
+	{
+		using namespace UE::Editor::DataStorage;
+		IEditorDataStorageCompatibilityProvider* DataStorageCompatibility = GetMutableDataStorageFeature<IEditorDataStorageCompatibilityProvider>(CompatibilityFeatureName);
+		RemoveObjectToFromDataStorage(Object, DataStorageCompatibility);
+	}
+}
diff --git a/Engine/Plugins/ChaosVD/Source/ChaosVD/Private/Trace/DataProcessors/ChaosVDTraceParticleDataProcessor.cpp b/Engine/Plugins/ChaosVD/Source/ChaosVD/Private/Trace/DataProcessors/ChaosVDTraceParticleDataProcessor.cpp
--- a/Engine/Plugins/ChaosVD/Source/ChaosVD/Private/Trace/DataProcessors/ChaosVDTraceParticleDataProcessor.cpp
+++ b/Engine/Plugins/ChaosVD/Source/ChaosVD/Private/Trace/DataProcessors/ChaosVDTraceParticleDataProcessor.cpp
@@ -31,7 +31,31 @@ bool FChaosVDTraceParticleDataProcessor::ProcessRawData(const TArray<uint8>& InD
 		EChaosVDSolverStageAccessorFlags StageAccessorFlags = EChaosVDSolverStageAccessorFlags::CreateNewIfEmpty | EChaosVDSolverStageAccessorFlags::CreateNewIfClosed; 
 		if (FChaosVDFrameStageData* CurrentSolverStage = ProviderSharedPtr->GetCurrentSolverStageDataForCurrentFrame(ParticleData->SolverID, StageAccessorFlags))
 		{
-			CurrentSolverStage->RecordedParticlesData.Add(ParticleData);
+			int32 NewParticleDataIndex = INDEX_NONE;
+
+			// Non-Explicit stages might have duplicated data if a particle was traced twice. As CVD only shows the last state per stage, we need to just keep the last seen state
+			// We do this on trace analysis to not do the de-duplication process during playback
+			if (!EnumHasAnyFlags(CurrentSolverStage->StageFlags, EChaosVDSolverStageFlags::ExplicitStage))
+			{
+				if (int32* ExistingParticleDataIndex = CurrentSolverStage->CurrentRecordedParticlesIndexes.Find(ParticleData->ParticleIndex))
+				{
+					CurrentSolverStage->RecordedParticlesData[*ExistingParticleDataIndex] = ParticleData;
+				}
+				else
+				{
+					NewParticleDataIndex = CurrentSolverStage->RecordedParticlesData.Add(ParticleData);
+				}
+			}
+			else
+			{
+				NewParticleDataIndex = CurrentSolverStage->RecordedParticlesData.Add(ParticleData);
+			}
+
+			if (NewParticleDataIndex != INDEX_NONE)
+			{
+				CurrentSolverStage->CurrentRecordedParticlesIndexes.Add(ParticleData->ParticleIndex, NewParticleDataIndex);
+			}
+			
 		}
 	}
 
diff --git a/Engine/Plugins/ChaosVD/Source/ChaosVD/Private/Visualizers/ChaosVDCharacterGroundConstraintsDataComponentVisualizer.cpp b/Engine/Plugins/ChaosVD/Source/ChaosVD/Private/Visualizers/ChaosVDCharacterGroundConstraintsDataComponentVisualizer.cpp
--- a/Engine/Plugins/ChaosVD/Source/ChaosVD/Private/Visualizers/ChaosVDCharacterGroundConstraintsDataComponentVisualizer.cpp
+++ b/Engine/Plugins/ChaosVD/Source/ChaosVD/Private/Visualizers/ChaosVDCharacterGroundConstraintsDataComponentVisualizer.cpp
@@ -161,7 +161,7 @@ void FChaosVDCharacterGroundConstraintDataComponentVisualizer::DrawConstraint(co
 
 	TSharedPtr<const FChaosVDParticleDataWrapper> CharacterParticleData = nullptr;
 
-	if (AChaosVDParticleActor* CharacterParticle = VisualizationContext.SolverInfoActor->GetParticleActor(InConstraintData.CharacterParticleIndex))
+	if (TSharedPtr<FChaosVDSceneParticle> CharacterParticle = VisualizationContext.SolverInfoActor->GetParticleInstance(InConstraintData.CharacterParticleIndex))
 	{
 		CharacterParticleData = CharacterParticle->GetParticleData();
 	}
diff --git a/Engine/Plugins/ChaosVD/Source/ChaosVD/Private/Visualizers/ChaosVDJointConstraintsDataComponentVisualizer.cpp b/Engine/Plugins/ChaosVD/Source/ChaosVD/Private/Visualizers/ChaosVDJointConstraintsDataComponentVisualizer.cpp
--- a/Engine/Plugins/ChaosVD/Source/ChaosVD/Private/Visualizers/ChaosVDJointConstraintsDataComponentVisualizer.cpp
+++ b/Engine/Plugins/ChaosVD/Source/ChaosVD/Private/Visualizers/ChaosVDJointConstraintsDataComponentVisualizer.cpp
@@ -254,12 +254,12 @@ void FChaosVDJointConstraintsDataComponentVisualizer::DrawJointConstraint(const
 	TSharedPtr<const FChaosVDParticleDataWrapper> ParticleData0 = nullptr;
 	TSharedPtr<const FChaosVDParticleDataWrapper> ParticleData1 = nullptr;
 
-	if (AChaosVDParticleActor* Particle0 = VisualizationContext.SolverInfoActor->GetParticleActor(InJointConstraintData.ParticleParIndexes[0]))
+	if (TSharedPtr<FChaosVDSceneParticle> Particle0 = VisualizationContext.SolverInfoActor->GetParticleInstance(InJointConstraintData.ParticleParIndexes[0]))
 	{
 		ParticleData0 = Particle0->GetParticleData();
 	}
 
-	if (AChaosVDParticleActor* Particle1 = VisualizationContext.SolverInfoActor->GetParticleActor(InJointConstraintData.ParticleParIndexes[1]))
+	if (TSharedPtr<FChaosVDSceneParticle> Particle1 = VisualizationContext.SolverInfoActor->GetParticleInstance(InJointConstraintData.ParticleParIndexes[1]))
 	{
 		ParticleData1 = Particle1->GetParticleData();
 	}
