[ChaosVD] Refactor that removes actors as a mean to represent recorded particles.

diff --git a/Engine/Plugins/ChaosVD/Source/ChaosVD/Private/Visualizers/ChaosVDParticleDataComponentVisualizer.cpp b/Engine/Plugins/ChaosVD/Source/ChaosVD/Private/Visualizers/ChaosVDParticleDataComponentVisualizer.cpp
--- a/Engine/Plugins/ChaosVD/Source/ChaosVD/Private/Visualizers/ChaosVDParticleDataComponentVisualizer.cpp
+++ b/Engine/Plugins/ChaosVD/Source/ChaosVD/Private/Visualizers/ChaosVDParticleDataComponentVisualizer.cpp
@@ -3,7 +3,6 @@
 #include "Visualizers/ChaosVDParticleDataComponentVisualizer.h"
 
 #include "Actors/ChaosVDSolverInfoActor.h"
-#include "ChaosVDParticleActor.h"
 #include "ChaosVDScene.h"
 #include "ChaosVDSettingsManager.h"
 #include "ChaosVDTabsIDs.h"
@@ -150,6 +149,7 @@ void FChaosVDParticleDataComponentVisualizer::DrawVisualization(const UActorComp
 	VisualizationContext.bShowDebugText = VisualizationSettings->bShowDebugText;
 	VisualizationContext.DebugDrawSettings = VisualizationSettings;
 	VisualizationContext.SolverDataSelectionObject = CVDScene->GetSolverDataSelectionObject().Pin();
+
 	
 
 	if (!VisualizationContext.IsVisualizationFlagEnabled(EChaosVDParticleDataVisualizationFlags::EnableDraw))
@@ -160,7 +160,7 @@ void FChaosVDParticleDataComponentVisualizer::DrawVisualization(const UActorComp
 	if (VisualizationContext.IsVisualizationFlagEnabled(EChaosVDParticleDataVisualizationFlags::DrawDataOnlyForSelectedParticle))
 	{
 		VisualizationContext.bIsSelectedData = true;
-		SolverDataActor->VisitSelectedParticleData([this, PDI, View, &VisualizationContext, Component](const TSharedPtr<const FChaosVDParticleDataWrapper>& InParticleDataViewer)
+		ParticleDataComponent->VisitSelectedParticleData([this, PDI, View, &VisualizationContext, Component](const TSharedPtr<const FChaosVDParticleDataWrapper>& InParticleDataViewer)
 		{
 			DrawVisualizationForParticleData(Component, PDI, View, VisualizationContext, InParticleDataViewer);
 
@@ -169,9 +169,17 @@ void FChaosVDParticleDataComponentVisualizer::DrawVisualization(const UActorComp
 	}
 	else
 	{
-		SolverDataActor->VisitAllParticleData([this, PDI, View, &VisualizationContext, Component, SolverDataActor](const TSharedPtr<const FChaosVDParticleDataWrapper>& InParticleDataViewer)
+		int32 SelectedParticleID = INDEX_NONE;
+
+		FChaosVDSceneParticle* SelectedParticle = ParticleDataComponent->GetSelectedParticle();
+		if (TSharedPtr<const FChaosVDParticleDataWrapper> SelectedParticleData = SelectedParticle ? SelectedParticle->GetParticleData() : nullptr)
 		{
-			VisualizationContext.bIsSelectedData = InParticleDataViewer && SolverDataActor->IsParticleSelectedByID(InParticleDataViewer->ParticleIndex);
+			SelectedParticleID = SelectedParticleData->ParticleIndex;
+		}
+
+		ParticleDataComponent->VisitAllParticleData([this, PDI, View, &VisualizationContext, Component, SelectedParticleID](const TSharedPtr<const FChaosVDParticleDataWrapper>& InParticleDataViewer)
+		{
+			VisualizationContext.bIsSelectedData = InParticleDataViewer && InParticleDataViewer->ParticleIndex == SelectedParticleID;
 			DrawVisualizationForParticleData(Component, PDI, View, VisualizationContext, InParticleDataViewer);
 
 			// If we reach the debug draw limit for this frame, there is no need to continue processing particles
@@ -202,7 +210,12 @@ bool FChaosVDParticleDataComponentVisualizer::SelectVisualizedData(const HChaosV
 
 	if (TSharedPtr<const FChaosVDParticleDataWrapper> ParticleDataViewer = VisProxy.DataSelectionHandle ?  VisProxy.DataSelectionHandle->GetDataAsShared<const FChaosVDParticleDataWrapper>() : nullptr)
 	{
-		bHandled = SolverDataActor->SelectParticleByID(ParticleDataViewer->ParticleIndex);
+		using namespace Chaos::VD::TypedElementDataUtil;
+		if (TSharedPtr<FChaosVDSceneParticle> ParticleInstance = SolverDataActor->GetParticleInstance(ParticleDataViewer->ParticleIndex))
+		{
+			InCVDScene->SetSelected(AcquireTypedElementHandleForStruct(ParticleInstance.Get(), true));
+			bHandled = true;
+		}
 	}
 
 	return bHandled;
@@ -301,7 +314,7 @@ void FChaosVDParticleDataComponentVisualizer::DrawVisualizationForParticleData(c
 			{
 				if (const TSharedPtr<FChaosVDScene> ScenePtr = InVisualizationContext.CVDScene.Pin())
 				{
-					if (AChaosVDParticleActor* SiblingParticle = ScenePtr->GetParticleActor(InVisualizationContext.SolverID, ConnectivityEdge.SiblingParticleID))
+					if (TSharedPtr<FChaosVDSceneParticle> SiblingParticle = ScenePtr->GetParticleInstance(InVisualizationContext.SolverID, ConnectivityEdge.SiblingParticleID))
 					{
 						if (TSharedPtr<const FChaosVDParticleDataWrapper> SiblingParticleData = SiblingParticle->GetParticleData())
 						{
diff --git a/Engine/Plugins/ChaosVD/Source/ChaosVD/Private/Visualizers/ChaosVDParticleDataComponentVisualizer.h b/Engine/Plugins/ChaosVD/Source/ChaosVD/Private/Visualizers/ChaosVDParticleDataComponentVisualizer.h
--- a/Engine/Plugins/ChaosVD/Source/ChaosVD/Private/Visualizers/ChaosVDParticleDataComponentVisualizer.h
+++ b/Engine/Plugins/ChaosVD/Source/ChaosVD/Private/Visualizers/ChaosVDParticleDataComponentVisualizer.h
@@ -5,14 +5,15 @@
 #include "Visualizers/ChaosVDComponentVisualizerBase.h"
 #include "ComponentVisualizer.h"
 #include "HAL/Platform.h"
-#include "IChaosVDParticleVisualizationDataProvider.h"
 #include "Settings/ChaosVDParticleVisualizationSettings.h"
 #include "Templates/SharedPointer.h"
 
 class FChaosVDGeometryBuilder;
+class UChaosVDParticleVisualizationDebugDrawSettings;
+
 struct FChaosParticleDataDebugDrawSettings;
+struct FChaosVDParticleDataWrapper;
 
-class UChaosVDParticleVisualizationDebugDrawSettings;
 
 struct FChaosVDParticleDataVisualizationContext : public FChaosVDVisualizationContext
 {
diff --git a/Engine/Plugins/ChaosVD/Source/ChaosVD/Private/Visualizers/ChaosVDSolverCollisionDataComponentVisualizer.cpp b/Engine/Plugins/ChaosVD/Source/ChaosVD/Private/Visualizers/ChaosVDSolverCollisionDataComponentVisualizer.cpp
--- a/Engine/Plugins/ChaosVD/Source/ChaosVD/Private/Visualizers/ChaosVDSolverCollisionDataComponentVisualizer.cpp
+++ b/Engine/Plugins/ChaosVD/Source/ChaosVD/Private/Visualizers/ChaosVDSolverCollisionDataComponentVisualizer.cpp
@@ -2,7 +2,6 @@
 
 #include "ChaosVDSolverCollisionDataComponentVisualizer.h"
 
-#include "ChaosVDParticleActor.h"
 #include "ChaosVDScene.h"
 #include "ChaosVDSettingsManager.h"
 #include "ChaosVDTabsIDs.h"
@@ -14,6 +13,7 @@
 #include "ToolMenus.h"
 #include "ToolMenuEntry.h"
 #include "ToolMenuSection.h"
+#include "Components/ChaosVDParticleDataComponent.h"
 #include "Visualizers/ChaosVDDebugDrawUtils.h"
 #include "Widgets/SChaosVDCollisionDataInspector.h"
 #include "Widgets/SChaosVDViewportToolbar.h"
@@ -59,12 +59,18 @@ void FChaosVDSolverCollisionDataComponentVisualizer::DrawVisualization(const UAc
 		return;
 	}
 
-	const AChaosVDSolverInfoActor* SolverInfoContainer = Cast<AChaosVDSolverInfoActor>(CollisionDataComponent->GetOwner());
+	AChaosVDSolverInfoActor* SolverInfoContainer = Cast<AChaosVDSolverInfoActor>(CollisionDataComponent->GetOwner());
 	if (!SolverInfoContainer)
 	{
 		return;
 	}
 
+	UChaosVDParticleDataComponent* ParticleDataComponent = SolverInfoContainer->GetParticleDataComponent();
+	if (!ParticleDataComponent)
+	{
+		return;
+	}
+
 	if (!SolverInfoContainer->IsVisible())
 	{
 		return;
@@ -91,18 +97,17 @@ void FChaosVDSolverCollisionDataComponentVisualizer::DrawVisualization(const UAc
 	}
 	else
 	{
-		const TArray<int32>& SelectedParticlesIDs = SolverInfoContainer->GetSelectedParticlesIDs();
-
-		for (const int32 SelectedParticleID : SelectedParticlesIDs)
+		ParticleDataComponent->VisitSelectedParticleData([this, PDI, View, &VisualizationContext, CollisionDataComponent](const TSharedPtr<const FChaosVDParticleDataWrapper>& InParticleDataViewer)
 		{
-			if (const TArray<TSharedPtr<FChaosVDParticlePairMidPhase>>* ParticleMidPhases = CollisionDataComponent->GetMidPhasesForParticle(SelectedParticleID, EChaosVDParticlePairSlot::Any))
+			if (const TArray<TSharedPtr<FChaosVDParticlePairMidPhase>>* ParticleMidPhases = CollisionDataComponent->GetMidPhasesForParticle(InParticleDataViewer->ParticleIndex, EChaosVDParticlePairSlot::Any))
 			{
 				for (const TSharedPtr<FChaosVDParticlePairMidPhase>& MidPhase : *ParticleMidPhases)
 				{
-					DrawMidPhaseData(Component, MidPhase, VisualizationContext, View, PDI);
+					DrawMidPhaseData(CollisionDataComponent, MidPhase, VisualizationContext, View, PDI);
 				}
 			}
-		}
+			return true;
+		});
 	}
 }
 
@@ -127,7 +132,7 @@ void FChaosVDSolverCollisionDataComponentVisualizer::DrawMidPhaseData(const UAct
 	
 	if (MidPhase.IsValid())
 	{
-		AChaosVDParticleActor* ParticleActor = ScenePtr->GetParticleActor(VisualizationContext.SolverID, MidPhase->Particle0Idx);
+		TSharedPtr<FChaosVDSceneParticle> ParticleActor = ScenePtr->GetParticleInstance(VisualizationContext.SolverID, MidPhase->Particle0Idx);
 		if (!ParticleActor)
 		{
 			return;
@@ -180,8 +185,8 @@ void FChaosVDSolverCollisionDataComponentVisualizer::DrawMidPhaseData(const UAct
 
 				const bool bPruned = ManifoldPoint.bDisabled;
 
-				AChaosVDParticleActor* CVDParticleActor1 = ScenePtr->GetParticleActor(VisualizationContext.SolverID, Constraint.Particle1Index);
-				AChaosVDParticleActor* CVDParticleActor0 = ScenePtr->GetParticleActor(VisualizationContext.SolverID, Constraint.Particle0Index);
+				TSharedPtr<FChaosVDSceneParticle> CVDParticleActor1 = ScenePtr->GetParticleInstance(VisualizationContext.SolverID, Constraint.Particle1Index);
+				TSharedPtr<FChaosVDSceneParticle> CVDParticleActor0 = ScenePtr->GetParticleInstance(VisualizationContext.SolverID, Constraint.Particle0Index);
 				if (CVDParticleActor1 == nullptr || CVDParticleActor0 == nullptr)
 				{
 					continue;
diff --git a/Engine/Plugins/ChaosVD/Source/ChaosVD/Private/Widgets/SChaosVDCollisionDataInspector.cpp b/Engine/Plugins/ChaosVD/Source/ChaosVD/Private/Widgets/SChaosVDCollisionDataInspector.cpp
--- a/Engine/Plugins/ChaosVD/Source/ChaosVD/Private/Widgets/SChaosVDCollisionDataInspector.cpp
+++ b/Engine/Plugins/ChaosVD/Source/ChaosVD/Private/Widgets/SChaosVDCollisionDataInspector.cpp
@@ -2,7 +2,7 @@
 
 #include "Widgets/SChaosVDCollisionDataInspector.h"
 
-#include "ChaosVDParticleActor.h"
+
 #include "ChaosVDScene.h"
 #include "IStructureDetailsView.h"
 #include "Modules/ModuleManager.h"
@@ -16,11 +16,11 @@
 
 #define LOCTEXT_NAMESPACE "ChaosVisualDebugger"
 
-void SChaosVDCollisionDataInspector::SetCollisionDataProviderObjectToInspect(IChaosVDCollisionDataProviderInterface* CollisionDataProvider)
+void SChaosVDCollisionDataInspector::SetCollisionDataListToInspect(TConstArrayView<TSharedPtr<FChaosVDParticlePairMidPhase>> CollisionDataList)
 {
 	ClearInspector();
 
-	if (CollisionDataProvider == nullptr)
+	if (CollisionDataList.IsEmpty())
 	{
 		return;
 	}
@@ -39,9 +39,8 @@ void SChaosVDCollisionDataInspector::SetCollisionDataProviderObjectToInspect(ICh
 	}
 
 	TArray<TSharedPtr<FName>> NewCollisionDataEntriesNameList;
-	TConstArrayView<TSharedPtr<FChaosVDParticlePairMidPhase>> FoundCollisionData = CollisionDataProvider->GetCollisionData();
 
-	for (const TSharedPtr<FChaosVDParticlePairMidPhase>& CollisionData : FoundCollisionData)
+	for (const TSharedPtr<FChaosVDParticlePairMidPhase>& CollisionData : CollisionDataList)
 	{
 		TSharedPtr<FChaosVDSolverDataSelectionHandle> SelectionHandle = SolverDataSelectionObject->MakeSelectionHandle(CollisionData);
 		if (TSharedPtr<FName> CollisionDataEntryName = GenerateNameForCollisionDataItem(SelectionHandle))
@@ -110,8 +109,11 @@ void SChaosVDCollisionDataInspector::ClearInspector()
 	CurrentSelectedName = nullptr;
 
 	bIsUpToDate = true;
-	
-	SecondaryCollisionDataDetailsPanel->SetStructureData(nullptr);
+
+	if (SecondaryCollisionDataDetailsPanel->GetStructureProvider())
+	{
+		SecondaryCollisionDataDetailsPanel->SetStructureData(nullptr);
+	}
 
 	SChaosVDConstraintDataInspector::ClearInspector();
 }
@@ -151,7 +153,7 @@ FText SChaosVDCollisionDataInspector::GetParticleName(EChaosVDParticlePairIndex
 
 	if (const TSharedPtr<FChaosVDScene> ScenePtr = SceneWeakPtr.Pin())
 	{
-		if (AChaosVDParticleActor* ParticleActor = ScenePtr->GetParticleActor(OutSolverID, OutParticleID))
+		if (TSharedPtr<FChaosVDSceneParticle> ParticleActor = ScenePtr->GetParticleInstance(OutSolverID, OutParticleID))
 		{
 			if (TSharedPtr<const FChaosVDParticleDataWrapper> ParticleData = ParticleActor->GetParticleData())
 			{
diff --git a/Engine/Plugins/ChaosVD/Source/ChaosVD/Private/Widgets/SChaosVDCollisionDataInspector.h b/Engine/Plugins/ChaosVD/Source/ChaosVD/Private/Widgets/SChaosVDCollisionDataInspector.h
--- a/Engine/Plugins/ChaosVD/Source/ChaosVD/Private/Widgets/SChaosVDCollisionDataInspector.h
+++ b/Engine/Plugins/ChaosVD/Source/ChaosVD/Private/Widgets/SChaosVDCollisionDataInspector.h
@@ -6,7 +6,6 @@
 #include "Widgets/SCompoundWidget.h"
 
 struct FChaosVDSolverDataSelectionHandle;
-class AChaosVDParticleActor;
 class FChaosVDScene;
 class IStructureDetailsView;
 class IChaosVDCollisionDataProviderInterface;
@@ -19,7 +18,7 @@ class SChaosVDCollisionDataInspector : public SChaosVDConstraintDataInspector
 {
 public:
 
-	void SetCollisionDataProviderObjectToInspect(IChaosVDCollisionDataProviderInterface* CollisionDataProvider);
+	void SetCollisionDataListToInspect(TConstArrayView<TSharedPtr<FChaosVDParticlePairMidPhase>> CollisionDataList);
 
 	virtual void SetConstraintDataToInspect(const TSharedPtr<FChaosVDSolverDataSelectionHandle>& InDataSelectionHandle) override;
 
@@ -57,4 +56,5 @@ protected:
 	TSharedPtr<FName> CurrentSelectedName;
 
 	TSharedPtr<IStructureDetailsView> SecondaryCollisionDataDetailsPanel;
+	
 };
diff --git a/Engine/Plugins/ChaosVD/Source/ChaosVD/Private/Widgets/SChaosVDConstraintDataInspector.cpp b/Engine/Plugins/ChaosVD/Source/ChaosVD/Private/Widgets/SChaosVDConstraintDataInspector.cpp
--- a/Engine/Plugins/ChaosVD/Source/ChaosVD/Private/Widgets/SChaosVDConstraintDataInspector.cpp
+++ b/Engine/Plugins/ChaosVD/Source/ChaosVD/Private/Widgets/SChaosVDConstraintDataInspector.cpp
@@ -9,6 +9,7 @@
 #include "Actors/ChaosVDSolverInfoActor.h"
 #include "Widgets/SChaosVDWarningMessageBox.h"
 #include "Modules/ModuleManager.h"
+#include "TEDS/ChaosVDStructTypedElementData.h"
 #include "Widgets/Input/SButton.h"
 #include "Widgets/Layout/SBox.h"
 #include "Widgets/Layout/SScrollBox.h"
@@ -291,7 +292,12 @@ void SChaosVDConstraintDataInspector::HandleSceneUpdated()
 void SChaosVDConstraintDataInspector::ClearInspector()
 {
 	DataBeingInspectedCopy = nullptr;
-	MainDataDetailsView->SetStructureData(nullptr);
+
+	if (MainDataDetailsView->GetStructureProvider())
+	{
+		MainDataDetailsView->SetStructureData(nullptr);
+	}
+	
 	CurrentDataSelectionHandle = MakeShared<FChaosVDSolverDataSelectionHandle>();
 }
 
@@ -299,7 +305,7 @@ FText SChaosVDConstraintDataInspector::GetParticleName_Internal(int32 SolverID,
 {
 	if (const TSharedPtr<FChaosVDScene> ScenePtr = SceneWeakPtr.Pin())
 	{
-		if (AChaosVDParticleActor* ParticleActor = ScenePtr->GetParticleActor(SolverID, ParticleID))
+		if (TSharedPtr<FChaosVDSceneParticle> ParticleActor = ScenePtr->GetParticleInstance(SolverID, ParticleID))
 		{
 			if (TSharedPtr<const FChaosVDParticleDataWrapper> ParticleData = ParticleActor->GetParticleData())
 			{
@@ -315,9 +321,10 @@ void SChaosVDConstraintDataInspector::SelectParticle(int32 SolverID, int32 Parti
 {
 	if (const TSharedPtr<FChaosVDScene> ScenePtr = SceneWeakPtr.Pin())
 	{
-		if (AChaosVDParticleActor* ParticleActor = ScenePtr->GetParticleActor(SolverID, ParticleID))
+		if (TSharedPtr<FChaosVDSceneParticle> ParticleActor = ScenePtr->GetParticleInstance(SolverID, ParticleID))
 		{
-			ScenePtr->SetSelectedObject(ParticleActor);
+			using namespace Chaos::VD::TypedElementDataUtil;
+			ScenePtr->SetSelected(AcquireTypedElementHandleForStruct(ParticleActor.Get(), true));
 		}
 	}
 }
diff --git a/Engine/Plugins/ChaosVD/Source/ChaosVD/Private/Widgets/SChaosVDDetailsView.cpp b/Engine/Plugins/ChaosVD/Source/ChaosVD/Private/Widgets/SChaosVDDetailsView.cpp
--- a/Engine/Plugins/ChaosVD/Source/ChaosVD/Private/Widgets/SChaosVDDetailsView.cpp
+++ b/Engine/Plugins/ChaosVD/Source/ChaosVD/Private/Widgets/SChaosVDDetailsView.cpp
@@ -2,7 +2,6 @@
 
 #include "Widgets/SChaosVDDetailsView.h"
 
-#include "ChaosVDParticleActor.h"
 #include "PropertyEditorModule.h"
 #include "Widgets/SChaosVDMainTab.h"
 
diff --git a/Engine/Plugins/ChaosVD/Source/ChaosVD/Private/Widgets/SChaosVDMainTab.cpp b/Engine/Plugins/ChaosVD/Source/ChaosVD/Private/Widgets/SChaosVDMainTab.cpp
--- a/Engine/Plugins/ChaosVD/Source/ChaosVD/Private/Widgets/SChaosVDMainTab.cpp
+++ b/Engine/Plugins/ChaosVD/Source/ChaosVD/Private/Widgets/SChaosVDMainTab.cpp
@@ -13,8 +13,7 @@
 #include "ChaosVDSolversTracksTab.h"
 #include "ChaosVDCollisionDataDetailsTab.h"
 #include "ChaosVDConstraintDataInspectorTab.h"
-#include "ChaosVDParticleActor.h"
-#include "ChaosVDParticleActorCustomization.h"
+#include "ChaosVDSceneParticleCustomization.h"
 #include "ChaosVDRecordedLogTab.h"
 #include "ChaosVDSceneQueryDataInspectorTab.h"
 #include "ChaosVDStyle.h"
@@ -578,7 +577,7 @@ void SChaosVDMainTab::SetCustomPropertyLayouts(IDetailsView* DetailsView)
 		return;
 	}
 
-	DetailsView->RegisterInstancedCustomPropertyLayout(AChaosVDParticleActor::StaticClass(), FOnGetDetailCustomizationInstance::CreateStatic(&FChaosVDParticleActorCustomization::MakeInstance, StaticCastWeakPtr<SChaosVDMainTab>(AsWeak())));
+	DetailsView->RegisterInstancedCustomPropertyLayout(FChaosVDSceneParticle::StaticStruct(), FOnGetDetailCustomizationInstance::CreateStatic(&FChaosVDSceneParticleCustomization::MakeInstance, StaticCastWeakPtr<SChaosVDMainTab>(AsWeak())));
 	DetailsView->RegisterInstancedCustomPropertyLayout(UChaosVDInstancedStaticMeshComponent::StaticClass(), FOnGetDetailCustomizationInstance::CreateStatic(&FChaosVDGeometryComponentCustomization::MakeInstance));
 	DetailsView->RegisterInstancedCustomPropertyLayout(UChaosVDStaticMeshComponent::StaticClass(), FOnGetDetailCustomizationInstance::CreateStatic(&FChaosVDGeometryComponentCustomization::MakeInstance));
 	DetailsView->RegisterInstancedCustomPropertyLayout(FChaosVDQueryVisitStep::StaticStruct(), FOnGetDetailCustomizationInstance::CreateStatic(&FChaosVDQueryVisitDataCustomization::MakeInstance));
diff --git a/Engine/Plugins/ChaosVD/Source/ChaosVD/Private/Widgets/SChaosVDNameListPicker.cpp b/Engine/Plugins/ChaosVD/Source/ChaosVD/Private/Widgets/SChaosVDNameListPicker.cpp
--- a/Engine/Plugins/ChaosVD/Source/ChaosVD/Private/Widgets/SChaosVDNameListPicker.cpp
+++ b/Engine/Plugins/ChaosVD/Source/ChaosVD/Private/Widgets/SChaosVDNameListPicker.cpp
@@ -68,10 +68,15 @@ void SChaosVDNameListPicker::Construct(const FArguments& InArgs)
 
 void SChaosVDNameListPicker::UpdateNameList(TArray<TSharedPtr<FName>>&& NewNameList)
 {
+	if (NewNameList.IsEmpty() && CachedNameList.IsEmpty())
+	{
+		return;
+	}
+
 	// TODO: It might be worth to check is the array is different before doing this, but it might be more expensive
 	CachedNameList = MoveTemp(NewNameList);
 
-	if(!HasElements())
+	if (!HasElements())
 	{
 		CurrentSelectedName = nullptr;
 	}
diff --git a/Engine/Plugins/ChaosVD/Source/ChaosVD/Private/Widgets/SChaosVDSceneQueryDataInspector.cpp b/Engine/Plugins/ChaosVD/Source/ChaosVD/Private/Widgets/SChaosVDSceneQueryDataInspector.cpp
--- a/Engine/Plugins/ChaosVD/Source/ChaosVD/Private/Widgets/SChaosVDSceneQueryDataInspector.cpp
+++ b/Engine/Plugins/ChaosVD/Source/ChaosVD/Private/Widgets/SChaosVDSceneQueryDataInspector.cpp
@@ -3,7 +3,6 @@
 #include "Widgets/SChaosVDSceneQueryDataInspector.h"
 
 #include "ChaosVDModule.h"
-#include "ChaosVDParticleActor.h"
 #include "ChaosVDScene.h"
 #include "EditorModeManager.h"
 #include "EditorViewportClient.h"
@@ -440,19 +439,19 @@ FReply SChaosVDSceneQueryDataInspector::SelectParticleForCurrentQueryData() cons
 		return FReply::Handled();
 	}
 
-	if (AChaosVDParticleActor* ParticleActor = ScenePtr->GetParticleActor(QueryDataBeingInspected->WorldSolverID, QueryDataBeingInspected->SQVisitData[QueryDataBeingInspected->CurrentVisitIndex].ParticleIndex))
+	if (TSharedPtr<FChaosVDSceneParticle> ParticleActor = ScenePtr->GetParticleInstance(QueryDataBeingInspected->WorldSolverID, QueryDataBeingInspected->SQVisitData[QueryDataBeingInspected->CurrentVisitIndex].ParticleIndex))
 	{
 		const int32 ShapeInstanceIndexToSelect = QueryDataBeingInspected->SQVisitData[QueryDataBeingInspected->CurrentVisitIndex].ShapeIndex;
 
 		// TODO : This will not work properly when the visited shape was a Union within a union
 		// CVD currently doesn't support multi selection, so we can't easily select all mesh instances that represent a union within a union
 		// We need to revisit this when multi selection support is added. Jira for tracking UE-212733
-		const TConstArrayView<TSharedPtr<FChaosVDMeshDataInstanceHandle>> AvailableMeshInstances = ParticleActor->GetMeshInstances();
-		for (const TSharedPtr<FChaosVDMeshDataInstanceHandle>& MeshInstance : AvailableMeshInstances)
+		const TConstArrayView<TSharedRef<FChaosVDMeshDataInstanceHandle>> AvailableMeshInstances = ParticleActor->GetMeshInstances();
+		for (const TSharedRef<FChaosVDMeshDataInstanceHandle>& MeshInstance : AvailableMeshInstances)
 		{
 			if (MeshInstance->GetState().ImplicitObjectInfo.ShapeInstanceIndex == ShapeInstanceIndexToSelect)
 			{
-				Chaos::VisualDebugger::SelectParticleWithGeometryInstance(ScenePtr.ToSharedRef(), ParticleActor, MeshInstance);
+				Chaos::VisualDebugger::SelectParticleWithGeometryInstance(ScenePtr.ToSharedRef(), ParticleActor.Get(), MeshInstance);
 				break;
 			}
 		}
@@ -564,12 +563,12 @@ void SChaosVDSceneQueryDataInspector::HandleSubQueryNameSelected(TSharedPtr<FNam
 
 void SChaosVDSceneQueryDataInspector::ClearInspector()
 {
-	if (SceneQueryDataDetailsView)
+	if (SceneQueryDataDetailsView && SceneQueryDataDetailsView->GetStructureProvider())
 	{
 		SceneQueryDataDetailsView->SetStructureData(nullptr);
 	}
 
-	if (SceneQueryHitDataDetailsView)
+	if (SceneQueryHitDataDetailsView && SceneQueryHitDataDetailsView->GetStructureProvider())
 	{
 		SceneQueryHitDataDetailsView->SetStructureData(nullptr);
 	}
diff --git a/Engine/Plugins/ChaosVD/Source/ChaosVD/Public/Actors/ChaosVDDataContainerBaseActor.h b/Engine/Plugins/ChaosVD/Source/ChaosVD/Public/Actors/ChaosVDDataContainerBaseActor.h
--- a/Engine/Plugins/ChaosVD/Source/ChaosVD/Public/Actors/ChaosVDDataContainerBaseActor.h
+++ b/Engine/Plugins/ChaosVD/Source/ChaosVD/Public/Actors/ChaosVDDataContainerBaseActor.h
@@ -5,6 +5,7 @@
 #include "GameFramework/Actor.h"
 #include "ChaosVDDataContainerBaseActor.generated.h"
 
+struct FChaosVDFrameStageData;
 class FChaosVDExtension;
 class FChaosVDScene;
 
@@ -25,6 +26,8 @@ public:
 
 	virtual void UpdateFromNewSolverFrameData(const FChaosVDSolverFrameData& InSolverFrameData);
 
+	virtual void UpdateFromNewSolverStageData(const FChaosVDSolverFrameData& InSolverFrameData, const FChaosVDFrameStageData& InSolverFrameStageData);
+
 	virtual void Destroyed() override;
 
 	virtual void SetScene(TWeakPtr<FChaosVDScene> InScene)
@@ -34,11 +37,18 @@ public:
 	
 	void SetSolverID(int32 InSolverID);
 
+	int32 GetSolverID() const
+	{
+		return SolverDataID;
+	}
+
 	virtual const FTransform& GetSimulationTransform() const
 	{
 		return FTransform::Identity;
 	}
 
+	virtual void UpdateVisibility(bool bIsVisible);
+
 	virtual bool IsVisible() const
 	{
 		return true;
@@ -51,6 +61,10 @@ public:
 		return SceneWeakPtr;
 	}
 	
+#if WITH_EDITOR
+	void SetIsTemporarilyHiddenInEditor(bool bIsHidden) override;
+#endif
+	
 protected:
 
 	void HandlePostInitializationExtensionRegistered(const TSharedRef<FChaosVDExtension>& NewExtension);
diff --git a/Engine/Plugins/ChaosVD/Source/ChaosVD/Public/ChaosVDGeometryBuilder.h b/Engine/Plugins/ChaosVD/Source/ChaosVD/Public/ChaosVDGeometryBuilder.h
--- a/Engine/Plugins/ChaosVD/Source/ChaosVD/Public/ChaosVDGeometryBuilder.h
+++ b/Engine/Plugins/ChaosVD/Source/ChaosVD/Public/ChaosVDGeometryBuilder.h
@@ -67,7 +67,7 @@ public:
 	 *	@param MeshIndex Index of the current component being processed. This is useful when this method is called recursively
 	 *	@param InTransform to apply to the generated components/geometry
 	 */
-	void CreateMeshesFromImplicitObject(const Chaos::FImplicitObject* InImplicitObject, AActor* Owner, TArray<TSharedPtr<FChaosVDExtractedGeometryDataHandle>>& OutMeshDataHandles, int32 AvailableShapeDataNum = 0, const int32 DesiredLODCount = 0, const Chaos::FRigidTransform3& InTransform = Chaos::FRigidTransform3(), const int32 MeshIndex = 0);
+	void CreateMeshesFromImplicitObject(const Chaos::FImplicitObject* InImplicitObject, TArray<TSharedPtr<FChaosVDExtractedGeometryDataHandle>>& OutMeshDataHandles, int32 AvailableShapeDataNum = 0, const int32 DesiredLODCount = 0, const Chaos::FRigidTransform3& InTransform = Chaos::FRigidTransform3(), const int32 MeshIndex = 0);
 
 	virtual void AddReferencedObjects(FReferenceCollector& Collector) override;
 	virtual FString GetReferencerName() const override
@@ -91,7 +91,7 @@ public:
 
 private:
 	
-	void CreateMeshesFromImplicit_Internal(const Chaos::FImplicitObject* InRootImplicitObject,const Chaos::FImplicitObject* InLeafImplicitObject, AActor* Owner, TArray<TSharedPtr<FChaosVDExtractedGeometryDataHandle>>& OutMeshDataHandles, const int32 DesiredLODCount = 0, const Chaos::FRigidTransform3& InTransform = Chaos::FRigidTransform3(), const int32 ParentShapeInstanceIndex = 0, int32 AvailableShapeDataNum = 0);
+	void CreateMeshesFromImplicit_Internal(const Chaos::FImplicitObject* InRootImplicitObject,const Chaos::FImplicitObject* InLeafImplicitObject, TArray<TSharedPtr<FChaosVDExtractedGeometryDataHandle>>& OutMeshDataHandles, const int32 DesiredLODCount = 0, const Chaos::FRigidTransform3& InTransform = Chaos::FRigidTransform3(), const int32 ParentShapeInstanceIndex = 0, int32 AvailableShapeDataNum = 0);
 	
 public:
 	/**
@@ -128,6 +128,9 @@ private:
 	void SetMeshComponentMaterial(IChaosVDGeometryComponent* GeometryComponent);
 
 public:
+
+	void HandleNewGeometryData(const Chaos::FConstImplicitObjectPtr& Geometry, const uint32 GeometryID);
+
 	/**
 	 * Finds or creates a Mesh component for the geometry data handle provided, and add a new instance of that geometry to it
 	 * @param InOwningParticleData Particle data from which the implicit object is from
@@ -311,6 +314,69 @@ private:
 	TWeakPtr<FChaosVDScene> SceneWeakPtr;
 
 	bool bInitialized = false;
+
+	struct FSourceGeometryHashCache
+	{
+		/** Returns the hash for the provided implicit object and caches it if it is the first time we see it.
+		 * @param ImplicitObject Implicit Object from which get the hash from
+		 */
+		[[nodiscard]] uint32 GetAndCacheGeometryHash(const Chaos::FImplicitObject* ImplicitObject)
+		{
+			if (!ImplicitObject)
+			{
+				return 0;
+			}
+
+			{
+				FReadScopeLock ReadLock(CacheLock);
+				if (uint32* FoundHash = CachedGeometryHashes.Find(ImplicitObject))
+				{
+					return *FoundHash;
+				}
+			}
+			
+			const uint32 GeometryHash = ImplicitObject->GetTypeHash();
+			CacheImplicitObjectHash(ImplicitObject, GeometryHash);
+
+			return GeometryHash;
+		}
+
+		/** Returns true if we have the hash for the provided implicit object in cache
+		 * @param ImplicitObject Implicit Object to cache
+		 */
+		bool HasGeometryInHashCache(const Chaos::FImplicitObject* ImplicitObject)
+		{
+			FReadScopeLock ReadLock(CacheLock);
+			return CachedGeometryHashes.Contains(ImplicitObject);
+		}
+
+		/** Caches the provided hash linking it to the provided implicit object
+		 * @param ImplicitObject Implicit Object from which the has was calculated from
+		 * @param Hash Hash from the Implicit Object
+		 */
+		void CacheImplicitObjectHash(const Chaos::FImplicitObject* ImplicitObject, uint32 Hash)
+		{
+			if (!ImplicitObject)
+			{
+				return;
+			}
+	
+			FWriteScopeLock WriteLock(CacheLock);
+			CachedGeometryHashes.Add(ImplicitObject, Hash);
+		}
+
+		/** Clears the hash cache */
+		void Reset()
+		{
+			FWriteScopeLock WriteLock(CacheLock);
+			CachedGeometryHashes.Reset();
+		}
+
+		FRWLock CacheLock;
+		TMap<const void*, uint32> CachedGeometryHashes;
+	};
+
+	FSourceGeometryHashCache SourceGeometryCache;
 	
 	friend class FChaosVDGeometryGenerationTask;
 };
@@ -457,7 +523,7 @@ TSharedPtr<FChaosVDMeshDataInstanceHandle> FChaosVDGeometryBuilder::CreateMeshDa
 	{
 		constexpr bool bIsWorldSpace = true;
 		const FTransform OwningParticleTransform(InOwningParticleData.ParticlePositionRotation.MR, InOwningParticleData.ParticlePositionRotation.MX);
-		MeshComponentHandle = CVDGeometryComponent->AddMeshInstance(OwningParticleTransform, bIsWorldSpace, InExtractedGeometryDataHandle, InOwningParticleData.ParticleIndex, InOwningParticleData.SolverID);
+		MeshComponentHandle = CVDGeometryComponent->AddMeshInstance(OwningParticleTransform, bIsWorldSpace, InExtractedGeometryDataHandle.ToSharedRef(), InOwningParticleData.ParticleIndex, InOwningParticleData.SolverID);
 		
 		MeshComponentHandle->SetGeometryBuilder(AsWeak());
 	}
@@ -475,13 +541,6 @@ void FChaosVDGeometryBuilder::UpdateMeshDataInstance(TSharedPtr<FChaosVDMeshData
 		UE_LOG(LogChaosVDEditor, Error, TEXT("[%s] Failed To Update Mesh Data Handle | Mesh Data Handle is invalid. "), ANSI_TO_TCHAR(__FUNCTION__));
 		return;
 	}
-	
-	TSharedPtr<FChaosVDExtractedGeometryDataHandle> ExtractedGeometryDataHandle = HandleToUpdate->GetGeometryHandle();
-	if (!ExtractedGeometryDataHandle.IsValid())
-	{
-		UE_LOG(LogChaosVDEditor, Error, TEXT("[%s] Failed To Update Mesh Data Handle | Geometry Data Handle is invalid. "), ANSI_TO_TCHAR(__FUNCTION__));
-		return;
-	}
 
 	ComponentType* Component = GetMeshComponentForNewInstance<ComponentType>(HandleToUpdate->GetGeometryHandle(), MeshAttributes);
 
diff --git a/Engine/Plugins/ChaosVD/Source/ChaosVD/Public/ChaosVDPlaybackController.h b/Engine/Plugins/ChaosVD/Source/ChaosVD/Public/ChaosVDPlaybackController.h
--- a/Engine/Plugins/ChaosVD/Source/ChaosVD/Public/ChaosVDPlaybackController.h
+++ b/Engine/Plugins/ChaosVD/Source/ChaosVD/Public/ChaosVDPlaybackController.h
@@ -413,9 +413,6 @@ protected:
 	/** Add the provided track info update to the queue. The update will be broadcast in the game thread */
 	void EnqueueTrackInfoUpdate(const TSharedRef<const FChaosVDTrackInfo>& InTrackInfo, FGuid InstigatorID);
 
-	/** Add the provided Geometry info data to the queue. The update will be broadcast in the game thread */
-	void EnqueueGeometryDataUpdate(const Chaos::FConstImplicitObjectPtr& NewGeometry, const uint32 GeometryID);
-
 	void PlaySolverStageData(int32 TrackID, const TSharedRef<FChaosVDScene>& InSceneToControlSharedPtr, const FChaosVDSolverFrameData& InSolverFrameData, int32 StageIndex);
 
 	template <typename TVisitorCallback>
@@ -449,9 +446,6 @@ protected:
 	/** Queue with a copy of all Track Info Updates that needs to be done in the Game thread */
 	TQueue<FChaosVDQueuedTrackInfoUpdate, EQueueMode::Mpsc> TrackInfoUpdateGTQueue;
 
-	/** Queue with a all the new geometry data that needs to be processed in the Game thread */
-	TQueue<FChaosVDGeometryDataUpdate, EQueueMode::Mpsc> GeometryDataUpdateGTQueue;
-
 	bool bPlayedFirstFrame = false;
 
 	int32 MaxFramesLaggingBehindDuringLiveSession = 50;
diff --git a/Engine/Plugins/ChaosVD/Source/ChaosVD/Public/ChaosVDScene.h b/Engine/Plugins/ChaosVD/Source/ChaosVD/Public/ChaosVDScene.h
--- a/Engine/Plugins/ChaosVD/Source/ChaosVD/Public/ChaosVDScene.h
+++ b/Engine/Plugins/ChaosVD/Source/ChaosVD/Public/ChaosVDScene.h
@@ -7,7 +7,9 @@
 #include "Components/ChaosVDSolverDataComponent.h"
 #include "Containers/UnrealString.h"
 #include "Containers/Map.h"
+#include "Containers/Ticker.h"
 #include "Elements/Framework/TypedElementListFwd.h"
+#include "Elements/Framework/TypedElementSelectionSet.h"
 #include "Elements/Interfaces/TypedElementAssetDataInterface.h"
 #include "UObject/GCObject.h"
 #include "UObject/ObjectMacros.h"
@@ -25,7 +27,6 @@ class AChaosVDSolverInfoActor;
 class AChaosVDSceneCollisionContainer;
 class UChaosVDCoreSettings;
 class FChaosVDGeometryBuilder;
-class AChaosVDParticleActor;
 class FReferenceCollector;
 class UChaosVDSceneQueryDataComponent;
 class UObject;
@@ -34,12 +35,13 @@ class USelection;
 class UTypedElementSelectionSet;
 class UWorld;
 
+struct FChaosVDSceneParticle;
+
 struct FTypedElementHandle;
 
 typedef TMap<int32, AChaosVDSolverInfoActor*> FChaosVDSolverInfoByIDMap;
 
 DECLARE_MULTICAST_DELEGATE(FChaosVDSceneUpdatedDelegate)
-DECLARE_MULTICAST_DELEGATE_OneParam(FChaosVDActorUpdatedDelegate, AChaosVDParticleActor*)
 DECLARE_MULTICAST_DELEGATE_OneParam(FChaosVDOnObjectSelectedDelegate, UObject*)
 DECLARE_MULTICAST_DELEGATE_OneParam(FChaosVDFocusRequestDelegate, FBox)
 DECLARE_MULTICAST_DELEGATE_OneParam(FChaosVDSolverInfoActorCreatedDelegate, AChaosVDSolverInfoActor*)
@@ -57,9 +59,10 @@ enum class EChaosVDSceneCleanUpOptions
 ENUM_CLASS_FLAGS(EChaosVDSceneCleanUpOptions)
 
 /** Recreates a UWorld from a recorded Chaos VD Frame */
-class FChaosVDScene : public FGCObject , public TSharedFromThis<FChaosVDScene>
+class FChaosVDScene : public FGCObject , public TSharedFromThis<FChaosVDScene>, public FTSTickerObjectBase
 {
 public:
+
 	FChaosVDScene();
 	virtual ~FChaosVDScene() override;
 
@@ -78,18 +81,8 @@ public:
 		return SceneUpdatedDelegate;
 	}
 
-	/** Updates, Adds and Remove actors to match the provided Step Data */
-	void UpdateFromRecordedStepData(const int32 SolverID, const FChaosVDFrameStageData& InRecordedStepData, const FChaosVDSolverFrameData& InFrameData);
-
-	void UpdateParticlesCollisionData(const FChaosVDFrameStageData& InRecordedStepData, int32 SolverID);
-	void UpdateJointConstraintsData(const FChaosVDFrameStageData& InRecordedStepData, int32 SolverID);
-
-	/**
-	 * Handles new geometry being loaded and available in the current CVD Recording instance
-	 * @param Geometry Ptr to the new loaded geometry
-	 * @param GeometryID ID of the geometry
-	 */
-	void HandleNewGeometryData(const Chaos::FConstImplicitObjectPtr& Geometry, const uint32 GeometryID);
+	/** Updates, Adds and Remove actors to match the provided Solver Stage Data */
+	void UpdateFromRecordedSolverStageData(const int32 SolverID, const FChaosVDFrameStageData& InRecordedStepData, const FChaosVDSolverFrameData& InFrameData);
 
 	/**
 	 * Handles the Playback switching to a new Game Thread Frame
@@ -136,19 +129,17 @@ public:
 	 * @param GeometryID ID of the geometry
 	 */
 	CHAOSVD_API Chaos::FConstImplicitObjectPtr GetUpdatedGeometry(int32 GeometryID) const;
-
-	/** Adds an object to the queue to be notified when a geometry for the provided ID is generated and ready to use
-	 * @param GeometryID ID of the geometry
-	 * @param ObjectWaitingForGeometry ID of the geometry
-	 */
-	void AddObjectWaitingForGeometry(uint32 GeometryID, IChaosVDGeometryOwnerInterface* ObjectWaitingForGeometry);
 	
 	/** Adds an object to the selection set if it was not selected already, making it selected in practice */
 	void SetSelectedObject(UObject* SelectedObject);
+	void SetSelected(const FTypedElementHandle& InElementHandle);
 
 	/** Evaluates an object and returns true if it is selected */
 	bool IsObjectSelected(const UObject* Object);
 
+	/** Evaluates an object and returns true if it is selected */
+	bool IsSelected(const FTypedElementHandle& InElementHandle) const;
+
 	/** Returns a ptr to the current selection set object */
 	UTypedElementSelectionSet* GetElementSelectionSet() const
 	{
@@ -180,7 +171,7 @@ public:
 	 * @param SolverID ID of the solver owning the Particle
 	 * @param ParticleID ID of the particle
 	 */
-	AChaosVDParticleActor* GetParticleActor(int32 SolverID, int32 ParticleID);
+	TSharedPtr<FChaosVDSceneParticle> GetParticleInstance(int32 SolverID, int32 ParticleID);
 
 	/** Returns a const reference for all Solver Data info actor currently available */
 	const FChaosVDSolverInfoByIDMap& GetSolverInfoActorsMap()
@@ -193,19 +184,6 @@ public:
 	 */
 	AChaosVDSolverInfoActor* GetSolverInfoActor(int32 SolverID);
 
-	/** Is the specified solver from a Server or a Client? (note: currently inferred from the solver name) */
-	bool IsSolverForServer(int32 Solver) const;
-
-	FChaosVDActorUpdatedDelegate& OnActorActiveStateChanged()
-	{
-		return ParticleActorUpdateDelegate;
-	}
-
-	FChaosVDActorUpdatedDelegate& OnActorLabelChanged()
-	{
-		return ParticleLabelUpdateDelegate;
-	}
-
 	/** Event called when a solver info actor is created */
 	FChaosVDSolverInfoActorCreatedDelegate& OnSolverInfoActorCreated()
 	{
@@ -233,18 +211,26 @@ public:
 		return AvailableDataContainerActors;
 	}
 
-	
 	TSharedPtr<FChaosVDRecording> GetLoadedRecording() const
 	{
 		return LoadedRecording;
 	}
 
+	FName GetTEDSSelectionSetName();
+
+	void ClearSelectionAndNotify();
+
+	void RequestUpdate();
+
+	virtual bool Tick(float DeltaTime) override;
+
 private:
+
+	void AddFromCVDWorldTagToActor(AActor* Actor);
+
+	FName TEDSSelectionSetName;
 	
-	void SetLoadedRecording(const TSharedPtr<FChaosVDRecording>& NewRecordingInstance)
-	{
-		LoadedRecording = NewRecordingInstance;
-	}
+	void SetLoadedRecording(const TSharedPtr<FChaosVDRecording>& NewRecordingInstance);
 	
 	TSharedPtr<FChaosVDRecording> LoadedRecording;
 	
@@ -260,13 +246,7 @@ private:
 
 	void PerformGarbageCollection();
 
-	/** Creates an ChaosVDParticle actor for the Provided recorded Particle Data */
-	AChaosVDParticleActor* SpawnParticleFromRecordedData(const TSharedPtr<FChaosVDParticleDataWrapper>& InParticleData, const FChaosVDSolverFrameData& InFrameData);
-
-	/** Returns the ID used to track this recorded particle data */
-	int32 GetIDForRecordedParticleData(const TSharedPtr<FChaosVDParticleDataWrapper>& InParticleData) const;
-
-	void CreateBaseLights(UWorld* TargetWorld) const;
+	void CreateBaseLights(UWorld* TargetWorld);
 
 	void CreatePostProcessingVolumes(UWorld* TargetWorld);
 
@@ -290,8 +270,6 @@ private:
 	void HandleDeSelectElement(const TTypedElement<ITypedElementSelectionInterface>& InElementSelectionHandle, FTypedElementListRef InSelectionSet, const FTypedElementSelectionOptions& InSelectionOptions);
 	void HandleSelectElement(const TTypedElement<ITypedElementSelectionInterface>& InElementSelectionHandle, FTypedElementListRef InSelectionSet, const FTypedElementSelectionOptions& InSelectionOptions);
 
-	void ClearSelectionAndNotify();
-
 	void InitializeSelectionSets();
 	void DeInitializeSelectionSets();
 
@@ -329,23 +307,21 @@ private:
 
 	bool bIsInitialized = false;
 
-	FChaosVDActorUpdatedDelegate ParticleActorUpdateDelegate;
-	FChaosVDActorUpdatedDelegate ParticleLabelUpdateDelegate;
-
 	FDelegateHandle ActorDestroyedHandle;
 
 	FChaosVDSolverInfoActorCreatedDelegate SolverInfoActorCreatedDelegate;
 
 	FChaosVDSolverVisibilityChangedDelegate SolverVisibilityChangedDelegate;
-
-	TMap<uint32, TArray<IChaosVDGeometryOwnerInterface*>> ObjectsWaitingForGeometry;
 	
 	TSharedPtr<FChaosVDSolverDataSelection> SolverDataSelectionObject;
 
 	TArray<TObjectPtr<AChaosVDDataContainerBaseActor>> AvailableDataContainerActors;
 
+	bool bPendingUpdateRequest = false;
+
 	friend class FChaosVDGeometryBuilder;
 	friend class FChaosVDSelectionCustomization;
 	friend class FChaosVDPlaybackViewportClient;
 	friend class FChaosVDPlaybackController;
 };
+
diff --git a/Engine/Plugins/ChaosVD/Source/ChaosVD/Public/Components/ChaosVDSolverDataComponent.h b/Engine/Plugins/ChaosVD/Source/ChaosVD/Public/Components/ChaosVDSolverDataComponent.h
--- a/Engine/Plugins/ChaosVD/Source/ChaosVD/Public/Components/ChaosVDSolverDataComponent.h
+++ b/Engine/Plugins/ChaosVD/Source/ChaosVD/Public/Components/ChaosVDSolverDataComponent.h
@@ -5,6 +5,7 @@
 #include "Components/ActorComponent.h"
 #include "ChaosVDSolverDataComponent.generated.h"
 
+struct FChaosVDFrameStageData;
 struct FChaosVDSolverFrameData;
 class FChaosVDScene;
 struct FChaosVDGameFrameData;
@@ -20,14 +21,27 @@ public:
 	virtual void SetScene(const TWeakPtr<FChaosVDScene>& InSceneWeakPtr);
 
 	virtual void UpdateFromNewGameFrameData(const FChaosVDGameFrameData& InGameFrameData) {}
+	virtual void UpdateFromNewSolverStageData(const FChaosVDSolverFrameData& InSolverFrameData, const FChaosVDFrameStageData& InSolverFrameStageData) {}
 	virtual void UpdateFromSolverFrameData(const FChaosVDSolverFrameData& InSolverFrameData) {}
+
+	bool IsVisible() const
+	{
+		return bIsVisible;
+	}
+
+	virtual void SetVisibility(bool bNewIsVisible);
 	
-	void SetSolverID(int32 InSolverID) { SolverID = InSolverID; }
+	void SetSolverID(int32 InSolverID)
+	{
+		SolverID = InSolverID;
+	}
 
 protected:
 
 	TWeakPtr<FChaosVDScene> SceneWeakPtr;
 
 	int32 SolverID = INDEX_NONE;
+
+	bool bIsVisible = true;
 };
 
