TEDS - Use Context.PushCommand to set the visibility on actors and CVD particles because those functions end up calling into the data storage directly which is illegal inside a query callback

AActor: Add the sync tag to the actor row in TEDS inside SetIsTemporarilyHiddenInEditor()

diff --git a/Engine/Plugins/ChaosVD/Source/ChaosVD/Private/TEDS/ChaosVDDataStorageVisibilityQueries.cpp b/Engine/Plugins/ChaosVD/Source/ChaosVD/Private/TEDS/ChaosVDDataStorageVisibilityQueries.cpp
--- a/Engine/Plugins/ChaosVD/Source/ChaosVD/Private/TEDS/ChaosVDDataStorageVisibilityQueries.cpp
+++ b/Engine/Plugins/ChaosVD/Source/ChaosVD/Private/TEDS/ChaosVDDataStorageVisibilityQueries.cpp
@@ -75,6 +75,25 @@ void UChaosVDDataStorageVisibilityQueries::RegisterVisibilityColumnToParticleQue
 {
 	using namespace UE::Editor::DataStorage::Queries;
 
+	struct FSetParticleVisibilityCommand
+	{
+		void operator()()
+		{
+			if (bIsVisible)
+			{
+				Particle->RemoveHiddenFlag(EChaosVDHideParticleFlags::HiddenBySceneOutliner);
+			}
+			else
+			{
+				Particle->AddHiddenFlag(EChaosVDHideParticleFlags::HiddenBySceneOutliner);
+			}
+			Particle->UpdateGeometryComponentsVisibility();
+			
+		}
+		FChaosVDSceneParticle* Particle;
+		bool bIsVisible;
+	};
+
 	DataStorage.RegisterQuery(
 		Select(
 			TEXT("Particle Visibility Column to Object"),
@@ -85,15 +104,13 @@ void UChaosVDDataStorageVisibilityQueries::RegisterVisibilityColumnToParticleQue
 				{
 					FChaosVDSceneParticle* Particle = static_cast<FChaosVDSceneParticle*>(RawObject.Object);
 
-					if (VisibilityColumn.bIsVisibleInEditor)
-					{
-						Particle->RemoveHiddenFlag(EChaosVDHideParticleFlags::HiddenBySceneOutliner);
-					}
-					else
-					{
-						Particle->AddHiddenFlag(EChaosVDHideParticleFlags::HiddenBySceneOutliner);
-					}
-					Particle->UpdateGeometryComponentsVisibility();
+					Context.PushCommand(FSetParticleVisibilityCommand
+												{
+													.Particle = Particle,
+													.bIsVisible = VisibilityColumn.bIsVisibleInEditor
+												});
+
+					
 				}
 			})
 		.Where()
diff --git a/Engine/Plugins/Experimental/EditorDataStorage/Source/TedsCore/Private/Compatibility/TypedElementActorVisibilityQueries.cpp b/Engine/Plugins/Experimental/EditorDataStorage/Source/TedsCore/Private/Compatibility/TypedElementActorVisibilityQueries.cpp
--- a/Engine/Plugins/Experimental/EditorDataStorage/Source/TedsCore/Private/Compatibility/TypedElementActorVisibilityQueries.cpp
+++ b/Engine/Plugins/Experimental/EditorDataStorage/Source/TedsCore/Private/Compatibility/TypedElementActorVisibilityQueries.cpp
@@ -62,6 +62,21 @@ void UActorVisibilityDataStorageFactory::RegisterVisibilityColumnToActorQuery(IE
 {
 	using namespace UE::Editor::DataStorage::Queries;
 
+	struct FSetActorVisibilityCommand
+	{
+		void operator()()
+		{
+			TStrongObjectPtr<AActor> PinnedActor = Actor.Pin();
+			if (PinnedActor)
+			{
+				PinnedActor->SetIsTemporarilyHiddenInEditor(!bIsVisible);
+			}
+			
+		}
+		TWeakObjectPtr<AActor> Actor;
+		bool bIsVisible;
+	};
+	
 	DataStorage.RegisterQuery(
 		Select(
 			TEXT("Actor Visibility Column to Object"),
@@ -70,7 +85,11 @@ void UActorVisibilityDataStorageFactory::RegisterVisibilityColumnToActorQuery(IE
 			{
 				if (AActor* ActorInstance = Cast<AActor>(Object.Object); ActorInstance != nullptr)
 				{
-					ActorInstance->SetIsTemporarilyHiddenInEditor(!VisibilityColumn.bIsVisibleInEditor);
+					Context.PushCommand(FSetActorVisibilityCommand
+							{
+								.Actor = ActorInstance,
+								.bIsVisible = VisibilityColumn.bIsVisibleInEditor
+							});
 				}
 			})
 		.Where()
diff --git a/Engine/Source/Runtime/Engine/Private/ActorEditor.cpp b/Engine/Source/Runtime/Engine/Private/ActorEditor.cpp
--- a/Engine/Source/Runtime/Engine/Private/ActorEditor.cpp
+++ b/Engine/Source/Runtime/Engine/Private/ActorEditor.cpp
@@ -43,6 +43,10 @@
 #include "Algo/RemoveIf.h"
 #include "Misc/DataValidation.h"
 #include "DeletedObjectPlaceholder.h"
+#include "Elements/Columns/TypedElementMiscColumns.h"
+#include "Elements/Common/EditorDataStorageFeatures.h"
+#include "Elements/Interfaces/TypedElementDataStorageCompatibilityInterface.h"
+#include "Elements/Interfaces/TypedElementDataStorageInterface.h"
 
 #define LOCTEXT_NAMESPACE "ErrorChecking"
 
@@ -983,6 +987,21 @@ void AActor::SetIsTemporarilyHiddenInEditor( bool bIsHidden )
 	{
 		bHiddenEdTemporary = bIsHidden;
 		MarkComponentsRenderStateDirty();
+
+		// When updating the visibility we need to add the sync tag to the row in TEDS to ensure the visibility is synced in TEDS properly
+		using namespace UE::Editor::DataStorage;
+		IEditorDataStorageProvider* Storage = GetMutableDataStorageFeature<IEditorDataStorageProvider>(StorageFeatureName);
+		IEditorDataStorageCompatibilityProvider* StorageCompatibility = GetMutableDataStorageFeature<IEditorDataStorageCompatibilityProvider>(CompatibilityFeatureName);
+
+		if (Storage && StorageCompatibility)
+		{
+			RowHandle ActorRow = StorageCompatibility->FindRowWithCompatibleObject(this);
+
+			if (Storage->IsRowAvailable(ActorRow))
+			{
+				Storage->AddColumn<FTypedElementSyncFromWorldTag>(ActorRow);
+			}
+		}
 	}
 }
 
