[ChaosVD] Improvements to geometry generation and geometry loading

- Deduplicated the number of meshes/component we create for simple shapes like Spheres and Boxes. Now instead of using the implicit object hash to determine if a box or sphere is unique, we use the same box or sphere mesh for any recorded implicit object box or sphere, and just calculate the correct transformation to apply so it matches the expected dimensions.
- Now we batch any operation done to an instanced static mesh component instead of executing it instantly. Now all operations are added to a queue which then is processed at the end of the frame. this also de-duplicated some operations that ended up doing similar executions (like Updating visibility and updating transform operation, can be combined as one).
- Implemented a simple, pseudo level streaming system. It is pseudo in the sense we are not streaming data on and off  disk, but instead now we only load geometry for any particle representation that is within a streaming range from the editor's camera. This greatly reduces the number of meshes and instanced static mesh component that are created every frame.
- Updated several methods to just take a TSharedRef instead of TSharedPtr. These methods were called by other ones that already validated the ptr, therefore we can avoid the extra checks and TSharedPtr to TSharedRef conversions
- Fixed a bug where when scrubbing with delta +1, particles were incorrectly removed if their particle index was reused from a previously destroyed particle.
- Fixed a bug where if you scrub the timeline, multiple playback actions per frame are queued (one per value change of the slider). Now we just play the last seen value every frame.
- Disabled the scene outliner population of we are scrubbing a timeline
- Instead of adding/removing objects from teds during playback when a particle is destroyed/created, we just add/remove a teds tag to reflect an "active state". This is the same we do internally in CVD for everything else. TEDS tags can be add in batches and it is a considerable faster operation.
- Added a query to CVD's outliner mode to intersect the new tag addition/removal, and queue add/remove operation for the respective outliner items. That queue will be processed when CVD is not playing or scrubbing, and even then it has a budget of 4ms to not bring down editor responsiveness.
- Fixed a bug where if a particle was destroyed, and later its id was re-used, the name in the outliner would not be updated.

#virtualized

diff --git a/Engine/Plugins/ChaosVD/Config/Engine.ini b/Engine/Plugins/ChaosVD/Config/Engine.ini
--- a/Engine/Plugins/ChaosVD/Config/Engine.ini
+++ b/Engine/Plugins/ChaosVD/Config/Engine.ini
@@ -4,4 +4,6 @@ SimOnlyMeshesMaterial=/ChaosVD/Materials/ChasoVDSimMesh.ChasoVDSimMesh
 InstancedMeshesMaterial=/ChaosVD/Materials/ChaosVDInstanced.ChaosVDInstanced
 InstancedMeshesQueryOnlyMaterial=/ChaosVD/Materials/ChaosVDInstancedQueryOnly.ChaosVDInstancedQueryOnly
 SkySphereActorClass=/ChaosVD/BP_CVD_SkySphere.BP_CVD_SkySphere_C
-AmbientCubeMapTexture=/ChaosVD/Textures/DaylightAmbientCubemap.DaylightAmbientCubemap
\ No newline at end of file
+AmbientCubeMapTexture=/ChaosVD/Textures/DaylightAmbientCubemap.DaylightAmbientCubemap
+BoxMesh=/ChaosVD/Meshes/Box.Box
+SphereMesh=/ChaosVD/Meshes/Sphere.Sphere
\ No newline at end of file
diff --git a/Engine/Plugins/ChaosVD/Source/ChaosVD/Private/Actors/ChaosVDDataContainerBaseActor.cpp b/Engine/Plugins/ChaosVD/Source/ChaosVD/Private/Actors/ChaosVDDataContainerBaseActor.cpp
--- a/Engine/Plugins/ChaosVD/Source/ChaosVD/Private/Actors/ChaosVDDataContainerBaseActor.cpp
+++ b/Engine/Plugins/ChaosVD/Source/ChaosVD/Private/Actors/ChaosVDDataContainerBaseActor.cpp
@@ -111,6 +111,18 @@ void AChaosVDDataContainerBaseActor::HandlePostInitializationExtensionRegistered
 	}
 }
 
+void AChaosVDDataContainerBaseActor::HandleWorldStreamingLocationUpdated(const FVector& InLocation)
+{
+	TInlineComponentArray<UChaosVDSolverDataComponent*> SolverDataComponents(this);
+	for (UChaosVDSolverDataComponent* Component : SolverDataComponents)
+	{
+		if (Component)
+		{
+			Component->HandleWorldStreamingLocationUpdated(InLocation);
+		}
+	}
+}
+
 void AChaosVDDataContainerBaseActor::SetSolverID(int32 InSolverID)
 {
 	SolverDataID = InSolverID;
diff --git a/Engine/Plugins/ChaosVD/Source/ChaosVD/Private/Actors/ChaosVDGeometryContainer.h b/Engine/Plugins/ChaosVD/Source/ChaosVD/Private/Actors/ChaosVDGeometryContainer.h
--- a/Engine/Plugins/ChaosVD/Source/ChaosVD/Private/Actors/ChaosVDGeometryContainer.h
+++ b/Engine/Plugins/ChaosVD/Source/ChaosVD/Private/Actors/ChaosVDGeometryContainer.h
@@ -4,6 +4,8 @@
 
 #include "GameFramework/Actor.h"
 #include "UObject/ObjectPtr.h"
+#include "ChaosVDScene.h"
+#include "Templates/SharedPointer.h"
 
 #include "ChaosVDGeometryContainer.generated.h"
 
@@ -16,5 +18,18 @@ class CHAOSVD_API AChaosVDGeometryContainer : public AActor
 public:
 	AChaosVDGeometryContainer();
 
+	virtual void SetScene(TWeakPtr<FChaosVDScene> InScene)
+	{
+		SceneWeakPtr = InScene;
+	}
+
+	TSharedPtr<FChaosVDScene> GetScene() const 
+	{
+		return SceneWeakPtr.Pin();
+	}
+
 	void CleanUp();
+
+protected:
+	TWeakPtr<FChaosVDScene> SceneWeakPtr;
 };
\ No newline at end of file
diff --git a/Engine/Plugins/ChaosVD/Source/ChaosVD/Private/ChaosVDBaseSceneObject.h b/Engine/Plugins/ChaosVD/Source/ChaosVD/Private/ChaosVDBaseSceneObject.h
--- a/Engine/Plugins/ChaosVD/Source/ChaosVD/Private/ChaosVDBaseSceneObject.h
+++ b/Engine/Plugins/ChaosVD/Source/ChaosVD/Private/ChaosVDBaseSceneObject.h
@@ -107,6 +107,36 @@ struct FChaosVDBaseSceneObject
 	{
 		CachedRowHandle = Handle;
 	}
+	
+	enum class EStreamingState
+	{
+		Visible,
+		Hidden
+	};
+
+	void SetStreamingState(EStreamingState NewState)
+	{
+		StreamingState = NewState;
+	}
+	
+	EStreamingState GetStreamingState() const
+	{
+		return StreamingState;
+	}
+
+	virtual void SyncStreamingState()
+	{
+	}
+
+	virtual FBox GetStreamingBounds() const
+	{
+		return FBox(EForceInit::ForceInitToZero);
+	}
+
+	virtual int32 GetStreamingID() const
+	{
+		return INDEX_NONE;
+	}
 
 protected:
 	TWeakPtr<FChaosVDBaseSceneObject> Parent = nullptr;
@@ -118,5 +148,7 @@ protected:
 	UE::Editor::DataStorage::RowHandle CachedRowHandle = UE::Editor::DataStorage::InvalidRowHandle;
 	
 	FName IconName = NAME_None;
+
+	EStreamingState StreamingState = EStreamingState::Hidden;
 };
 
diff --git a/Engine/Plugins/ChaosVD/Source/ChaosVD/Private/ChaosVDExtractedGeometryDataHandle.cpp b/Engine/Plugins/ChaosVD/Source/ChaosVD/Private/ChaosVDExtractedGeometryDataHandle.cpp
--- a/Engine/Plugins/ChaosVD/Source/ChaosVD/Private/ChaosVDExtractedGeometryDataHandle.cpp
+++ b/Engine/Plugins/ChaosVD/Source/ChaosVD/Private/ChaosVDExtractedGeometryDataHandle.cpp
@@ -4,6 +4,11 @@
 
 #include "Chaos/ImplicitObject.h"
 
+void FChaosVDExtractedGeometryDataHandle::SetGeometryKey(const uint32 Key)
+{
+	GeometryKey = Key;
+}
+
 FName FChaosVDExtractedGeometryDataHandle::GetTypeName() const
 {
 	using namespace Chaos;
diff --git a/Engine/Plugins/ChaosVD/Source/ChaosVD/Private/ChaosVDExtractedGeometryDataHandle.h b/Engine/Plugins/ChaosVD/Source/ChaosVD/Private/ChaosVDExtractedGeometryDataHandle.h
--- a/Engine/Plugins/ChaosVD/Source/ChaosVD/Private/ChaosVDExtractedGeometryDataHandle.h
+++ b/Engine/Plugins/ChaosVD/Source/ChaosVD/Private/ChaosVDExtractedGeometryDataHandle.h
@@ -20,37 +20,70 @@ namespace Chaos
 struct FChaosVDExtractedGeometryDataHandle
 {
 	/** Key to be used to match this geometry with recorded Shape Data */
-	uint32 GetDataComponentKey() const { return DataComponentKey; }
-	void SetDataComponentKey(const uint32 Key) { DataComponentKey = Key; }
+	uint32 GetDataComponentKey() const
+	{
+		return DataComponentKey;
+	}
+	void SetDataComponentKey(const uint32 Key)
+	{
+		DataComponentKey = Key;
+	}
 
 	/** Key representing the generated static mesh - Handles for different implicit objects can point to the same geometry key, if after being unpacked they result in essentially the same mesh */
-	uint32 GetGeometryKey() const { return GeometryKey; }
-	void SetGeometryKey(const uint32 Key) { GeometryKey = Key; }
+	uint32 GetGeometryKey() const
+	{
+		return GeometryKey;
+	}
+	void SetGeometryKey(const uint32 Key);
 
 	/** Name of the type this underlying implicit object this instance represents */
 	FName GetTypeName() const;
 
 	/** Extracted transform of the generated implicit object in component space */
-	const FTransform& GetRelativeTransform() const { return Transform; }
-	void SetTransform(const FTransform& InTransform) { Transform = InTransform; }
+	const FTransform& GetRelativeTransform() const
+	{
+		return Transform;
+	}
+	void SetTransform(const FTransform& InTransform)
+	{
+		Transform = InTransform;
+	}
 
 	/** Ptr to the implicit object this handle provides access to and from which the mesh was generated */
-	const Chaos::FImplicitObject* GetImplicitObject() const { return ImplicitObject; }
+	const Chaos::FImplicitObject* GetImplicitObject() const
+	{
+		return ImplicitObject;
+	}
 
 	/** Returns the Index to correct the shape data instance in the shape instances array of the particle that owns this geometry */
-	int32 GetShapeInstanceIndex() const { return ShapeInstanceIndex; }
+	int32 GetShapeInstanceIndex() const
+	{
+		return ShapeInstanceIndex;
+	}
 
 	/** If the object was an union and this was a leaf, this getter returns the root union object, otherwise it will be the same as GetImplicitObject */
-	const Chaos::FImplicitObject* GetRootImplicitObject() const { return RootImplicitObject; }
+	const Chaos::FImplicitObject* GetRootImplicitObject() const
+	{
+		return RootImplicitObject;
+	}
 	
 	/** Sets the Implicit Object Ptr from which the data was extracted and the mesh was generated */
-	void SetImplicitObject(const Chaos::FImplicitObject* InImplicitObject) { ImplicitObject = InImplicitObject; }
+	void SetImplicitObject(const Chaos::FImplicitObject* InImplicitObject)
+	{
+		ImplicitObject = InImplicitObject;
+	}
 
 	/** Sets the Index on which this Implicit object was in the Hierarchy of the root one */
-	void SetShapeInstanceIndex(int32 InObjectIndex) { ShapeInstanceIndex = InObjectIndex; }
+	void SetShapeInstanceIndex(int32 InObjectIndex)
+	{
+		ShapeInstanceIndex = InObjectIndex;
+	}
 	
 	/** Sets the Implicit Object Ptr from which the data was extracted and the mesh was generated */
-	void SetRootImplicitObject(const Chaos::FImplicitObject* InImplicitObject) { RootImplicitObject = InImplicitObject; }
+	void SetRootImplicitObject(const Chaos::FImplicitObject* InImplicitObject)
+	{
+		RootImplicitObject = InImplicitObject;
+	}
 
 	bool operator==(const FChaosVDExtractedGeometryDataHandle& Other) const
 	{
diff --git a/Engine/Plugins/ChaosVD/Source/ChaosVD/Private/ChaosVDGeometryBuilder.cpp b/Engine/Plugins/ChaosVD/Source/ChaosVD/Private/ChaosVDGeometryBuilder.cpp
--- a/Engine/Plugins/ChaosVD/Source/ChaosVD/Private/ChaosVDGeometryBuilder.cpp
+++ b/Engine/Plugins/ChaosVD/Source/ChaosVD/Private/ChaosVDGeometryBuilder.cpp
@@ -6,6 +6,7 @@
 #include "ChaosVDConvexMeshGenerator.h"
 #include "ChaosVDGeometryDataComponent.h"
 #include "ChaosVDModule.h"
+#include "ChaosVDSettingsManager.h"
 #include "ChaosVDTriMeshGenerator.h"
 #include "DynamicMeshToMeshDescription.h"
 #include "Generators/CapsuleGenerator.h"
@@ -40,6 +41,12 @@ namespace Chaos::VisualDebugger
 			TEXT("p.Chaos.VD.Tool.GeometryGenerationTaskLaunchBudgetSeconds"),
 			GeometryGenerationTaskLaunchBudgetSeconds,
 			TEXT("How much time we can spend on the Geoemtry builder tick launching Geometry Generation Tasks"));
+	
+		static bool bDeduplicateSimpleGeometry = true;
+		static FAutoConsoleVariableRef CVarDeduplicateSimpleGeometry(
+			TEXT("p.Chaos.VD.Tool.DeduplicateSimpleGeometry"),
+			bDeduplicateSimpleGeometry,
+			TEXT("If set to true, Box and Spheres will be represented with pre-made static meshes with a calculated scale based on the implicit object data"));
 	}
 
 	void SetTriangleAttributes(const UE::Geometry::FMeshShapeGenerator& Generator, FDynamicMesh3& OutDynamicMesh, int32 AppendedTriangleID, int32 GeneratorTriangleIndex)
@@ -168,6 +175,20 @@ FChaosVDGeometryBuilder::~FChaosVDGeometryBuilder()
 	DeInitialize();
 }
 
+void FChaosVDGeometryBuilder::CachePreBuiltMeshes()
+{
+	using namespace Chaos;
+	if (UChaosVDCoreSettings* Settings = FChaosVDSettingsManager::Get().GetSettingsObject<UChaosVDCoreSettings>())
+	{
+		// Note: Sphere and Box are 0 and 1 respectively, which is the key value we would get if we hash them. I expect this to not cause a collision, but if it does, we will need
+		// to create buckets per implicit object type, or we could just create a new instance of this geometry builder per type
+		// CVD already supports multiple builders.
+		// If this unlikely scenario happens, there is an ensure that should trigger in FChaosVDGeometryBuilder::ExtractGeometryDataForImplicit
+		StaticMeshCacheMap.Add(ImplicitObjectType::Box, Settings->BoxMesh.Get());
+		StaticMeshCacheMap.Add(ImplicitObjectType::Sphere, Settings->SphereMesh.Get());
+	}
+}
+
 void FChaosVDGeometryBuilder::Initialize(const TWeakPtr<FChaosVDScene>& ChaosVDScene)
 {
 	if (!ChaosVDScene.Pin())
@@ -215,13 +236,6 @@ void FChaosVDGeometryBuilder::Initialize(const TWeakPtr<FChaosVDScene>& ChaosVDS
 		UMeshComponent* MeshComponent = Object.Get();
 		if (IChaosVDGeometryComponent* CVDMeshComponent = Cast<IChaosVDGeometryComponent>(MeshComponent))
 		{
-			// The Mesh component no longer has instances on it,
-			// this means the component was returned to the pool or is scheduled to be destroyed while we were waiting 
-			if (CVDMeshComponent->GetMeshDataInstanceHandles().IsEmpty())
-			{
-				return true;
-			}
-
 			GeometryBuilder->SetMeshComponentMaterial(CVDMeshComponent);
 		}
 		return true;
@@ -257,6 +271,8 @@ void FChaosVDGeometryBuilder::Initialize(const TWeakPtr<FChaosVDScene>& ChaosVDS
 	constexpr int32 MeshPendingDisposalContainerDefaultSize = 500; 
 	MeshComponentsPendingDisposal.Reserve(MeshPendingDisposalContainerDefaultSize);
 
+	CachePreBuiltMeshes();
+
 	bInitialized = true;
 }
 
@@ -560,7 +576,6 @@ UStaticMesh* FChaosVDGeometryBuilder::CreateAndCacheStaticMesh(const uint32 Geom
 
 	return MainStaticMesh;
 }
-
 void FChaosVDGeometryBuilder::SetMeshComponentMaterial(IChaosVDGeometryComponent* GeometryComponent)
 {
 	if (!GeometryComponent)
@@ -568,6 +583,13 @@ void FChaosVDGeometryBuilder::SetMeshComponentMaterial(IChaosVDGeometryComponent
 		return;
 	}
 
+	// The component could have been set back to the pool before it was processed.
+	// which can happen if a recording is scrub back and forth too fast
+	if (GeometryComponent->GetIsDestroyed())
+	{
+		return;
+	}
+
 	UMaterialInterface* Material = ComponentMeshPool.GetMaterialForType(GeometryComponent->GetMaterialType());
 	ensure(Material);
 
@@ -631,6 +653,9 @@ void FChaosVDGeometryBuilder::DestroyMeshComponent(UMeshComponent* MeshComponent
 		
 		RemoveMeshComponentWaitingForGeometry(AsCVDGeometryComponent->GetGeometryKey(), MeshComponent);	
         AsCVDGeometryComponent->OnComponentEmpty()->RemoveAll(this);
+
+		// Mark destroyed right away to avoid other system using the component by mistake
+		AsCVDGeometryComponent->SetIsDestroyed(true);
 	}
 
 	MeshComponentsPendingDisposal.Add(MeshComponent);
@@ -837,7 +862,9 @@ const Chaos::FImplicitObject* FChaosVDGeometryBuilder::UnpackImplicitObject(cons
 void FChaosVDGeometryBuilder::AdjustedTransformForImplicit(const Chaos::FImplicitObject* InImplicit, FTransform& OutAdjustedTransform)
 {
 	using namespace Chaos;
-	switch (GetInnerType(InImplicit->GetType()))
+	
+	EImplicitObjectType InnerType = GetInnerType(InImplicit->GetType());
+	switch (InnerType)
 	{
 		// Currently, only capsules and spheres transforms needs to be re-adjusted to take into account non-zero center locations
 		case ImplicitObjectType::Capsule:
@@ -853,42 +880,105 @@ void FChaosVDGeometryBuilder::AdjustedTransformForImplicit(const Chaos::FImplici
 			}
 			break;
 		}
-		case ImplicitObjectType::Sphere:
-			{
-				if (const Chaos::FSphere* Sphere = InImplicit->template GetObject<Chaos::FSphere>())
+		
+		default:
+			break;
+	}
+
+	if (Chaos::VisualDebugger::Cvars::bDeduplicateSimpleGeometry)
+	{
+		switch (InnerType)
+		{
+			case ImplicitObjectType::Sphere:
+				{
+					if (const Chaos::FSphere* Sphere = InImplicit->template GetObject<Chaos::FSphere>())
+					{
+						const FVector FinalLocation = OutAdjustedTransform.TransformPosition(FVector(Sphere->GetCenterf()));
+						const FVector FinalScale = OutAdjustedTransform.GetScale3D() * (Sphere->GetRadiusf());
+						OutAdjustedTransform.SetScale3D(FinalScale);
+						OutAdjustedTransform.SetLocation(FinalLocation);
+
+					}
+					break;
+				}
+			case ImplicitObjectType::Box:
 				{
-					const FVector FinalLocation = OutAdjustedTransform.TransformPosition(FVector(Sphere->GetCenterf()));
-					OutAdjustedTransform.SetLocation(FinalLocation);
+					if (const Chaos::TBox<FReal, 3>* Box = InImplicit->template GetObject<Chaos::TBox<FReal, 3>>())
+					{
+						const FVector FinalLocation = OutAdjustedTransform.TransformPosition(FVector(Box->GetCenter()));
+						const FVector FinalScale = OutAdjustedTransform.GetScale3D() * Box->Extents();
+						OutAdjustedTransform.SetLocation(FinalLocation);
+						OutAdjustedTransform.SetScale3D(FinalScale);
+					}
 				}
+			default:
 				break;
-			}
-		default:
-			break;
+		}
+	}
+	else
+	{
+		switch (InnerType)
+		{
+			case ImplicitObjectType::Sphere:
+				{
+					if (const Chaos::FSphere* Sphere = InImplicit->template GetObject<Chaos::FSphere>())
+					{
+						const FVector FinalLocation = OutAdjustedTransform.TransformPosition(FVector(Sphere->GetCenterf()));
+						OutAdjustedTransform.SetLocation(FinalLocation);
+					}
+					break;
+				}
+			default:
+				break;
+		}
 	}
 }
 
 TSharedPtr<FChaosVDExtractedGeometryDataHandle> FChaosVDGeometryBuilder::ExtractGeometryDataForImplicit(const Chaos::FImplicitObject* InImplicitObject, const Chaos::FRigidTransform3& InTransform)
 {
+	const uint32 ImplicitObjectHash = SourceGeometryCache.GetAndCacheGeometryHash(InImplicitObject);
+
 	Chaos::FRigidTransform3 ExtractedTransform = InTransform;
 	const bool bNeedsUnpack = ImplicitObjectNeedsUnpacking(InImplicitObject);
 	if (const Chaos::FImplicitObject* ImplicitObjectToProcess = bNeedsUnpack ? UnpackImplicitObject(InImplicitObject, ExtractedTransform) : InImplicitObject)
 	{
+		Chaos::EImplicitObjectType ImplicitObjectType = Chaos::GetInnerType(ImplicitObjectToProcess->GetType());
 		TSharedPtr<FChaosVDExtractedGeometryDataHandle> MeshDataHandle = MakeShared<FChaosVDExtractedGeometryDataHandle>();
 
-		const uint32 GeometryKey = SourceGeometryCache.GetAndCacheGeometryHash(ImplicitObjectToProcess);		
-		MeshDataHandle->SetGeometryKey(GeometryKey);
+		uint32 GeometryKey = 0;
+		if (UsesPreBuiltGeometry(ImplicitObjectType))
+		{
+			GeometryKey = GetTypeHash(ImplicitObjectType);
 
-		// For the Component data key, we need the hash of the implicit as it is (packed) because we will need to match it when looking for shape data
-		MeshDataHandle->SetDataComponentKey(bNeedsUnpack ? SourceGeometryCache.GetAndCacheGeometryHash(InImplicitObject) : GeometryKey);
+			MeshDataHandle->SetGeometryKey(GeometryKey);
 
-		AdjustedTransformForImplicit(ImplicitObjectToProcess, ExtractedTransform);
-		MeshDataHandle->SetTransform(ExtractedTransform);
+			// For the Component data key, we need the hash of the implicit as it is (packed) because we will need to match it when looking for shape data
+			MeshDataHandle->SetDataComponentKey(bNeedsUnpack ? ImplicitObjectHash : SourceGeometryCache.GetAndCacheGeometryHash(ImplicitObjectToProcess));
+		}
+		else
+		{
+			GeometryKey = SourceGeometryCache.GetAndCacheGeometryHash(ImplicitObjectToProcess);
+
+			if (Chaos::VisualDebugger::Cvars::bDeduplicateSimpleGeometry &&
+				!ensureMsgf(GeometryKey != Chaos::ImplicitObjectType::Box && GeometryKey != Chaos::ImplicitObjectType::Sphere, TEXT("A calculate geometry key is colliding with a deduplicated geometry key (either Box or Sphere). This should not happen. | Geometry key [%u] | Type [%u]"), GeometryKey, ImplicitObjectType))
+			{
+				UE_LOG(LogChaosVDEditor, Verbose, TEXT("Geometry key [%u] | Implicit Object Type [%s]"), GeometryKey, *Chaos::GetImplicitObjectTypeName(ImplicitObjectType).ToString())
+			}
+			
+			MeshDataHandle->SetGeometryKey(GeometryKey);
+
+			// For the Component data key, we need the hash of the implicit as it is (packed) because we will need to match it when looking for shape data
+			MeshDataHandle->SetDataComponentKey(bNeedsUnpack ? ImplicitObjectHash : GeometryKey);
+		}
 
 		if (!HasGeometryInCache(GeometryKey))
 		{
 			DispatchCreateAndCacheMeshForImplicitAsync(GeometryKey, ImplicitObjectToProcess);
 		}
 
+		AdjustedTransformForImplicit(ImplicitObjectToProcess, ExtractedTransform);
+		MeshDataHandle->SetTransform(ExtractedTransform);
+
 		return MeshDataHandle;
 	}
 
@@ -969,22 +1059,23 @@ void FChaosVDGeometryBuilder::RemoveMeshComponentWaitingForGeometry(uint32 Geome
 	MeshComponentsWaitingForGeometry->RemoveObject(GeometryKey, MeshComponent);
 }
 
+void FChaosVDGeometryBuilder::RequestMeshForComponent(const TSharedRef<FChaosVDExtractedGeometryDataHandle>& SourceGeometry, UMeshComponent* MeshComponent)
+{
+	using namespace Chaos;
+
+	AddMeshComponentWaitingForGeometry(SourceGeometry->GetGeometryKey(), MeshComponent);
+}
+
+bool FChaosVDGeometryBuilder::UsesPreBuiltGeometry(Chaos::EImplicitObjectType ObjetType) const
+{
+	return Chaos::VisualDebugger::Cvars::bDeduplicateSimpleGeometry ? ObjetType == Chaos::ImplicitObjectType::Box || ObjetType == Chaos::ImplicitObjectType::Sphere : false;
+}
+
 void FChaosVDGeometryBuilder::HandleStaticMeshComponentInstanceIndexUpdated(UInstancedStaticMeshComponent* InComponent, TArrayView<const FInstancedStaticMeshDelegates::FInstanceIndexUpdateData> InIndexUpdates)
 {
-	if (IChaosVDGeometryComponent* DataComponent = Cast<IChaosVDGeometryComponent>(InComponent))
-	{
-		const TArrayView<TSharedPtr<FChaosVDMeshDataInstanceHandle>> MeshDataHandlesView = DataComponent->GetMeshDataInstanceHandles();
-		for (TSharedPtr<FChaosVDMeshDataInstanceHandle>& Handle : MeshDataHandlesView)
-		{
-			if (Handle)
-			{
-				Handle->HandleInstanceIndexUpdated(InIndexUpdates);
-			}
-			else
-			{
-				UE_LOG(LogChaosVDEditor, Error, TEXT("[%s] Failed to update Instance Index for component [%s] | Handle is in valid"), ANSI_TO_TCHAR(__FUNCTION__), *GetNameSafe(InComponent));
-			}
-		}
+	if (UChaosVDInstancedStaticMeshComponent* CVDInstancedStaticMeshComponent = Cast<UChaosVDInstancedStaticMeshComponent>(InComponent))
+	{	
+		CVDInstancedStaticMeshComponent->HandleInstanceIndexUpdated(InIndexUpdates);
 	}
 }
 
diff --git a/Engine/Plugins/ChaosVD/Source/ChaosVD/Private/ChaosVDGeometryDataComponent.cpp b/Engine/Plugins/ChaosVD/Source/ChaosVD/Private/ChaosVDGeometryDataComponent.cpp
--- a/Engine/Plugins/ChaosVD/Source/ChaosVD/Private/ChaosVDGeometryDataComponent.cpp
+++ b/Engine/Plugins/ChaosVD/Source/ChaosVD/Private/ChaosVDGeometryDataComponent.cpp
@@ -10,14 +10,13 @@
 #include "Components/MeshComponent.h"
 #include "Materials/Material.h"
 #include "Materials/MaterialInterface.h"
-#include "Materials/MaterialInstanceDynamic.h"
 #include "TEDS/ChaosVDStructTypedElementData.h"
 #include "Visualizers/ChaosVDParticleDataComponentVisualizer.h"
 
 TObjectPtr<UChaosVDParticleVisualizationColorSettings> FChaosVDGeometryComponentUtils::CachedParticleColorSettings = nullptr;
 TObjectPtr<UChaosVDParticleVisualizationSettings> FChaosVDGeometryComponentUtils::CachedParticleVisualizationSettings = nullptr;
 
-FChaosVDMeshDataInstanceHandle::FChaosVDMeshDataInstanceHandle(int32 InInstanceIndex, UMeshComponent* InMeshComponent, int32 InParticleID, int32 InSolverID, const TSharedRef<FChaosVDExtractedGeometryDataHandle>& InSourceGeometryHandle)
+FChaosVDInstancedMeshData::FChaosVDInstancedMeshData(int32 InInstanceIndex, UMeshComponent* InMeshComponent, int32 InParticleID, int32 InSolverID, const TSharedRef<FChaosVDExtractedGeometryDataHandle>& InSourceGeometryHandle)
 : ExtractedGeometryHandle(InSourceGeometryHandle)
 {
 	InstanceState.MeshComponent = InMeshComponent;
@@ -43,10 +42,11 @@ FChaosVDMeshDataInstanceHandle::FChaosVDMeshDataInstanceHandle(int32 InInstanceI
 	InstanceState.ImplicitObjectInfo.bIsRootObject = ExtractedGeometryHandle->GetRootImplicitObject() == ExtractedGeometryHandle->GetImplicitObject();
 	InstanceState.ImplicitObjectInfo.ShapeInstanceIndex = ExtractedGeometryHandle->GetShapeInstanceIndex();
 	InstanceState.ImplicitObjectInfo.ImplicitObjectType = ExtractedGeometryHandle->GetTypeName();
+	InstanceState.ImplicitObjectInfo.ImplicitObjectTypeEnum = Chaos::GetInnerType(ExtractedGeometryHandle->GetImplicitObject()->GetType());
 	InstanceState.ImplicitObjectInfo.RelativeTransform = ExtractedGeometryHandle->GetRelativeTransform();	
 }
 
-void FChaosVDMeshDataInstanceHandle::SetWorldTransform(const FTransform& InTransform)
+void FChaosVDInstancedMeshData::SetWorldTransform(const FTransform& InTransform)
 {
 	const FTransform ExtractedRelativeTransform = ExtractedGeometryHandle->GetRelativeTransform();
 
@@ -56,23 +56,23 @@ void FChaosVDMeshDataInstanceHandle::SetWorldTransform(const FTransform& InTrans
 	
 	if (IChaosVDGeometryComponent* CVDGeometryComponent = Cast<IChaosVDGeometryComponent>(GetMeshComponent()))
 	{
-		CVDGeometryComponent->UpdateInstanceWorldTransform(AsShared(), InstanceState.CurrentWorldTransform);
+		CVDGeometryComponent->UpdateWorldTransformForInstance(AsShared());
 	}
 }
 
-void FChaosVDMeshDataInstanceHandle::SetInstanceColor(const FLinearColor& NewColor)
+void FChaosVDInstancedMeshData::SetInstanceColor(const FLinearColor& NewColor)
 {
 	if (InstanceState.CurrentGeometryColor != NewColor)
 	{
 		if (IChaosVDGeometryComponent* CVDGeometryComponent = Cast<IChaosVDGeometryComponent>(GetMeshComponent()))
 		{
-			CVDGeometryComponent->UpdateInstanceColor(AsShared(), NewColor);
 			InstanceState.CurrentGeometryColor = NewColor;
+			CVDGeometryComponent->UpdateColorForInstance(AsShared());
 		}
 	}
 }
 
-void FChaosVDMeshDataInstanceHandle::UpdateMeshComponentForCollisionData(const FChaosVDShapeCollisionData& InCollisionData)
+void FChaosVDInstancedMeshData::UpdateMeshComponentForCollisionData(const FChaosVDShapeCollisionData& InCollisionData)
 {
 	if (InCollisionData.bIsValid && InstanceState.CollisionData != InCollisionData)
 	{
@@ -98,12 +98,7 @@ void FChaosVDMeshDataInstanceHandle::UpdateMeshComponentForCollisionData(const F
 			{
 				if (RequiredMeshAttributes != CVDOldGeometryComponent->GetMeshComponentAttributeFlags())
 				{
-					if (InstanceState.bIsSelected)
-					{
-						CVDOldGeometryComponent->SetIsSelected(AsShared(), false);
-					}
-
-					CVDOldGeometryComponent->RemoveMeshInstance(AsShared());
+					CVDOldGeometryComponent->RemoveMeshInstance(AsShared(), IChaosVDGeometryComponent::ERemovalMode::Instant);
 
 					GeometryBuilderPtr->UpdateMeshDataInstance<UChaosVDInstancedStaticMeshComponent>(AsShared(), RequiredMeshAttributes);
 
@@ -118,15 +113,15 @@ void FChaosVDMeshDataInstanceHandle::UpdateMeshComponentForCollisionData(const F
 					// Reset the color so it is updated in the next Update color calls (which always happens after updating the shape instance data)
 					InstanceState.CurrentGeometryColor = FLinearColor(ForceInitToZero);
 		
-					CVDNewGeometryComponent->UpdateInstanceVisibility(AsShared(), InstanceState.bIsVisible);
-					CVDNewGeometryComponent->SetIsSelected(AsShared(), InstanceState.bIsSelected);
+					CVDNewGeometryComponent->UpdateVisibilityForInstance(AsShared());
+					CVDNewGeometryComponent->UpdateSelectionStateForInstance(AsShared());
 				}
 			}
 		}
 	}
 }
 
-void FChaosVDMeshDataInstanceHandle::SetGeometryCollisionData(const FChaosVDShapeCollisionData& InCollisionData)
+void FChaosVDInstancedMeshData::SetGeometryCollisionData(const FChaosVDShapeCollisionData&& InCollisionData)
 {
 	// If this is a static mesh component, we can't just update change the material. We need to remove this instance from the current component and move it to a
 	// component that has the correct translucent mesh
@@ -138,80 +133,40 @@ void FChaosVDMeshDataInstanceHandle::SetGeometryCollisionData(const FChaosVDShap
 	InstanceState.CollisionData = InCollisionData;
 }
 
-void FChaosVDMeshDataInstanceHandle::SetIsSelected(bool bInIsSelected)
+void FChaosVDInstancedMeshData::SetIsSelected(bool bInIsSelected)
 {
+	InstanceState.bIsSelected = bInIsSelected;
 	if (IChaosVDGeometryComponent* CVDGeometryComponent = Cast<IChaosVDGeometryComponent>(GetMeshComponent()))
 	{
-		CVDGeometryComponent->SetIsSelected(AsShared(), bInIsSelected);
+		CVDGeometryComponent->UpdateSelectionStateForInstance(AsShared());
 	}
-
-	InstanceState.bIsSelected = bInIsSelected;
 }
 
-void FChaosVDMeshDataInstanceHandle::SetVisibility(bool bInIsVisible)
+void FChaosVDInstancedMeshData::SetVisibility(bool bInIsVisible)
 {
 	if (InstanceState.bIsVisible != bInIsVisible)
 	{
-		if (IChaosVDGeometryComponent* CVDGeometryComponent = Cast<IChaosVDGeometryComponent>(GetMeshComponent()))
-		{
-			CVDGeometryComponent->UpdateInstanceVisibility(AsShared(), bInIsVisible);
-		}
-
 		InstanceState.bIsVisible = bInIsVisible;
-	}
-}
 
-void FChaosVDMeshDataInstanceHandle::HandleInstanceIndexUpdated(TArrayView<const FInstancedStaticMeshDelegates::FInstanceIndexUpdateData> InIndexUpdates)
-{
-	// When an index changes, we receive an array with all the indexes that were modified. We need to only act upon the update of for the index we are tracking in this handle
-	for (const FInstancedStaticMeshDelegates::FInstanceIndexUpdateData& IndexUpdateData : InIndexUpdates)
-	{
-		switch (IndexUpdateData.Type)
+		if (IChaosVDGeometryComponent* CVDGeometryComponent = Cast<IChaosVDGeometryComponent>(GetMeshComponent()))
 		{
-			case FInstancedStaticMeshDelegates::EInstanceIndexUpdateType::Added:
-				break; // We don't need to process 'Added' updates as they can't affect existing IDs
-			case FInstancedStaticMeshDelegates::EInstanceIndexUpdateType::Relocated:
-				{
-					if (InstanceState.MeshInstanceIndex == IndexUpdateData.OldIndex)
-					{
-						InstanceState.MeshInstanceIndex = IndexUpdateData.Index;
-					}
-					break;
-				}
-
-			case FInstancedStaticMeshDelegates::EInstanceIndexUpdateType::Removed:
-			case FInstancedStaticMeshDelegates::EInstanceIndexUpdateType::Cleared:
-			case FInstancedStaticMeshDelegates::EInstanceIndexUpdateType::Destroyed:
-				{
-					if (InstanceState.MeshInstanceIndex == IndexUpdateData.Index)
-					{
-						InstanceState.MeshInstanceIndex = INDEX_NONE;
-					}
-					break;
-				}
-			default:
-				break;
+			CVDGeometryComponent->UpdateVisibilityForInstance(AsShared());
 		}
 	}
 }
 
-void FChaosVDGeometryComponentUtils::UpdateCollisionDataFromShapeArray(const TArray<FChaosVDShapeCollisionData>& InShapeArray, const TSharedRef<FChaosVDMeshDataInstanceHandle>& InInstanceHandle)
+void FChaosVDGeometryComponentUtils::UpdateCollisionDataFromShapeArray(TArrayView<FChaosVDShapeCollisionData> InShapeArray, const TSharedRef<FChaosVDInstancedMeshData>& InInstanceHandle)
 {
-	if (InShapeArray.IsEmpty())
-	{
-		return;
-	}
+	FChaosVDMeshDataInstanceState& InstanceState = InInstanceHandle->GetState();
 
-	const TSharedRef<FChaosVDExtractedGeometryDataHandle>& ExtractedGeometryHandle = InInstanceHandle->GetGeometryHandle();
-
-	const int32 ShapeInstanceIndex = ExtractedGeometryHandle->GetShapeInstanceIndex();
-	if (!InShapeArray.IsValidIndex(ShapeInstanceIndex))
+	if (UNLIKELY(!InShapeArray.IsValidIndex(InstanceState.ImplicitObjectInfo.ShapeInstanceIndex)))
 	{
 		const FName ImplicitObjectTypeName = InInstanceHandle->GetState().ImplicitObjectInfo.ImplicitObjectType;
+		const TSharedRef<FChaosVDExtractedGeometryDataHandle>& ExtractedGeometryHandle = InInstanceHandle->GetGeometryHandle();
 		const Chaos::FImplicitObject* RootImplicitObject = ExtractedGeometryHandle->GetRootImplicitObject();
 		const FName RootImplicitObjectTypeName = !InInstanceHandle->GetState().ImplicitObjectInfo.bIsRootObject && RootImplicitObject ? Chaos::GetImplicitObjectTypeName(Chaos::GetInnerType(RootImplicitObject->GetType())) : TEXT("None");
 		
-		FString ErrorMessage = FString::Printf(TEXT("[%s] Failed to find shape instance data at Index [%d] | Particle ID[%d] | Available Shape instance Data Num [%d] | Implicit Type [%s] - Root Implicit Type [%s] | This geometry will be hidden..."), ANSI_TO_TCHAR(__FUNCTION__), ShapeInstanceIndex, InInstanceHandle->GetOwningParticleID(), InShapeArray.Num(), *ImplicitObjectTypeName.ToString(), *RootImplicitObjectTypeName.ToString());
+		FString ErrorMessage = FString::Printf(TEXT("[%s] Failed to find shape instance data at Index [%d] | Particle ID[%d] | Available Shape instance Data Num [%d] | Implicit Type [%s] - Root Implicit Type [%s] | This geometry will be hidden..."), ANSI_TO_TCHAR(__FUNCTION__), InstanceState.ImplicitObjectInfo.ShapeInstanceIndex, InInstanceHandle->GetOwningParticleID(), InShapeArray.Num(), *ImplicitObjectTypeName.ToString(), *RootImplicitObjectTypeName.ToString());
 		
 		UE_LOG(LogChaosVDEditor, Verbose, TEXT("[%s]"), *ErrorMessage);
 
@@ -223,17 +178,17 @@ void FChaosVDGeometryComponentUtils::UpdateCollisionDataFromShapeArray(const TAr
 	else if (InInstanceHandle->bFailedToUpdateShapeInstanceData)
 	{
 		InInstanceHandle->bFailedToUpdateShapeInstanceData = false;
-		UE_LOG(LogChaosVDEditor, Verbose, TEXT("[%s] Recovered from failing to find shape instance data at Index [%d] | Particle ID[%d] | Available Shape instance Data Num [%d] | This geometry will be shown again..."), ANSI_TO_TCHAR(__FUNCTION__), ShapeInstanceIndex, InInstanceHandle->GetOwningParticleID(), InShapeArray.Num());
+		UE_LOG(LogChaosVDEditor, Verbose, TEXT("[%s] Recovered from failing to find shape instance data at Index [%d] | Particle ID[%d] | Available Shape instance Data Num [%d] | This geometry will be shown again..."), ANSI_TO_TCHAR(__FUNCTION__), InstanceState.ImplicitObjectInfo.ShapeInstanceIndex, InInstanceHandle->GetOwningParticleID(), InShapeArray.Num());
 	}
 
-	FChaosVDShapeCollisionData CollisionDataToUpdate = InShapeArray[ShapeInstanceIndex];
-	CollisionDataToUpdate.bIsComplex = FChaosVDGeometryBuilder::DoesImplicitContainType(ExtractedGeometryHandle->GetImplicitObject(), Chaos::ImplicitObjectType::HeightField) || FChaosVDGeometryBuilder::DoesImplicitContainType(ExtractedGeometryHandle->GetImplicitObject(), Chaos::ImplicitObjectType::TriangleMesh);
+	FChaosVDShapeCollisionData CollisionDataToUpdate = InShapeArray[InstanceState.ImplicitObjectInfo.ShapeInstanceIndex];
+	CollisionDataToUpdate.bIsComplex = InstanceState.ImplicitObjectInfo.ImplicitObjectTypeEnum == Chaos::ImplicitObjectType::HeightField || InstanceState.ImplicitObjectInfo.ImplicitObjectTypeEnum == Chaos::ImplicitObjectType::TriangleMesh;
 	CollisionDataToUpdate.bIsValid = true;
 
-	InInstanceHandle->SetGeometryCollisionData(CollisionDataToUpdate);
+	InInstanceHandle->SetGeometryCollisionData(MoveTemp(CollisionDataToUpdate));
 }
 
-void FChaosVDGeometryComponentUtils::UpdateMeshColor(const TSharedPtr<FChaosVDMeshDataInstanceHandle>& InInstanceHandle, const FChaosVDParticleDataWrapper& InParticleData, bool bIsServer)
+void FChaosVDGeometryComponentUtils::UpdateMeshColor(const TSharedRef<FChaosVDInstancedMeshData>& InInstanceHandle, const FChaosVDParticleDataWrapper& InParticleData, bool bIsServer)
 {
 	const FChaosVDShapeCollisionData& ShapeData = InInstanceHandle->GetGeometryCollisionData();
 	const bool bIsQueryOnly = ShapeData.bQueryCollision && !ShapeData.bSimCollision;
@@ -249,13 +204,8 @@ void FChaosVDGeometryComponentUtils::UpdateMeshColor(const TSharedPtr<FChaosVDMe
 	}
 }
 
-void FChaosVDGeometryComponentUtils::UpdateMeshVisibility(const TSharedPtr<FChaosVDMeshDataInstanceHandle>& InInstanceHandle, const FChaosVDParticleDataWrapper& InParticleData, bool bIsActive)
+void FChaosVDGeometryComponentUtils::UpdateMeshVisibility(const TSharedRef<FChaosVDInstancedMeshData>& InInstanceHandle, const FChaosVDParticleDataWrapper& InParticleData, bool bIsActive)
 {
-	if (!InInstanceHandle)
-	{
-		return;
-	}
-
 	if (!bIsActive)
 	{
 		InInstanceHandle->SetVisibility(bIsActive);
@@ -280,7 +230,7 @@ void FChaosVDGeometryComponentUtils::UpdateMeshVisibility(const TSharedPtr<FChao
 		// TODO: Re-visit the way we determine visibility of the meshes.
 		// Now that the options have grown and they will continue to do so, these checks are becoming hard to read and extend
 
-		const bool bIsHeightfield = InInstanceHandle->GetGeometryHandle()->GetImplicitObject() && Chaos::GetInnerType(InInstanceHandle->GetGeometryHandle()->GetImplicitObject()->GetType()) == Chaos::ImplicitObjectType::HeightField;
+		const bool bIsHeightfield = InInstanceHandle->GetState().ImplicitObjectInfo.ImplicitObjectTypeEnum == Chaos::ImplicitObjectType::HeightField;
 
 		if (bIsHeightfield && EnumHasAnyFlags(CurrentVisibilityFlags, EChaosVDGeometryVisibilityFlags::ShowHeightfields))
 		{
@@ -308,16 +258,11 @@ void FChaosVDGeometryComponentUtils::UpdateMeshVisibility(const TSharedPtr<FChao
 	}
 }
 
-FLinearColor FChaosVDGeometryComponentUtils::GetGeometryParticleColor(const TSharedPtr<FChaosVDExtractedGeometryDataHandle>& InGeometryHandle, const FChaosVDParticleDataWrapper& InParticleData, bool bIsServer)
+FLinearColor FChaosVDGeometryComponentUtils::GetGeometryParticleColor(const TSharedRef<FChaosVDExtractedGeometryDataHandle>& InGeometryHandle, const FChaosVDParticleDataWrapper& InParticleData, bool bIsServer)
 {
 	constexpr FLinearColor DefaultColor(0.088542f, 0.088542f, 0.088542f);
 	FLinearColor ColorToApply = DefaultColor;
 
-	if (!InGeometryHandle)
-	{
-		return ColorToApply;
-	}
-
 	const UChaosVDParticleVisualizationColorSettings* VisualizationSettings = GetParticleColorSettings();
 	if (!VisualizationSettings)
 	{
@@ -412,7 +357,7 @@ UMaterialInterface* FChaosVDGeometryComponentUtils::GetBaseMaterialForType(EChao
 	}	
 }
 
-void Chaos::VisualDebugger::SelectParticleWithGeometryInstance(const TSharedRef<FChaosVDScene>& InScene, FChaosVDSceneParticle* Particle, const TSharedPtr<FChaosVDMeshDataInstanceHandle>& InMeshDataHandle)
+void Chaos::VisualDebugger::SelectParticleWithGeometryInstance(const TSharedRef<FChaosVDScene>& InScene, FChaosVDSceneParticle* Particle, const TSharedPtr<FChaosVDInstancedMeshData>& InMeshDataHandle)
 {
 	InScene->SetSelectedObject(nullptr);
 
diff --git a/Engine/Plugins/ChaosVD/Source/ChaosVD/Private/ChaosVDGeometryDataComponent.h b/Engine/Plugins/ChaosVD/Source/ChaosVD/Private/ChaosVDGeometryDataComponent.h
--- a/Engine/Plugins/ChaosVD/Source/ChaosVD/Private/ChaosVDGeometryDataComponent.h
+++ b/Engine/Plugins/ChaosVD/Source/ChaosVD/Private/ChaosVDGeometryDataComponent.h
@@ -2,17 +2,14 @@
 
 #pragma once
 #include "ChaosVDExtractedGeometryDataHandle.h"
-#include "InstancedStaticMeshDelegates.h"
-#include "Chaos/ImplicitFwd.h"
-#include "Containers/Array.h"
-#include "Containers/Map.h"
 #include "Chaos/ImplicitObject.h"
+#include "Components/MeshComponent.h"
 #include "DataWrappers/ChaosVDCollisionDataWrappers.h"
 #include "HAL/Platform.h"
+#include "InstancedStaticMeshDelegates.h"
+#include "Templates/SharedPointer.h"
 #include "UObject/Interface.h"
 #include "UObject/StrongObjectPtr.h"
-#include "Templates/SharedPointer.h"
-#include "Components/MeshComponent.h"
 
 #include "ChaosVDGeometryDataComponent.generated.h"
 
@@ -65,6 +62,8 @@ struct FChaosVDImplicitObjectBasicView
 	UPROPERTY(VisibleAnywhere, Category="Recorded GeometryData")
 	FName ImplicitObjectType;
 
+	Chaos::EImplicitObjectType ImplicitObjectTypeEnum = Chaos::ImplicitObjectType::Unknown;
+
 	/** Index of the Shape Instance data for this geometry in the Shape Instance data array */
 	UPROPERTY(VisibleAnywhere, Category="Recorded GeometryData")
 	int32 ShapeInstanceIndex = INDEX_NONE;
@@ -130,17 +129,23 @@ struct FChaosVDMeshDataInstanceState
 	bool bShowCVDDebugData = false;
 };
 
-/** Handle that provides access to a specific mesh instance on a CVD Mesh component (instanced or static)*/
-class FChaosVDMeshDataInstanceHandle : public TSharedFromThis<FChaosVDMeshDataInstanceHandle>
+/** Structure that represents a specific mesh instance on a CVD Mesh component (instanced or static) */
+class FChaosVDInstancedMeshData : public TSharedFromThis<FChaosVDInstancedMeshData>
 {
 public:
-	explicit FChaosVDMeshDataInstanceHandle(int32 InInstanceIndex, UMeshComponent* InMeshComponent, int32 InParticleID, int32 InSolverID, const TSharedRef<FChaosVDExtractedGeometryDataHandle>& InSourceGeometryHandle);
+	explicit FChaosVDInstancedMeshData(int32 InInstanceIndex, UMeshComponent* InMeshComponent, int32 InParticleID, int32 InSolverID, const TSharedRef<FChaosVDExtractedGeometryDataHandle>& InSourceGeometryHandle);
 
 	/** Returns the Particle ID of the particle owning this mesh instance */
-	int32 GetOwningParticleID() const { return InstanceState.OwningParticleID; }
+	int32 GetOwningParticleID() const
+	{
+		return InstanceState.OwningParticleID;
+	}
 
 	/** Returns the Solver ID of the particle owning this mesh instance */
-	int32 GetOwningSolverID() const { return InstanceState.OwningSolverID; }
+	int32 GetOwningSolverID() const
+	{
+		return InstanceState.OwningSolverID;
+	}
 
 	/** Applies the provided world transform to the mesh instance this handle represents
 	 * @param InTransform Transform to apply
@@ -148,32 +153,56 @@ public:
 	void SetWorldTransform(const FTransform& InTransform);
 
 	/** Returns the world transform of the mesh instance this handle represents */
-	const FTransform& GetWorldTransform() const { return InstanceState.CurrentWorldTransform; }
+	const FTransform& GetWorldTransform() const
+	{
+		return InstanceState.CurrentWorldTransform;
+	}
 
 	/** Returns the geometry handle used to create the mesh instance this handle represents */
-	const TSharedRef<FChaosVDExtractedGeometryDataHandle>& GetGeometryHandle() const { return ExtractedGeometryHandle; }
+	const TSharedRef<FChaosVDExtractedGeometryDataHandle>& GetGeometryHandle() const
+	{
+		return ExtractedGeometryHandle;
+	}
 
 	/** Applies to the provided color to the mesh instance this handle represents */
 	void SetInstanceColor(const FLinearColor& NewColor);
 
 	/** Returns the current color of the mesh instance this handle represents */
-	FLinearColor GetInstanceColor() const { return InstanceState.CurrentGeometryColor; }
+	FLinearColor GetInstanceColor() const
+	{
+		return InstanceState.CurrentGeometryColor;
+	}
 	
 	/** Applies to the provided shape collision data to the mesh instance this handle represents */
 	void UpdateMeshComponentForCollisionData(const FChaosVDShapeCollisionData& InCollisionData);
 
 	/** Returns the mesh component used to render the mesh instance this handle represents */
-	UMeshComponent* GetMeshComponent() const { return InstanceState.MeshComponent; }
+	UMeshComponent* GetMeshComponent() const
+	{
+		return InstanceState.MeshComponent;
+	}
 
 	/** Returns the instance index of the mesh instance this handle represents */
-	int32 GetMeshInstanceIndex() const { return InstanceState.MeshInstanceIndex; }
+	int32 GetMeshInstanceIndex() const
+	{
+		return InstanceState.MeshInstanceIndex;
+	}
 
 	/** Returns the type of the component used to render the mesh instance this handle represents */
-	EChaosVDMeshComponent GetMeshComponentType() const { return InstanceState.MeshComponentType; }
+	EChaosVDMeshComponent GetMeshComponentType() const
+	{
+		return InstanceState.MeshComponentType;
+	}
 
 	/** Sets a Ptr to the geometry builder used to generate and manage the geometry/mesh components this handle represents */
-	void SetGeometryBuilder(const TWeakPtr<FChaosVDGeometryBuilder>& InGeometryBuilder) { GeometryBuilderInstance = InGeometryBuilder; }
-	TWeakPtr<FChaosVDGeometryBuilder> GetGeometryBuilder() { return GeometryBuilderInstance; }
+	void SetGeometryBuilder(const TWeakPtr<FChaosVDGeometryBuilder>& InGeometryBuilder)
+	{
+		GeometryBuilderInstance = InGeometryBuilder;
+	}
+	TWeakPtr<FChaosVDGeometryBuilder> GetGeometryBuilder()
+	{
+		return GeometryBuilderInstance;
+	}
 
 	/** Marks this mesh instance as selected. Used to handle Selection in Editor */
 	void SetIsSelected(bool bInIsSelected);
@@ -182,32 +211,58 @@ public:
 	void SetVisibility(bool bInIsVisible);
 
 	/** Returns the current visibility state this mesh instance */
-	bool GetVisibility() const { return InstanceState.bIsVisible; }
+	bool GetVisibility() const
+	{
+		return InstanceState.bIsVisible;
+	}
 
 	/** Applies a new shape collision data to this mesh instance */
-	void SetGeometryCollisionData(const FChaosVDShapeCollisionData& InCollisionData);
+	void SetGeometryCollisionData(const FChaosVDShapeCollisionData&& InCollisionData);
 
 	/** Returns the current shape collision data to this mesh instance */
-	FChaosVDShapeCollisionData& GetGeometryCollisionData() { return InstanceState.CollisionData; }
-
-	/** Handles a mesh instance index update reported by the mesh component used to render this mesh instance */
-	void HandleInstanceIndexUpdated(TArrayView<const FInstancedStaticMeshDelegates::FInstanceIndexUpdateData> InIndexUpdates);
+	FChaosVDShapeCollisionData& GetGeometryCollisionData()
+	{
+		return InstanceState.CollisionData;
+	}
 
-	FChaosVDMeshDataInstanceState& GetState() { return InstanceState; }
+	FChaosVDMeshDataInstanceState& GetState()
+	{
+		return InstanceState;
+	}
 
 	/** Used only for debugging purposes - It will be set to true if we received new Shape Instance data but the Shape Index for the implicit object we represent is not valid */
-	bool bFailedToUpdateShapeInstanceData = false; 
+	bool bFailedToUpdateShapeInstanceData = false;
+	
+	/** Sets the mesh instance index of the mesh instance this handle represents */
+	void SetMeshInstanceIndex(int32 NewIndex)
+	{
+		InstanceState.MeshInstanceIndex = NewIndex;
+	}
+
+	/** Returns true if this instance is queued to be destroyed */
+	bool IsPendingDestroy() const
+	{
+		return bIsPendingDestroy;
+	}
+
+	/** Marks a instance that is queued to be destroyed at the end of the frame */
+	void MarkPendingDestroy()
+	{
+		bIsPendingDestroy = true;
+	}
 
 private:
 
 	/** Sets the mesh component used to render the mesh instance this handle represents */
-	void SetMeshComponent(UMeshComponent* NewComponent) { InstanceState.MeshComponent = NewComponent; }
-
-	/** Sets the mesh instance index of the mesh instance this handle represents */
-	void SetMeshInstanceIndex(int32 NewIndex) { InstanceState.MeshInstanceIndex = NewIndex; }
+	void SetMeshComponent(UMeshComponent* NewComponent)
+	{
+		InstanceState.MeshComponent = NewComponent;
+	}
 
 	FChaosVDMeshDataInstanceState InstanceState;
 
+	bool bIsPendingDestroy = false;
+
 	TSharedRef<FChaosVDExtractedGeometryDataHandle> ExtractedGeometryHandle;
 
 	TWeakPtr<FChaosVDGeometryBuilder> GeometryBuilderInstance = nullptr;
@@ -245,10 +300,7 @@ public:
 	virtual uint32 GetGeometryKey() const PURE_VIRTUAL(IChaosVDGeometryDataComponent::GetGeometryKey, return 0;);
 
 	/** Returns the CVD Mesh Data Instance handle for the provided Instance index */
-	virtual TSharedPtr<FChaosVDMeshDataInstanceHandle> GetMeshDataInstanceHandle(int32 InstanceIndex) const PURE_VIRTUAL(IChaosVDGeometryDataComponent::GetMeshDataInstanceHandle, return nullptr;);
-
-	/** Returns all the CVD Mesh Data Instance handles this component is rendering */
-	virtual TArrayView<TSharedPtr<FChaosVDMeshDataInstanceHandle>> GetMeshDataInstanceHandles() PURE_VIRTUAL(IChaosVDGeometryDataComponent::GetMeshDataInstanceHandles, return TArrayView<TSharedPtr<FChaosVDMeshDataInstanceHandle>>(););
+	virtual TSharedPtr<FChaosVDInstancedMeshData> GetMeshDataInstanceHandle(int32 InstanceIndex) const PURE_VIRTUAL(IChaosVDGeometryDataComponent::GetMeshDataInstanceHandle, return nullptr;);
 
 	/**
 	 * Add a new instance to this mesh component
@@ -259,24 +311,25 @@ public:
 	 * @param SolverID Solver ID of the particle owning this geometry
 	 * @return CVD Mesh instance handle that provides access to this component and specific instance, allowing manipulation of it
 	 */
-	virtual TSharedPtr<FChaosVDMeshDataInstanceHandle> AddMeshInstance(const FTransform InstanceTransform, bool bIsWorldSpace, const TSharedRef<FChaosVDExtractedGeometryDataHandle>& InGeometryHandle, int32 ParticleID, int32 SolverID) PURE_VIRTUAL(IChaosVDGeometryDataComponent::AddMeshInstance, return nullptr;);
+	virtual TSharedPtr<FChaosVDInstancedMeshData> AddMeshInstance(const FTransform InstanceTransform, bool bIsWorldSpace, const TSharedRef<FChaosVDExtractedGeometryDataHandle>& InGeometryHandle, int32 ParticleID, int32 SolverID) PURE_VIRTUAL(IChaosVDGeometryDataComponent::AddMeshInstance, return nullptr;);
 
 	/**
 	 * Adds a new instance to this mesh component, but using an existing Mesh Data Handle instead of creating a new one
-	 * @param MeshDataHandle Handle that will provide access to this mesh instance
-	 * @param InstanceTransform The transform the new instance will use
-	 * @param bIsWorldSpace True if the provide transform is in world space
-	 * @param InGeometryHandle The CVD Geometry Handle with data about the geometry instantiate
-	 * @param ParticleID Particle ID owning this geometry
-	 * @param SolverID Solver ID of the particle owning this geometry
+	 * @param InMeshDataHandle Handle that will provide access to this mesh instance
 	 */
-	virtual void AddMeshInstanceForHandle(TSharedPtr<FChaosVDMeshDataInstanceHandle> MeshDataHandle, const FTransform InstanceTransform, bool bIsWorldSpace, const TSharedRef<FChaosVDExtractedGeometryDataHandle>& InGeometryHandle, int32 ParticleID, int32 SolverID) PURE_VIRTUAL(IChaosVDGeometryDataComponent::AddMeshInstanceForHandle);
+	virtual void AddExistingMeshInstance(const TSharedRef<FChaosVDInstancedMeshData>& InMeshDataHandle) PURE_VIRTUAL(IChaosVDGeometryDataComponent::AddExistingMeshInstance);
+
+	enum class ERemovalMode
+	{
+		Instant,
+		Deferred
+	};
 
 	/**
 	 * Removes the instance the provided handle represents
 	 * @param InHandleToRemove 
 	 */
-	virtual void RemoveMeshInstance(TSharedPtr<FChaosVDMeshDataInstanceHandle> InHandleToRemove) PURE_VIRTUAL(IChaosVDGeometryDataComponent::RemoveMeshInstance);
+	virtual void RemoveMeshInstance(const TSharedRef<FChaosVDInstancedMeshData>& InHandleToRemove, ERemovalMode Mode = ERemovalMode::Deferred) PURE_VIRTUAL(IChaosVDGeometryDataComponent::RemoveMeshInstance);
 
 	/** True if the mesh this component represents is ready for use */
 	virtual bool IsMeshReady() const  PURE_VIRTUAL(IChaosVDGeometryDataComponent::IsMeshReady, return false;);
@@ -290,17 +343,17 @@ public:
 	/** Triggers when the component does not have any instance to render. Used to allow it to return to the mesh component tool for future re-use*/
 	virtual FChaosVDMeshComponentEmptyDelegate* OnComponentEmpty() PURE_VIRTUAL(IChaosVDGeometryDataComponent::OnComponentEmpty, return nullptr;);
 
-	/** Updates the visibility of this component based on the stored CVD data*/
-	virtual void UpdateInstanceVisibility(const TSharedPtr<FChaosVDMeshDataInstanceHandle>& InInstanceHandle, bool bIsVisible) PURE_VIRTUAL(IChaosVDGeometryDataComponent::UpdateInstanceVisibility);
+	/** Updates the visibility of this component based on the stored CVD data */
+	virtual void UpdateVisibilityForInstance(const TSharedRef<FChaosVDInstancedMeshData>& InInstanceHandle) PURE_VIRTUAL(IChaosVDGeometryDataComponent::UpdateVisibilityForInstance);
 
 	/** Changes the selection state of the provided instance - Used for Selection in Editor */
-	virtual void SetIsSelected(const TSharedPtr<FChaosVDMeshDataInstanceHandle>& InInstanceHandle, bool bIsSelected) PURE_VIRTUAL(IChaosVDGeometryDataComponent::SetIsSelected);
+	virtual void UpdateSelectionStateForInstance(const TSharedRef<FChaosVDInstancedMeshData>& InInstanceHandle) PURE_VIRTUAL(IChaosVDGeometryDataComponent::UpdateSelectionStateForInstance);
 
 	/** Updates the colors of this component based on the stored CVD data */
-	virtual void UpdateInstanceColor(const TSharedPtr<FChaosVDMeshDataInstanceHandle>& InInstanceHandle, FLinearColor NewColor) PURE_VIRTUAL(IChaosVDGeometryDataComponent::UpdateInstanceColor);
+	virtual void UpdateColorForInstance(const TSharedRef<FChaosVDInstancedMeshData>& InInstanceHandle) PURE_VIRTUAL(IChaosVDGeometryDataComponent::UpdateColorForInstance);
 
 	/** Updates the colors of this component based on the stored CVD data */
-	virtual void UpdateInstanceWorldTransform(const TSharedPtr<FChaosVDMeshDataInstanceHandle>& InInstanceHandle, const FTransform& InTransform) PURE_VIRTUAL(IChaosVDGeometryDataComponent::UpdateInstanceWorldTransform);
+	virtual void UpdateWorldTransformForInstance(const TSharedRef<FChaosVDInstancedMeshData>& InInstanceHandle) PURE_VIRTUAL(IChaosVDGeometryDataComponent::UpdateWorldTransformForInstance);
 
 	/** Sets the CVD Mesh Attribute flags this component is compatible with*/
 	virtual void SetMeshComponentAttributeFlags(EChaosVDMeshAttributesFlags Flags) PURE_VIRTUAL(IChaosVDGeometryDataComponent::SetMeshComponentAttributeFlags);
@@ -317,19 +370,23 @@ public:
 	virtual void SetGeometryBuilder(TWeakPtr<FChaosVDGeometryBuilder> GeometryBuilder) PURE_VIRTUAL(IChaosVDGeometryDataComponent::SetGeometryBuilder);
 
 	virtual EChaosVDMaterialType GetMaterialType() const PURE_VIRTUAL(IChaosVDGeometryDataComponent::GetMaterialType,  return EChaosVDMaterialType::SMOpaque;);
+
+	virtual bool GetIsDestroyed() const PURE_VIRTUAL(IChaosVDGeometryDataComponent::GetIsDestroyed,  return false;);
+
+	virtual void SetIsDestroyed(bool bNewIsDestroyed) PURE_VIRTUAL(IChaosVDGeometryDataComponent::SetIsDestroyed);
 };
 
 class FChaosVDGeometryComponentUtils
 {
 public:
 	/** Finds and updates the Shape data using the provided array as source*/
-	static void UpdateCollisionDataFromShapeArray(const TArray<FChaosVDShapeCollisionData>& InShapeArray, const TSharedRef<FChaosVDMeshDataInstanceHandle>& InInstanceHandle);
+	static void UpdateCollisionDataFromShapeArray(TArrayView<FChaosVDShapeCollisionData> InShapeArray, const TSharedRef<FChaosVDInstancedMeshData>& InInstanceHandle);
 
 	/** Calculates and updates the color used to render the mesh represented by the provided handle, based on the particle state */
-	static void UpdateMeshColor(const TSharedPtr<FChaosVDMeshDataInstanceHandle>& InInstanceHandle, const FChaosVDParticleDataWrapper& InParticleData, bool bIsServer);
+	static void UpdateMeshColor(const TSharedRef<FChaosVDInstancedMeshData>& InInstanceHandle, const FChaosVDParticleDataWrapper& InParticleData, bool bIsServer);
 
 	/** Calculates the correct visibility state based on the particle state, and applies it to the mesh instance the provided handle represents */
-	static void UpdateMeshVisibility(const TSharedPtr<FChaosVDMeshDataInstanceHandle>& InInstanceHandle, const FChaosVDParticleDataWrapper& InParticleData, bool bIsActive);
+	static void UpdateMeshVisibility(const TSharedRef<FChaosVDInstancedMeshData>& InInstanceHandle, const FChaosVDParticleDataWrapper& InParticleData, bool bIsActive);
 
 	/** Returns the material to use as a base to create material instances for the provided type */
 	static UMaterialInterface* GetBaseMaterialForType(EChaosVDMaterialType Type);
@@ -340,7 +397,7 @@ public:
 
 private:
 	/** Returns the color that needs to be used to present the provided particle data based on its state and current selected options */
-	static FLinearColor GetGeometryParticleColor(const TSharedPtr<FChaosVDExtractedGeometryDataHandle>& InGeometryHandle, const FChaosVDParticleDataWrapper& InParticleData, bool bIsServer);
+	static FLinearColor GetGeometryParticleColor(const TSharedRef<FChaosVDExtractedGeometryDataHandle>& InGeometryHandle, const FChaosVDParticleDataWrapper& InParticleData, bool bIsServer);
 
 	static const UChaosVDParticleVisualizationColorSettings* GetParticleColorSettings();
 	static const UChaosVDParticleVisualizationSettings* GetParticleVisualizationSettings();
@@ -349,7 +406,6 @@ private:
 	static TObjectPtr<UChaosVDParticleVisualizationSettings> CachedParticleVisualizationSettings;
 };
 
-
 template <typename TComponent>
 EChaosVDMaterialType FChaosVDGeometryComponentUtils::GetMaterialTypeForComponent(EChaosVDMeshAttributesFlags MeshAttributes)
 {
@@ -366,5 +422,5 @@ EChaosVDMaterialType FChaosVDGeometryComponentUtils::GetMaterialTypeForComponent
 
 namespace Chaos::VisualDebugger
 {
-	void SelectParticleWithGeometryInstance(const TSharedRef<FChaosVDScene>& InScene, FChaosVDSceneParticle* Particle, const TSharedPtr<FChaosVDMeshDataInstanceHandle>& InMeshDataHandle);
+	void SelectParticleWithGeometryInstance(const TSharedRef<FChaosVDScene>& InScene, FChaosVDSceneParticle* Particle, const TSharedPtr<FChaosVDInstancedMeshData>& InMeshDataHandle);
 }
diff --git a/Engine/Plugins/ChaosVD/Source/ChaosVD/Private/ChaosVDPlaybackController.cpp b/Engine/Plugins/ChaosVD/Source/ChaosVD/Private/ChaosVDPlaybackController.cpp
--- a/Engine/Plugins/ChaosVD/Source/ChaosVD/Private/ChaosVDPlaybackController.cpp
+++ b/Engine/Plugins/ChaosVD/Source/ChaosVD/Private/ChaosVDPlaybackController.cpp
@@ -1031,17 +1031,17 @@ void FChaosVDPlaybackController::TickPlayback(float DeltaTime)
 
 	CurrentPlaybackTime += DeltaTime;
 
-	const float CurrentFrameTime = GetFrameTimeForTrack(ActiveTrack->TrackType, ActiveTrack->TrackID, ActiveTrack);
+	const float TargetFrameTime = GetFrameTimeForTrack(ActiveTrack->TrackType, ActiveTrack->TrackID, ActiveTrack);
 
 	const bool bIsGameTrack = ActiveTrack->TrackType == EChaosVDTrackType::Game;
 
-	while (CurrentPlaybackTime > CurrentFrameTime)
-	{
-		CurrentPlaybackTime -= CurrentFrameTime;
-		const int32 NextFrame = ActiveTrack->CurrentFrame + 1;
-		const int32 NextStageNumber =  bIsGameTrack ? 0 : GetTrackLastStageAtFrame_AssumesLocked(ActiveTrack->TrackType, ActiveTrack->TrackID, NextFrame);
-		GoToTrackFrameAndSync(PlaybackSelfInstigatorID, ActiveTrack->TrackType, ActiveTrack->TrackID, NextFrame, NextStageNumber);
-	}
+	const int32 FramesToAdvance = static_cast<int32>(FMath::Floor(CurrentPlaybackTime / TargetFrameTime));
+
+	CurrentPlaybackTime -= TargetFrameTime * static_cast<float>(FramesToAdvance);
+
+	const int32 NextFrame = FMath::Clamp(ActiveTrack->CurrentFrame + FramesToAdvance, 0, ActiveTrack->MaxFrames);
+	const int32 NextStageNumber =  bIsGameTrack ? 0 : GetTrackLastStageAtFrame_AssumesLocked(ActiveTrack->TrackType, ActiveTrack->TrackID, NextFrame);
+	GoToTrackFrameAndSync(PlaybackSelfInstigatorID, ActiveTrack->TrackType, ActiveTrack->TrackID, NextFrame, NextStageNumber);
 }
 
 void FChaosVDPlaybackController::GetTracksByType(EChaosVDTrackType Type, TArray<TSharedPtr<FChaosVDTrackInfo>>& OutTracks)
@@ -1138,6 +1138,11 @@ bool FChaosVDPlaybackController::IsPlaying() const
 	return ActiveTrack->bIsPlaying;
 }
 
+void FChaosVDPlaybackController::SetScrubbingTimeline(bool bNewIsScrubbingTimeline)
+{
+	bIsScrubbingTimeline = bNewIsScrubbingTimeline;
+}
+
 void FChaosVDPlaybackController::TrySetActiveTrack(int32 SlotIndex)
 {
 	VisitAvailableTracks([this, SlotIndex](const TSharedPtr<FChaosVDTrackInfo>& TrackInfo)
diff --git a/Engine/Plugins/ChaosVD/Source/ChaosVD/Private/ChaosVDPlaybackViewportClient.cpp b/Engine/Plugins/ChaosVD/Source/ChaosVD/Private/ChaosVDPlaybackViewportClient.cpp
--- a/Engine/Plugins/ChaosVD/Source/ChaosVD/Private/ChaosVDPlaybackViewportClient.cpp
+++ b/Engine/Plugins/ChaosVD/Source/ChaosVD/Private/ChaosVDPlaybackViewportClient.cpp
@@ -106,7 +106,7 @@ void FChaosVDPlaybackViewportClient::ProcessClick(FSceneView& View, HHitProxy* H
 
 		if (AsCVDGeometryComponent && MeshInstanceIndex != INDEX_NONE)
 		{
-			if (const TSharedPtr<FChaosVDMeshDataInstanceHandle> MeshDataHandle = AsCVDGeometryComponent->GetMeshDataInstanceHandle(MeshInstanceIndex))
+			if (const TSharedPtr<FChaosVDInstancedMeshData> MeshDataHandle = AsCVDGeometryComponent->GetMeshDataInstanceHandle(MeshInstanceIndex))
 			{
 				if (TSharedPtr<FChaosVDSceneParticle> ClickedActor = ScenePtr->GetParticleInstance(MeshDataHandle->GetOwningSolverID(), MeshDataHandle->GetOwningParticleID()))
 				{
@@ -201,7 +201,7 @@ void FChaosVDPlaybackViewportClient::FocusOnSelectedObject()
 			using namespace Chaos::VD::TypedElementDataUtil;
 			if (FChaosVDSceneParticle* Particle = GetStructDataFromTypedElementHandle<FChaosVDSceneParticle>(SelectedParticlesHandles[0]))
 			{
-				FBox BoxToTrack = Particle->GetComponentsBoundingBox();
+				FBox BoxToTrack = Particle->GetBoundingBox();
 				FocusViewportOnBox(BoxToTrack.ExpandBy(TrackingViewDistance), true);
 			}	
 		}
@@ -258,6 +258,8 @@ void FChaosVDPlaybackViewportClient::Draw(const FSceneView* View, FPrimitiveDraw
 
 	if (const TSharedPtr<FChaosVDScene> ScenePtr = CVDScene.Pin())
 	{
+		ScenePtr->UpdateWorldStreamingLocation(GetViewLocation());
+
 		//TODO: Currently we can safely assume that any component in these actors is meant to have a visualizer, but we might need a proper interface for these components in the future
 		TInlineComponentArray<const UActorComponent*> ComponentsToVisualize;
 
diff --git a/Engine/Plugins/ChaosVD/Source/ChaosVD/Private/ChaosVDScene.cpp b/Engine/Plugins/ChaosVD/Source/ChaosVD/Private/ChaosVDScene.cpp
--- a/Engine/Plugins/ChaosVD/Source/ChaosVD/Private/ChaosVDScene.cpp
+++ b/Engine/Plugins/ChaosVD/Source/ChaosVD/Private/ChaosVDScene.cpp
@@ -71,6 +71,8 @@ void FChaosVDScene::Initialize()
 		StreamableManager->RequestSyncLoad(Settings->InstancedMeshesMaterial.ToSoftObjectPath());
 		StreamableManager->RequestSyncLoad(Settings->InstancedMeshesQueryOnlyMaterial.ToSoftObjectPath());
 		StreamableManager->RequestSyncLoad(Settings->AmbientCubeMapTexture.ToSoftObjectPath());
+		StreamableManager->RequestSyncLoad(Settings->BoxMesh.ToSoftObjectPath());
+		StreamableManager->RequestSyncLoad(Settings->SphereMesh.ToSoftObjectPath());
 	}
 	
 	PhysicsVDWorld = CreatePhysicsVDWorld();
@@ -171,6 +173,7 @@ void FChaosVDScene::AddFromCVDWorldTagToActor(AActor* Actor)
 			if (IEditorDataStorageProvider* DataStorage = GetMutableDataStorageFeature<IEditorDataStorageProvider>(StorageFeatureName))
 			{
 				DataStorage->AddColumn<FTypedElementFromCVDWorldTag>(Row);
+				DataStorage->AddColumn<FChaosVDActiveObjectTag>(Row);
 				DataStorage->AddColumn<FTypedElementSyncFromWorldTag>(Row);
 			}
 		}
@@ -195,6 +198,11 @@ void FChaosVDScene::SetLoadedRecording(const TSharedPtr<FChaosVDRecording>& NewR
 	}
 }
 
+AActor* FChaosVDScene::GetMeshComponentsContainerActor() const
+{
+	return MeshComponentContainerActor;
+}
+
 AChaosVDSolverInfoActor* FChaosVDScene::GetOrCreateSolverInfoActor(int32 SolverID)
 {
 	if (AChaosVDSolverInfoActor** SolverInfoActorPtrPtr = SolverDataContainerBySolverID.Find(SolverID))
@@ -493,6 +501,7 @@ AActor* FChaosVDScene::CreateMeshComponentsContainer(UWorld* TargetWorld)
 
 	MeshComponentContainerActor = TargetWorld->SpawnActor<AChaosVDGeometryContainer>();
 	MeshComponentContainerActor->SetFolderPath(GeometryFolderPath);
+	MeshComponentContainerActor->SetScene(AsWeak());
 
 	return MeshComponentContainerActor;
 }
@@ -617,6 +626,18 @@ bool FChaosVDScene::Tick(float DeltaTime)
 	return true;
 }
 
+void FChaosVDScene::UpdateWorldStreamingLocation(const FVector& InLocation)
+{
+	WorldStreamingLocation = InLocation;
+	for (TObjectPtr<AChaosVDDataContainerBaseActor>& DataContainerActor : AvailableDataContainerActors)
+	{
+		if (DataContainerActor)
+		{
+			DataContainerActor->HandleWorldStreamingLocationUpdated(WorldStreamingLocation);
+		}
+	}
+}
+
 void FChaosVDScene::InitializeSelectionSets()
 {
 	SelectionSet = NewObject<UTypedElementSelectionSet>(GetTransientPackage(), NAME_None, RF_Transactional);
diff --git a/Engine/Plugins/ChaosVD/Source/ChaosVD/Private/ChaosVDSceneParticle.cpp b/Engine/Plugins/ChaosVD/Source/ChaosVD/Private/ChaosVDSceneParticle.cpp
--- a/Engine/Plugins/ChaosVD/Source/ChaosVD/Private/ChaosVDSceneParticle.cpp
+++ b/Engine/Plugins/ChaosVD/Source/ChaosVD/Private/ChaosVDSceneParticle.cpp
@@ -36,7 +36,7 @@ namespace Chaos::VisualDebugger::Cvars
 
 FChaosVDSceneParticle::~FChaosVDSceneParticle()
 {
-	VisitGeometryInstances([](const TSharedRef<FChaosVDMeshDataInstanceHandle>& MeshDataHandle)
+	VisitGeometryInstances([](const TSharedRef<FChaosVDInstancedMeshData>& MeshDataHandle)
 	{
 		if (IChaosVDGeometryComponent* AsGeometryComponent = Cast<IChaosVDGeometryComponent>(MeshDataHandle->GetMeshComponent()))
 		{
@@ -58,49 +58,58 @@ FChaosVDSceneParticle::FChaosVDSceneParticle()
 	SetIconName(DefaultIconName);
 }
 
-void FChaosVDSceneParticle::UpdateTEDSData()
+void FChaosVDSceneParticle::RemoveAllGeometry()
 {
-	using namespace UE::Editor::DataStorage;
-	SyncTEDSRegistration(GetMutableDataStorageFeature<IEditorDataStorageProvider>(StorageFeatureName), GetMutableDataStorageFeature<IEditorDataStorageCompatibilityProvider>(CompatibilityFeatureName));
-	SyncTEDSTags();
+	VisitGeometryInstances([](const TSharedRef<FChaosVDInstancedMeshData>& MeshDataHandle)
+	{
+		if (IChaosVDGeometryComponent* AsGeometryComponent = Cast<IChaosVDGeometryComponent>(MeshDataHandle->GetMeshComponent()))
+		{
+			AsGeometryComponent->RemoveMeshInstance(MeshDataHandle);
+		}
+	});
+
+	MeshDataHandles.Reset();
+
+	EnumAddFlags(DirtyFlags, EChaosVDSceneParticleDirtyFlags::Geometry);
+}
+
+FBox FChaosVDSceneParticle::GetStreamingBounds() const
+{
+	ensure(CachedBounds.IsValid);
+
+	return CachedBounds;
 }
 
-void FChaosVDSceneParticle::SyncTEDSRegistration(IEditorDataStorageProvider* StorageProvider, IEditorDataStorageCompatibilityProvider* DataStorageCompatibility)
+void FChaosVDSceneParticle::SyncStreamingState()
 {
-	using namespace Chaos::VD::TedsUtils;
-	if (EnumHasAnyFlags(DirtyFlags, EChaosVDSceneParticleDirtyFlags::Active))
+	if (StreamingState == EStreamingState::Visible)
 	{
-		if (StorageProvider)
+		if (MeshDataHandles.IsEmpty())
 		{
-			if (IsActive())
-			{
-				CachedRowHandle = AddObjectToDataStorage(this, StorageProvider, DataStorageCompatibility);
-			}
-			else
-			{
-				RemoveObjectToFromDataStorage(this, DataStorageCompatibility);
-			}
+			EnumAddFlags(DirtyFlags, EChaosVDSceneParticleDirtyFlags::PreUpdatePass);
+			EnumAddFlags(DirtyFlags, EChaosVDSceneParticleDirtyFlags::Geometry);
+			EnumAddFlags(DirtyFlags, EChaosVDSceneParticleDirtyFlags::Coloring);
+			EnumAddFlags(DirtyFlags, EChaosVDSceneParticleDirtyFlags::Visibility);
+			EnumAddFlags(DirtyFlags, EChaosVDSceneParticleDirtyFlags::Transform);
+			EnumAddFlags(DirtyFlags, EChaosVDSceneParticleDirtyFlags::CollisionData);
+
+			ProcessPendingParticleDataUpdates();
 		}
-		EnumRemoveFlags(DirtyFlags, EChaosVDSceneParticleDirtyFlags::Active);
+	}
+	else if (!MeshDataHandles.IsEmpty())
+	{
+		RemoveAllGeometry();
 	}
 }
 
-void FChaosVDSceneParticle::SyncTEDSTags()
+int32 FChaosVDSceneParticle::GetStreamingID() const
 {
-	if (EnumHasAnyFlags(DirtyFlags, EChaosVDSceneParticleDirtyFlags::TEDS))
+	if (ensure(ParticleDataPtr))
 	{
-		using namespace UE::Editor::DataStorage;
-
-		if (CachedRowHandle != InvalidRowHandle)
-		{
-			if (IEditorDataStorageProvider* DataStorage = GetMutableDataStorageFeature<IEditorDataStorageProvider>(StorageFeatureName))
-			{
-				DataStorage->AddColumn<FTypedElementSyncFromWorldTag>(CachedRowHandle);
-			}	
-		}
+		return ParticleDataPtr->ParticleIndex;
 	}
 
-	EnumRemoveFlags(DirtyFlags, EChaosVDSceneParticleDirtyFlags::TEDS);
+	return INDEX_NONE;
 }
 
 void FChaosVDSceneParticle::UpdateParent(const TSharedPtr<FChaosVDParticleDataWrapper>& InRecordedData)
@@ -177,17 +186,26 @@ void FChaosVDSceneParticle::PreUpdateFromRecordedParticleData(const TSharedPtr<F
 	const bool bDisabledStateChanged = ParticleDataPtr && (ParticleDataPtr->ParticleDynamicsMisc.bDisabled != InRecordedData->ParticleDynamicsMisc.bDisabled);
 	const bool bHasNewGeometry = !ParticleDataPtr || (ParticleDataPtr->GeometryHash != InRecordedData->GeometryHash);
 	const bool bStateChanged = !ParticleDataPtr || (ParticleDataPtr->ParticleDynamicsMisc.MObjectState != InRecordedData->ParticleDynamicsMisc.MObjectState);
+	// Particle name shouldn't change, but CVD keeps particle instances alive even when a particle is destroyed so they can be reused if a user scrubs back to a frame where the particle existed
+	// But it could also be the case that the particle ID was re-used, therefore the name could have changed as this is a new particle
+	const bool bHasNewName = !ParticleDataPtr || (ParticleDataPtr->DebugNameHash != InRecordedData->DebugNameHash);
 
 	static FString UnnamedParticle("UnnamedParticle");
-	if (!ParticleDataPtr)
+	if (bHasNewName)
 	{
 		SetDisplayName(InRecordedData->DebugName.IsEmpty() ? UnnamedParticle : InRecordedData->DebugName);
+		EnumAddFlags(DirtyFlags, EChaosVDSceneParticleDirtyFlags::TEDS);
 	}
 
 	ParticleDataPtr = InRecordedData;
 
 	if (bHasNewGeometry)
 	{
+		if (const TSharedPtr<FChaosVDScene>& ScenePtr = SceneWeakPtr.Pin())
+		{
+			CurrentRootGeometry = ScenePtr->GetUpdatedGeometry(ParticleDataPtr->GeometryHash);
+		}
+
 		EnumAddFlags(DirtyFlags, EChaosVDSceneParticleDirtyFlags::Geometry);
 	}
 
@@ -212,12 +230,22 @@ void FChaosVDSceneParticle::PreUpdateFromRecordedParticleData(const TSharedPtr<F
 	{
 		EnumAddFlags(DirtyFlags, EChaosVDSceneParticleDirtyFlags::Coloring);
 	}
+
+	if (EnumHasAnyFlags(DirtyFlags, EChaosVDSceneParticleDirtyFlags::Transform) || EnumHasAnyFlags(DirtyFlags, EChaosVDSceneParticleDirtyFlags::Geometry))
+	{
+		CalculateAndCacheBounds();
+	}
 	
 	EnumAddFlags(DirtyFlags, EChaosVDSceneParticleDirtyFlags::PreUpdatePass);
 }
 
 void FChaosVDSceneParticle::ProcessPendingParticleDataUpdates()
 {
+	if (!ParticleDataPtr)
+	{
+		return;
+	}
+
 	if (!ensure(IsInGameThread()))
 	{
 		UE_LOG(LogChaosVDEditor, Error, TEXT("Attempted to update for particle [%s] outside of the game thread!. This is not supported!"), *GetDisplayName())
@@ -230,34 +258,32 @@ void FChaosVDSceneParticle::ProcessPendingParticleDataUpdates()
 		return;
 	}
 
-	if (!ParticleDataPtr)
+	if (StreamingState == EStreamingState::Visible)
 	{
-		return;
-	}
+		if (EnumHasAnyFlags(DirtyFlags, EChaosVDSceneParticleDirtyFlags::Geometry))
+		{
+			UpdateGeometry(CurrentRootGeometry, EChaosVDActorGeometryUpdateFlags::ForceUpdate);
+		}
 
-	if (EnumHasAnyFlags(DirtyFlags, EChaosVDSceneParticleDirtyFlags::Geometry))
-	{
-		UpdateGeometry(ParticleDataPtr->GeometryHash, EChaosVDActorGeometryUpdateFlags::ForceUpdate);
-	}
-	
-	if (EnumHasAnyFlags(DirtyFlags, EChaosVDSceneParticleDirtyFlags::Transform))
-	{
-		ApplyPendingTransformData();
-	}
+		if (EnumHasAnyFlags(DirtyFlags, EChaosVDSceneParticleDirtyFlags::Transform))
+		{
+			ApplyPendingTransformData();
+		}
 
-	if (EnumHasAnyFlags(DirtyFlags, EChaosVDSceneParticleDirtyFlags::CollisionData))
-	{
-		UpdateShapeDataComponents();
-	}
+		if (EnumHasAnyFlags(DirtyFlags, EChaosVDSceneParticleDirtyFlags::CollisionData))
+		{
+			UpdateShapeDataComponents();
+		}
 
-	if (EnumHasAnyFlags(DirtyFlags, EChaosVDSceneParticleDirtyFlags::Visibility))
-	{
-		UpdateGeometryComponentsVisibility();
-	}
+		if (EnumHasAnyFlags(DirtyFlags, EChaosVDSceneParticleDirtyFlags::Visibility))
+		{
+			UpdateGeometryComponentsVisibility();
+		}
 
-	if (EnumHasAnyFlags(DirtyFlags, EChaosVDSceneParticleDirtyFlags::Coloring))
-	{
-		UpdateGeometryColors();
+		if (EnumHasAnyFlags(DirtyFlags, EChaosVDSceneParticleDirtyFlags::Coloring))
+		{
+			UpdateGeometryColors();
+		}
 	}
 
 	if (EnumHasAnyFlags(DirtyFlags, EChaosVDSceneParticleDirtyFlags::Parent))
@@ -268,14 +294,14 @@ void FChaosVDSceneParticle::ProcessPendingParticleDataUpdates()
 	EnumRemoveFlags(DirtyFlags, EChaosVDSceneParticleDirtyFlags::PreUpdatePass);
 }
 
-void FChaosVDSceneParticle::SetSelectedMeshInstance(const TWeakPtr<FChaosVDMeshDataInstanceHandle>& GeometryInstanceToSelect)
+void FChaosVDSceneParticle::SetSelectedMeshInstance(const TWeakPtr<FChaosVDInstancedMeshData>& GeometryInstanceToSelect)
 {
 	if (!ParticleDataPtr)
 	{
 		return;
 	}
 
-	const TSharedPtr<FChaosVDMeshDataInstanceHandle> GeometryInstanceToSelectPtr = GeometryInstanceToSelect.Pin();
+	const TSharedPtr<FChaosVDInstancedMeshData> GeometryInstanceToSelectPtr = GeometryInstanceToSelect.Pin();
 	if (!GeometryInstanceToSelectPtr)
 	{
 		return;
@@ -287,11 +313,6 @@ void FChaosVDSceneParticle::SetSelectedMeshInstance(const TWeakPtr<FChaosVDMeshD
 	}
 }
 
-void FChaosVDSceneParticle::HandleNewGeometryLoaded(uint32 GeometryID, const Chaos::FConstImplicitObjectPtr& InGeometryData)
-{
-	UpdateGeometry(InGeometryData, EChaosVDActorGeometryUpdateFlags::ForceUpdate);
-}
-
 void FChaosVDSceneParticle::HandleDeSelected()
 {
 	CurrentSelectedGeometryInstance = nullptr;
@@ -317,9 +338,9 @@ bool FChaosVDSceneParticle::IsSelected()
 
 void FChaosVDSceneParticle::ProcessUpdatedAndRemovedHandles(TArray<TSharedPtr<FChaosVDExtractedGeometryDataHandle>>& OutExtractedGeometryDataHandles)
 {
-	for (TArray<TSharedRef<FChaosVDMeshDataInstanceHandle>>::TIterator MeshDataHandleRemoveIterator = MeshDataHandles.CreateIterator(); MeshDataHandleRemoveIterator; ++MeshDataHandleRemoveIterator)
+	for (TArray<TSharedRef<FChaosVDInstancedMeshData>>::TIterator MeshDataHandleRemoveIterator = MeshDataHandles.CreateIterator(); MeshDataHandleRemoveIterator; ++MeshDataHandleRemoveIterator)
 	{
-		TSharedRef<FChaosVDMeshDataInstanceHandle>& ExistingMeshDataHandle = *MeshDataHandleRemoveIterator;
+		TSharedRef<FChaosVDInstancedMeshData>& ExistingMeshDataHandle = *MeshDataHandleRemoveIterator;
 
 		bool bExists = false;
 
@@ -349,6 +370,7 @@ void FChaosVDSceneParticle::ProcessUpdatedAndRemovedHandles(TArray<TSharedPtr<FC
 			if (IChaosVDGeometryComponent* AsGeometryComponent = Cast<IChaosVDGeometryComponent>(ExistingMeshDataHandle->GetMeshComponent()))
 			{
 				AsGeometryComponent->RemoveMeshInstance(ExistingMeshDataHandle);
+				ExistingMeshDataHandle->MarkPendingDestroy();
 			}
 
 			MeshDataHandleRemoveIterator.RemoveCurrent();
@@ -395,7 +417,7 @@ void FChaosVDSceneParticle::UpdateGeometry(const Chaos::FConstImplicitObjectPtr&
 	// If the new implicit object is empty, then we can just clear all the mesh components and early out
 	if (ObjectsToGenerateNum == 0)
 	{
-		VisitGeometryInstances([this](const TSharedRef<FChaosVDMeshDataInstanceHandle>& MeshDataHandle)
+		VisitGeometryInstances([this](const TSharedRef<FChaosVDInstancedMeshData>& MeshDataHandle)
 		{
 			if (IChaosVDGeometryComponent* AsGeometryComponent = Cast<IChaosVDGeometryComponent>(MeshDataHandle->GetMeshComponent()))
 			{
@@ -431,15 +453,21 @@ void FChaosVDSceneParticle::UpdateGeometry(const Chaos::FConstImplicitObjectPtr&
 		for (const TSharedPtr<FChaosVDExtractedGeometryDataHandle>& ExtractedGeometryDataHandle : OutExtractedGeometryDataHandles)
 		{
 			//TODO: Time Slice component creation
-			TSharedPtr<FChaosVDMeshDataInstanceHandle> MeshDataInstance;
+			TSharedPtr<FChaosVDInstancedMeshData> MeshDataInstance;
+
+			if (!ExtractedGeometryDataHandle)
+			{
+				UE_LOG(LogChaosVDEditor, Error, TEXT("[%s] Failed To Create mesh data instance for [%s] | Invalid Source geometry"), ANSI_TO_TCHAR(__FUNCTION__), *GetDisplayName());
+				continue;
+			}
 
 			if (bHasToUseStaticMeshComponent)
 			{
-				MeshDataInstance = GeometryGenerator->CreateMeshDataInstance<UChaosVDStaticMeshComponent>(*ParticleDataPtr.Get(), ExtractedGeometryDataHandle);
+				MeshDataInstance = GeometryGenerator->CreateMeshDataInstance<UChaosVDStaticMeshComponent>(*ParticleDataPtr.Get(), ExtractedGeometryDataHandle.ToSharedRef());
 			}
 			else
 			{
-				MeshDataInstance = GeometryGenerator->CreateMeshDataInstance<UChaosVDInstancedStaticMeshComponent>(*ParticleDataPtr.Get(), ExtractedGeometryDataHandle);
+				MeshDataInstance = GeometryGenerator->CreateMeshDataInstance<UChaosVDInstancedStaticMeshComponent>(*ParticleDataPtr.Get(), ExtractedGeometryDataHandle.ToSharedRef());
 			}
 
 			if (!MeshDataInstance)
@@ -448,7 +476,7 @@ void FChaosVDSceneParticle::UpdateGeometry(const Chaos::FConstImplicitObjectPtr&
 				continue;
 			}
 
-			const UMeshComponent* CreatedMeshComponent = MeshDataInstance->GetMeshComponent();
+			IChaosVDGeometryComponent* CreatedMeshComponent = Cast<IChaosVDGeometryComponent>(MeshDataInstance->GetMeshComponent());
 			if (!CreatedMeshComponent)
 			{
 				UE_LOG(LogChaosVDEditor, Error, TEXT("[%s] Failed To Create mesh component for [%s]"), ANSI_TO_TCHAR(__FUNCTION__), *GetDisplayName());
@@ -480,47 +508,35 @@ void FChaosVDSceneParticle::UpdateGeometry(const Chaos::FConstImplicitObjectPtr&
 	}
 }
 
-void FChaosVDSceneParticle::UpdateGeometry(uint32 NewGeometryHash, EChaosVDActorGeometryUpdateFlags OptionsFlags)
+void FChaosVDSceneParticle::CalculateAndCacheBounds() const
 {
-	if (const TSharedPtr<FChaosVDScene>& ScenePtr = SceneWeakPtr.Pin())
+	if (!ensure(CurrentRootGeometry))
 	{
-		if (const Chaos::FConstImplicitObjectPtr& Geometry = ScenePtr->GetUpdatedGeometry(NewGeometryHash))
-		{
-			UpdateGeometry(Geometry, OptionsFlags);
-		}
+		return;
+	}
+
+	if (CurrentRootGeometry->HasBoundingBox())
+	{
+		Chaos::FAABB3 ChaosBox = CurrentRootGeometry->CalculateTransformedBounds(PendingParticleTransform);
+		CachedBounds = FBox(ChaosBox.Min(),ChaosBox.Max());
 	}
 }
 
-FBox FChaosVDSceneParticle::GetComponentsBoundingBox() const
+FBox FChaosVDSceneParticle::GetBoundingBox() const
 {
-	FBox BoundingBox = FBox(ForceInitToZero);
-	if (ParticleDataPtr)
+	if (!CachedBounds.IsValid)
 	{
-		FBoxSphereBounds::Builder BoundsBuilder;
+		CalculateAndCacheBounds();
+	}
 
-		for (const TSharedRef<FChaosVDMeshDataInstanceHandle>& MeshDataHandle : MeshDataHandles)
-		{
-			if (const UStaticMeshComponent* StaticMeshComponent = Cast<UStaticMeshComponent>(MeshDataHandle->GetMeshComponent()))
-			{
-				if (const UStaticMesh* Mesh = StaticMeshComponent->GetStaticMesh())
-				{
-					BoundsBuilder += Mesh->GetBounds().TransformBy(MeshDataHandle->GetWorldTransform());
-				}
-			}
-		}
+	return CachedBounds;
+}
 
-		// Geometry might not be generated yet, so we need a placeholder box so the Focus on object feature works
-		if (!BoundsBuilder.IsValid())
-		{
-			return BoundingBox.ExpandBy(10.0f).MoveTo(ParticleDataPtr->ParticlePositionRotation.MX);
-		}
-		else
-		{
-			BoundingBox = FBoxSphereBounds(BoundsBuilder).GetBox();
-		}
-	}
+Chaos::TAABB<double, 3> FChaosVDSceneParticle::GetChaosBoundingBox() const
+{
+	FBox Bounds = GetBoundingBox();
 
-	return BoundingBox;
+	return Chaos::TAABB<double, 3>(Bounds.Min, Bounds.Max);
 }
 
 TConstArrayView<TSharedPtr<FChaosVDParticlePairMidPhase>> FChaosVDSceneParticle::GetCollisionData()
@@ -549,9 +565,9 @@ FName FChaosVDSceneParticle::GetProviderName()
 
 void FChaosVDSceneParticle::UpdateMeshInstancesSelectionState()
 {
-	TSharedPtr<FChaosVDMeshDataInstanceHandle> CurrentSelectedGeometry = CurrentSelectedGeometryInstance.Pin();
+	TSharedPtr<FChaosVDInstancedMeshData> CurrentSelectedGeometry = CurrentSelectedGeometryInstance.Pin();
 	const bool bIsOwningParticleSelectedInEditor = IsSelected();
-	VisitGeometryInstances([this, bIsOwningParticleSelectedInEditor, CurrentSelectedGeometry](const TSharedRef<FChaosVDMeshDataInstanceHandle>& MeshDataHandle)
+	VisitGeometryInstances([this, bIsOwningParticleSelectedInEditor, CurrentSelectedGeometry](const TSharedRef<FChaosVDInstancedMeshData>& MeshDataHandle)
 	{
 		const bool bShouldSelectInstance = bIsOwningParticleSelectedInEditor ? (CurrentSelectedGeometry ? CurrentSelectedGeometryInstance == MeshDataHandle : true) : false;
 
@@ -632,7 +648,7 @@ const TArray<TSharedPtr<FChaosVDConstraintDataWrapperBase>>* FChaosVDScenePartic
 
 void FChaosVDSceneParticle::UpdateShapeDataComponents()
 {
-	VisitGeometryInstances([this](const TSharedRef<FChaosVDMeshDataInstanceHandle>& MeshDataHandle)
+	VisitGeometryInstances([this](const TSharedRef<FChaosVDInstancedMeshData>& MeshDataHandle)
 	{
 		if (ParticleDataPtr)
 		{
@@ -640,13 +656,14 @@ void FChaosVDSceneParticle::UpdateShapeDataComponents()
 		}
 	});
 
+
 	EnumAddFlags(DirtyFlags, EChaosVDSceneParticleDirtyFlags::Visibility);
 	EnumRemoveFlags(DirtyFlags, EChaosVDSceneParticleDirtyFlags::CollisionData);
 }
 
 void FChaosVDSceneParticle::ApplyPendingTransformData()
 {
-	VisitGeometryInstances([this](const TSharedRef<FChaosVDMeshDataInstanceHandle>& MeshDataHandle)
+	VisitGeometryInstances([this](const TSharedRef<FChaosVDInstancedMeshData>& MeshDataHandle)
 	{
 		MeshDataHandle->SetWorldTransform(PendingParticleTransform);
 	});
@@ -661,9 +678,9 @@ void FChaosVDSceneParticle::SetParent(const TSharedPtr<FChaosVDBaseSceneObject>&
 	EnumAddFlags(DirtyFlags, EChaosVDSceneParticleDirtyFlags::TEDS);
 }
 
-void FChaosVDSceneParticle::UpdateGeometryComponentsVisibility()
+void FChaosVDSceneParticle::UpdateGeometryComponentsVisibility(EChaosVDParticleVisibilityUpdateFlags Flags)
 {
-	VisitGeometryInstances([this](const TSharedRef<FChaosVDMeshDataInstanceHandle>& MeshDataHandle)
+	VisitGeometryInstances([this](const TSharedRef<FChaosVDInstancedMeshData>& MeshDataHandle)
 	{
 		if (ParticleDataPtr)
 		{
@@ -672,16 +689,19 @@ void FChaosVDSceneParticle::UpdateGeometryComponentsVisibility()
 	});
 	
 	EnumRemoveFlags(DirtyFlags, EChaosVDSceneParticleDirtyFlags::Visibility);
-	
-	if (TSharedPtr<FChaosVDScene> ScenePtr = SceneWeakPtr.Pin())
+
+	if (EnumHasAnyFlags(Flags, EChaosVDParticleVisibilityUpdateFlags::DirtyScene))
 	{
-		ScenePtr->RequestUpdate();
+		if (TSharedPtr<FChaosVDScene> ScenePtr = SceneWeakPtr.Pin())
+		{
+			ScenePtr->RequestUpdate();
+		}
 	}
 }
 
 void FChaosVDSceneParticle::UpdateGeometryColors()
 {
-	VisitGeometryInstances([this](const TSharedRef<FChaosVDMeshDataInstanceHandle>& MeshDataHandle)
+	VisitGeometryInstances([this](const TSharedRef<FChaosVDInstancedMeshData>& MeshDataHandle)
 	{
 		if (ParticleDataPtr)
 		{
diff --git a/Engine/Plugins/ChaosVD/Source/ChaosVD/Private/ChaosVDSceneParticle.h b/Engine/Plugins/ChaosVD/Source/ChaosVD/Private/ChaosVDSceneParticle.h
--- a/Engine/Plugins/ChaosVD/Source/ChaosVD/Private/ChaosVDSceneParticle.h
+++ b/Engine/Plugins/ChaosVD/Source/ChaosVD/Private/ChaosVDSceneParticle.h
@@ -15,10 +15,18 @@
 class IEditorDataStorageCompatibilityProvider;
 class IEditorDataStorageProvider;
 
+UENUM()
+enum class EChaosVDParticleVisibilityUpdateFlags
+{
+	None = 0,
+	DirtyScene = 1 << 0,
+};
+ENUM_CLASS_FLAGS(EChaosVDParticleVisibilityUpdateFlags)
+
 USTRUCT()
-struct FChaosVDSceneParticle : public FChaosVDBaseSceneObject
+struct FChaosVDSceneParticle : public FChaosVDBaseSceneObject 
 {
-	
+
 GENERATED_BODY()
 
 public:
@@ -29,18 +37,15 @@ public:
 	void PreUpdateFromRecordedParticleData(const TSharedPtr<FChaosVDParticleDataWrapper>& InRecordedData, const Chaos::FRigidTransform3& SimulationTransform);
 	void ProcessPendingParticleDataUpdates();
 
-	void UpdateGeometry(const Chaos::FConstImplicitObjectPtr& InImplicitObject, EChaosVDActorGeometryUpdateFlags OptionsFlags = EChaosVDActorGeometryUpdateFlags::None);
-
-	void UpdateGeometry(uint32 NewGeometryHash, EChaosVDActorGeometryUpdateFlags OptionsFlags = EChaosVDActorGeometryUpdateFlags::None);
-
-	TSharedPtr<const FChaosVDParticleDataWrapper> GetParticleData()
+	TSharedPtr<const FChaosVDParticleDataWrapper> GetParticleData() const
 	{
 		return ParticleDataPtr;
 	}
 
 	virtual void SetParent(const TSharedPtr<FChaosVDBaseSceneObject>& NewParent) override;
 
-	void UpdateGeometryComponentsVisibility();
+	void UpdateGeometryComponentsVisibility(EChaosVDParticleVisibilityUpdateFlags Flags = EChaosVDParticleVisibilityUpdateFlags::None);
+
 	void UpdateGeometryColors();
 
 	/** Changes the active state of this CVD Particle Actor */
@@ -68,7 +73,8 @@ public:
 		return bIsActive;
 	}
 
-	FBox GetComponentsBoundingBox() const;
+	FBox GetBoundingBox() const;
+	Chaos::TAABB<double, 3> GetChaosBoundingBox() const;
 
 	// BEGIN IChaosVDCollisionDataProvider Interface
 	TConstArrayView<TSharedPtr<FChaosVDParticlePairMidPhase>> GetCollisionData();
@@ -94,19 +100,18 @@ public:
 	void UpdateMeshInstancesSelectionState();
 
 	// BEGIN IChaosVDGeometryOwner Interface
-	TConstArrayView<TSharedRef<FChaosVDMeshDataInstanceHandle>> GetMeshInstances() const
+	TConstArrayView<TSharedRef<FChaosVDInstancedMeshData>> GetMeshInstances() const
 	{
 		return MeshDataHandles;
 	}
 
-	virtual void SetSelectedMeshInstance(const TWeakPtr<FChaosVDMeshDataInstanceHandle>& GeometryInstanceToSelect);
+	virtual void SetSelectedMeshInstance(const TWeakPtr<FChaosVDInstancedMeshData>& GeometryInstanceToSelect);
 
-	virtual TWeakPtr<FChaosVDMeshDataInstanceHandle> GetSelectedMeshInstance() const
+	virtual TWeakPtr<FChaosVDInstancedMeshData> GetSelectedMeshInstance() const
 	{
 		return CurrentSelectedGeometryInstance;
 	}
 
-	void HandleNewGeometryLoaded(uint32 GeometryID, const Chaos::FConstImplicitObjectPtr& InGeometryData);
 	// END IChaosVDGeometryOwner Interface
 
 	void HandleDeSelected();
@@ -134,12 +139,19 @@ public:
 		return DirtyFlags;
 	}
 
-	void UpdateTEDSData();
+	void RemoveAllGeometry();
+
+	virtual FBox GetStreamingBounds() const override;
+
+	virtual void SyncStreamingState() override;
+
+	virtual int32 GetStreamingID() const override;
 
 protected:
+	
+	void UpdateGeometry(const Chaos::FConstImplicitObjectPtr& InImplicitObject, EChaosVDActorGeometryUpdateFlags OptionsFlags = EChaosVDActorGeometryUpdateFlags::None);
 
-	void SyncTEDSRegistration(IEditorDataStorageProvider* StorageProvider, IEditorDataStorageCompatibilityProvider* DataStorageCompatibility);
-	void SyncTEDSTags();
+	void CalculateAndCacheBounds() const;
 
 	void UpdateParent(const TSharedPtr<FChaosVDParticleDataWrapper>& InRecordedData);
 
@@ -162,21 +174,25 @@ protected:
 	EChaosVDHideParticleFlags HideParticleFlags = EChaosVDHideParticleFlags::None;
 	EChaosVDSceneParticleDirtyFlags DirtyFlags = EChaosVDSceneParticleDirtyFlags::None;
 
+	Chaos::FConstImplicitObjectPtr CurrentRootGeometry = nullptr;
+
 	TWeakPtr<FChaosVDSceneParticle> ParentParticleInstance;
 
 	TSharedPtr<FChaosVDParticleDataWrapper> ParticleDataPtr;
 
 	FDelegateHandle GeometryUpdatedDelegate;
 
-	TArray<TSharedRef<FChaosVDMeshDataInstanceHandle>> MeshDataHandles;
+	TArray<TSharedRef<FChaosVDInstancedMeshData>> MeshDataHandles;
 
-	TWeakPtr<FChaosVDMeshDataInstanceHandle> CurrentSelectedGeometryInstance;
+	TWeakPtr<FChaosVDInstancedMeshData> CurrentSelectedGeometryInstance;
 
 	TWeakPtr<FChaosVDScene> SceneWeakPtr;
 
 	FTransform PendingParticleTransform;
 	FTransform CachedSimulationTransform;
 
+	mutable FBox CachedBounds = FBox(ForceInitToZero);
+
 	friend class FChaosVDSceneParticleCustomization;
 	friend class UChaosVDParticleDataComponent;
 };
@@ -184,7 +200,7 @@ protected:
 template <typename TVisitorCallback>
 void FChaosVDSceneParticle::VisitGeometryInstances(const TVisitorCallback& VisitorCallback)
 {
-	for (TSharedRef<FChaosVDMeshDataInstanceHandle>& MeshDataHandle : MeshDataHandles)
+	for (TSharedRef<FChaosVDInstancedMeshData>& MeshDataHandle : MeshDataHandles)
 	{
 		VisitorCallback(MeshDataHandle);
 	}
diff --git a/Engine/Plugins/ChaosVD/Source/ChaosVD/Private/ChaosVDSceneParticleCustomization.cpp b/Engine/Plugins/ChaosVD/Source/ChaosVD/Private/ChaosVDSceneParticleCustomization.cpp
--- a/Engine/Plugins/ChaosVD/Source/ChaosVD/Private/ChaosVDSceneParticleCustomization.cpp
+++ b/Engine/Plugins/ChaosVD/Source/ChaosVD/Private/ChaosVDSceneParticleCustomization.cpp
@@ -96,7 +96,7 @@ void FChaosVDSceneParticleCustomization::CustomizeDetails(IDetailLayoutBuilder&
 
 			TSharedPtr<IPropertyHandle> InspectedDataPropertyHandlePtr;
 
-			if (TSharedPtr<FChaosVDMeshDataInstanceHandle> SelectedGeometryInstance = SelectedParticleInstance->GetSelectedMeshInstance().Pin())
+			if (TSharedPtr<FChaosVDInstancedMeshData> SelectedGeometryInstance = SelectedParticleInstance->GetSelectedMeshInstance().Pin())
 			{
 				InspectedDataPropertyHandlePtr = AddExternalStructure(CachedGeometryDataInstanceCopy, DetailBuilder, GeometryCategoryName, LOCTEXT("GeometryShapeDataStructName", "Geometry Shape Data"));
 			}
@@ -136,9 +136,9 @@ void FChaosVDSceneParticleCustomization::HandleSceneUpdated()
 	}
 
 	// If we have selected a mesh instance, the only data being added to the details panel is the Shape Instance data, so can just update that data here
-	if (TSharedPtr<FChaosVDMeshDataInstanceHandle> SelectedGeometryInstance = CurrentObservedParticle->GetSelectedMeshInstance().Pin())
+	if (TSharedPtr<FChaosVDInstancedMeshData> SelectedGeometryInstance = CurrentObservedParticle->GetSelectedMeshInstance().Pin())
 	{
-		CurrentObservedParticle->VisitGeometryInstances([this, SelectedGeometryInstance](const TSharedRef<FChaosVDMeshDataInstanceHandle>& MeshDataHandle)
+		CurrentObservedParticle->VisitGeometryInstances([this, SelectedGeometryInstance](const TSharedRef<FChaosVDInstancedMeshData>& MeshDataHandle)
 		{
 			if (MeshDataHandle == SelectedGeometryInstance)
 			{
diff --git a/Engine/Plugins/ChaosVD/Source/ChaosVD/Private/ChaosVDSceneParticleFlags.h b/Engine/Plugins/ChaosVD/Source/ChaosVD/Private/ChaosVDSceneParticleFlags.h
--- a/Engine/Plugins/ChaosVD/Source/ChaosVD/Private/ChaosVDSceneParticleFlags.h
+++ b/Engine/Plugins/ChaosVD/Source/ChaosVD/Private/ChaosVDSceneParticleFlags.h
@@ -41,5 +41,6 @@ enum class EChaosVDSceneParticleDirtyFlags : uint16
 	CollisionData =	1 << 6,
 	PreUpdatePass =	1 << 7,
 	TEDS =	1 << 8,
+	StreamingBounds = 1 << 9,
 };
 ENUM_CLASS_FLAGS(EChaosVDSceneParticleDirtyFlags)
diff --git a/Engine/Plugins/ChaosVD/Source/ChaosVD/Private/ChaosVDSceneStreaming.cpp b/Engine/Plugins/ChaosVD/Source/ChaosVD/Private/ChaosVDSceneStreaming.cpp
new file mode 100644
--- /dev/null
+++ b/Engine/Plugins/ChaosVD/Source/ChaosVD/Private/ChaosVDSceneStreaming.cpp
@@ -0,0 +1,363 @@
+// Copyright Epic Games, Inc. All Rights Reserved.
+
+#include "ChaosVDSceneStreaming.h"
+
+#include "ChaosVDBaseSceneObject.h"
+#include "ChaosVDModule.h"
+#include "ChaosVDScene.h"
+#include "ChaosVDSceneParticle.h"
+#include "ChaosVDSettingsManager.h"
+#include "Misc/ScopeRWLock.h"
+#include "Settings/ChaosVDGeneralSettings.h"
+
+FChaosVDSceneStreaming::FChaosVDSceneStreaming()
+{
+	UpdateStreamingQueryShape();
+	StreamingAccelerationStructure.SetTreeToDynamic();
+
+	if (UChaosVDGeneralSettings* Settings = FChaosVDSettingsManager::Get().GetSettingsObject<UChaosVDGeneralSettings>())
+	{
+		Settings->OnSettingsChanged().AddRaw(this, &FChaosVDSceneStreaming::HandleSettingsChanged);
+
+		HandleSettingsChanged(Settings);
+	}
+}
+
+FChaosVDSceneStreaming::~FChaosVDSceneStreaming()
+{
+	if (UChaosVDGeneralSettings* Settings = FChaosVDSettingsManager::Get().GetSettingsObject<UChaosVDGeneralSettings>())
+	{
+		Settings->OnSettingsChanged().RemoveAll(this);
+	}
+
+	Reset();
+}
+
+bool FChaosVDSceneStreaming::Tick(float DeltaTime)
+{
+	if (!bIsStreamingSystemEnabled)
+	{
+		return true;
+	}
+
+	if (bProcessPendingOperationsQueueInWorkerThread && CurrentProcessingTaskHandle.IsCompleted() && !PendingOperationsPerParticle.IsEmpty())
+	{
+		CurrentProcessingTaskHandle = UE::Tasks::Launch(TEXT("UpdatingCVDStreamingAccel"), [this]()
+		{
+			ProcessPendingOperations();
+		});
+	}
+
+	if (!bProcessPendingOperationsQueueInWorkerThread)
+	{
+		ProcessPendingOperations();
+	}
+
+	{
+		bool bHasDirtyFlags = false;
+		{
+			FReadScopeLock ReadLock(DirtyFlagsLock);
+			bHasDirtyFlags = DirtyFlags != EChaosVDStreamingDirtyFlags::None;
+		}
+
+		const bool bIsOverTheThreshold = LastStreamingLocationUpdate.IsSet() ? FVector::Distance(CurrentStreamingSourceLocation, LastStreamingLocationUpdate.GetValue()) > MovementThreshold : true;
+		if (bIsOverTheThreshold || bHasDirtyFlags)
+		{
+			{
+				FWriteScopeLock WriteLock(DirtyFlagsLock);
+				DirtyFlags = EChaosVDStreamingDirtyFlags::None;
+			}
+			
+			LastStreamingLocationUpdate = CurrentStreamingSourceLocation;
+			UpdateStreamingState();
+		}
+	}
+
+	return true;
+}
+
+void FChaosVDSceneStreaming::EnqueuePendingTrackingOperation(const FPendingTrackingOperation&& Operation)
+{
+	if (!bIsStreamingSystemEnabled)
+	{
+		return;
+	}
+
+	PendingOperationsPerParticle.Enqueue(Operation);
+}
+
+void FChaosVDSceneStreaming::EnqueuePendingTrackingOperation(const TSharedRef<FChaosVDBaseSceneObject>& InSceneObject, FPendingTrackingOperation::EType Type)
+{
+	if (!bIsStreamingSystemEnabled)
+	{
+		return;
+	}
+
+	PendingOperationsPerParticle.Enqueue(CreateStreamingTrackingOperation(InSceneObject, FPendingTrackingOperation::EType::AddOrUpdate));
+}
+
+void FChaosVDSceneStreaming::UpdateStreamingSourceLocation(const FVector& NewLocation)
+{
+	CurrentStreamingSourceLocation = NewLocation;
+	SteamingViewBox = SteamingViewBox.MoveTo(NewLocation);
+}
+
+bool FChaosVDSceneStreaming::IsInStreamingRange(const FBox& Bounds) const
+{
+	return bIsStreamingSystemEnabled ? SteamingViewBox.Intersect(Bounds) : true;
+}
+
+void FChaosVDSceneStreaming::UpdateStreamingStateForObject(const TSharedRef<FChaosVDBaseSceneObject>& InSceneObject) const
+{
+	if (IsInStreamingRange(InSceneObject->GetStreamingBounds()))
+	{
+		if (InSceneObject->GetStreamingState() == FChaosVDSceneParticle::EStreamingState::Hidden)
+		{
+			InSceneObject->SetStreamingState(FChaosVDSceneParticle::EStreamingState::Visible);
+		}
+	}
+	else
+	{
+		if (InSceneObject->GetStreamingState() == FChaosVDSceneParticle::EStreamingState::Visible)
+		{
+			InSceneObject->SetStreamingState(FChaosVDSceneParticle::EStreamingState::Hidden);
+		}
+	}
+}
+
+FChaosVDSceneStreaming::FPendingTrackingOperation FChaosVDSceneStreaming::CreateStreamingTrackingOperation(const TSharedRef<FChaosVDBaseSceneObject>& InSceneObject, FPendingTrackingOperation::EType Type)
+{
+	return {.Bounds =  InSceneObject->GetStreamingBounds(), .ObjectID = InSceneObject->GetStreamingID(), .OperationType = Type };
+}
+
+void FChaosVDSceneStreaming::Reset()
+{
+	ensure(IsInGameThread());
+
+	constexpr double TimeoutSeconds = 10.0;
+	if (!ensure(CurrentProcessingTaskHandle.Wait(FTimespan::FromSeconds(TimeoutSeconds))))
+	{
+		UE_LOG(LogChaosVDEditor, Error, TEXT("Failed to cancel pending operations processing queue after [%f]"), TimeoutSeconds)
+	}
+
+	{
+		FWriteScopeLock WriteLock(StreamingAccelerationOperationsLock);
+		StreamingAccelerationStructure.Reset();
+		StreamingAccelerationStructure.SetTreeToDynamic();
+	}
+}
+
+void FChaosVDSceneStreaming::SetStreamingDataSource(IChaosVDStreamingDataSource* InStreamingDataSource)
+{
+	StreamingDataSource = InStreamingDataSource;
+}
+
+void FChaosVDSceneStreaming::SetStreamingEnabled(bool bNewEnabled)
+{
+	if (bIsStreamingSystemEnabled != bNewEnabled)
+	{
+		if (bNewEnabled)
+		{
+			ReBuildAccelStructureFromSourceDataArray();
+			
+			SetDirtyFlags(EChaosVDStreamingDirtyFlags::StreamingEnabled);
+		}
+		else
+		{
+			PendingOperationsPerParticle.Empty();
+			MakeEverythingVisible();
+		}
+
+		bIsStreamingSystemEnabled = bNewEnabled;
+	}
+}
+
+void FChaosVDSceneStreaming::HandleSettingsChanged(UObject* SettingsObject)
+{
+	if (UChaosVDGeneralSettings* Settings = Cast<UChaosVDGeneralSettings>(SettingsObject))
+	{
+		if (StreamingExtent != Settings->StreamingBoxExtentSize)
+		{
+			SetStreamingExtent(Settings->StreamingBoxExtentSize);
+		}
+
+		if (bProcessPendingOperationsQueueInWorkerThread && !Settings->bProcessPendingOperationsQueueInWorkerThread)
+		{
+			CurrentProcessingTaskHandle.Wait();
+		}
+
+		bProcessPendingOperationsQueueInWorkerThread = Settings->bProcessPendingOperationsQueueInWorkerThread;
+		
+		SetStreamingEnabled(Settings->bStreamingSystemEnabled);
+	}
+}
+
+void FChaosVDSceneStreaming::SetDirtyFlags(EChaosVDStreamingDirtyFlags Flag)
+{
+	FWriteScopeLock WriteScopeLock(DirtyFlagsLock);
+	EnumAddFlags(DirtyFlags, Flag);
+}
+
+void FChaosVDSceneStreaming::RemoveDirtyFlag(EChaosVDStreamingDirtyFlags Flag)
+{
+	FWriteScopeLock WriteScopeLock(DirtyFlagsLock);
+	EnumRemoveFlags(DirtyFlags, Flag);
+}
+
+void FChaosVDSceneStreaming::ProcessPendingOperations()
+{
+	while (!PendingOperationsPerParticle.IsEmpty())
+	{
+		FPendingTrackingOperation PendingOperation;
+		PendingOperationsPerParticle.Dequeue(PendingOperation);
+
+		{
+			FWriteScopeLock WriteLock(StreamingAccelerationOperationsLock);
+			if (PendingOperation.OperationType == FPendingTrackingOperation::EType::AddOrUpdate)
+			{
+				StreamingAccelerationStructure.UpdateElement(PendingOperation.ObjectID, Chaos::TAABB<double,3>(PendingOperation.Bounds.Min, PendingOperation.Bounds.Max), true);
+			}
+			else if (PendingOperation.OperationType == FPendingTrackingOperation::EType::Remove)
+			{
+				StreamingAccelerationStructure.RemoveElement(PendingOperation.ObjectID);
+			}
+		}
+
+		SetDirtyFlags(EChaosVDStreamingDirtyFlags::AccelerationStructure);
+	}
+}
+
+void FChaosVDSceneStreaming::UpdateStreamingState()
+{
+	if (!StreamingDataSource)
+	{
+		return;
+	}
+
+	FMemMark StackMarker(FMemStack::Get());
+	TSet<int32> InStreamingRangeObjects;
+	{
+		FQueryVisitor Collector(InStreamingRangeObjects);
+		FReadScopeLock ReadLock(StreamingAccelerationOperationsLock);
+		StreamingAccelerationStructure.Overlap(Chaos::TAABB<double,3>(SteamingViewBox.Min, SteamingViewBox.Max), Collector);
+	}
+
+	TQueue<TSharedPtr<FChaosVDBaseSceneObject>, EQueueMode::Mpsc> ObjectsThatNeedStreamingStateUpdate;
+
+	{
+		TConstArrayView<TSharedRef<FChaosVDBaseSceneObject>> StreamableObjectsArray = StreamingDataSource->GetStreamableSceneObjects();
+		FReadScopeLock ObjectsReadLock(StreamingDataSource->GetObjectsLock());
+
+		ParallelFor(StreamableObjectsArray.Num(),[&StreamableObjectsArray, &InStreamingRangeObjects, &ObjectsThatNeedStreamingStateUpdate, this](int32 Index)
+		{
+			const TSharedRef<FChaosVDBaseSceneObject>& SceneObject = StreamableObjectsArray[Index];
+
+			if (InStreamingRangeObjects.Contains(SceneObject->GetStreamingID()))
+			{
+				if (SceneObject->GetStreamingState() == FChaosVDSceneParticle::EStreamingState::Hidden)
+				{
+					SceneObject->SetStreamingState(FChaosVDSceneParticle::EStreamingState::Visible);
+					ObjectsThatNeedStreamingStateUpdate.Enqueue(SceneObject);
+				}
+			}
+			else
+			{
+				if (SceneObject->GetStreamingState() == FChaosVDSceneParticle::EStreamingState::Visible)
+				{
+					SceneObject->SetStreamingState(FChaosVDSceneParticle::EStreamingState::Hidden);
+					ObjectsThatNeedStreamingStateUpdate.Enqueue(SceneObject);
+				}
+			}
+		});
+	}
+
+	bool bWasStreamingUpdated = !ObjectsThatNeedStreamingStateUpdate.IsEmpty();
+
+	while (!ObjectsThatNeedStreamingStateUpdate.IsEmpty())
+	{
+		TSharedPtr<FChaosVDBaseSceneObject> SceneObject;
+		ObjectsThatNeedStreamingStateUpdate.Dequeue(SceneObject);
+
+		if (SceneObject)
+		{
+			SceneObject->SyncStreamingState();
+		}
+	}
+
+	if (bWasStreamingUpdated)
+	{
+		if (TSharedPtr<FChaosVDScene> ScenePtr = SceneWeakPtr.Pin())
+		{
+			ScenePtr->RequestUpdate();
+		}
+	}
+}
+
+void FChaosVDSceneStreaming::SetStreamingExtent(float NewExtent)
+{
+	StreamingExtent = NewExtent;
+	UpdateStreamingQueryShape();
+}
+
+void FChaosVDSceneStreaming::UpdateStreamingQueryShape()
+{
+	SteamingViewBox = FBox(ForceInitToZero).ExpandBy(StreamingExtent).MoveTo(CurrentStreamingSourceLocation);
+
+	SetDirtyFlags(EChaosVDStreamingDirtyFlags::StreamingExtents);
+}
+
+void FChaosVDSceneStreaming::ReBuildAccelStructureFromSourceDataArray()
+{
+	Reset();
+
+	if (!StreamingDataSource)
+	{
+		return;
+	}
+
+	{
+		TConstArrayView<TSharedRef<FChaosVDBaseSceneObject>> StreamableObjectsArray = StreamingDataSource->GetStreamableSceneObjects();
+		FReadScopeLock ObjectsReadLock(StreamingDataSource->GetObjectsLock());
+
+		FWriteScopeLock AccelStructLock(StreamingAccelerationOperationsLock);
+		for (const TSharedRef<FChaosVDBaseSceneObject>& SceneObject : StreamableObjectsArray)
+		{
+			FBox ObjectStreamingBounds = SceneObject->GetStreamingBounds();
+			StreamingAccelerationStructure.UpdateElement(SceneObject->GetStreamingID(), Chaos::TAABB<double, 3>(ObjectStreamingBounds.Min, ObjectStreamingBounds.Max), true);
+		}
+	}
+
+	if (TSharedPtr<FChaosVDScene> ScenePtr = SceneWeakPtr.Pin())
+	{
+		ScenePtr->RequestUpdate();
+	}
+}
+
+void FChaosVDSceneStreaming::MakeEverythingVisible()
+{
+	Reset();
+
+	if (!StreamingDataSource)
+	{
+		return;
+	}
+
+	{
+		TConstArrayView<TSharedRef<FChaosVDBaseSceneObject>> StreamableObjectsArray = StreamingDataSource->GetStreamableSceneObjects();
+		FReadScopeLock ObjectsReadLock(StreamingDataSource->GetObjectsLock());
+		for (const TSharedRef<FChaosVDBaseSceneObject>& SceneObject : StreamableObjectsArray)
+		{
+			if (SceneObject->GetStreamingState() != FChaosVDSceneParticle::EStreamingState::Visible)
+			{
+				SceneObject->SetStreamingState(FChaosVDBaseSceneObject::EStreamingState::Visible);
+			}
+
+			SceneObject->SyncStreamingState();
+		}
+	}
+
+	if (TSharedPtr<FChaosVDScene> ScenePtr = SceneWeakPtr.Pin())
+	{
+		ScenePtr->RequestUpdate();
+	}
+}
diff --git a/Engine/Plugins/ChaosVD/Source/ChaosVD/Private/ChaosVDSceneStreaming.h b/Engine/Plugins/ChaosVD/Source/ChaosVD/Private/ChaosVDSceneStreaming.h
new file mode 100644
--- /dev/null
+++ b/Engine/Plugins/ChaosVD/Source/ChaosVD/Private/ChaosVDSceneStreaming.h
@@ -0,0 +1,235 @@
+// Copyright Epic Games, Inc. All Rights Reserved.
+
+#pragma once
+#include "Chaos/AABBTree.h"
+#include "Chaos/ISpatialAcceleration.h"
+#include "Containers/Queue.h"
+#include "Math/Box.h"
+#include "Settings/ChaosVDGeneralSettings.h"
+
+#include "ChaosVDSceneStreaming.generated.h"
+
+class FChaosVDScene;
+struct FChaosVDBaseSceneObject;
+
+UENUM()
+enum class EChaosVDStreamingDirtyFlags : uint32
+{
+	None = 0,
+	StreamingExtents = 1 << 0,
+	StreamingEnabled = 1 << 1,
+	StreamingSourceLocation = 1 << 2,
+	AccelerationStructure = 1 << 3
+};
+ENUM_CLASS_FLAGS(EChaosVDStreamingDirtyFlags)
+
+class IChaosVDStreamingDataSource
+{
+public:
+	virtual ~IChaosVDStreamingDataSource() = default;
+
+	virtual TConstArrayView<TSharedRef<FChaosVDBaseSceneObject>> GetStreamableSceneObjects() const = 0;
+	virtual FRWLock& GetObjectsLock() const = 0;
+};
+
+/** Simple pseudo level streaming system that works with a collection of ChaosVDSceneObjects -
+ * This system only updates a desired streaming state (from multiple threads) and issues a sync request in the GT.
+ * How these actions are executed depends on the implementations of each ChaosVDSceneObjects derived object
+ */
+class FChaosVDSceneStreaming : public FTSTickerObjectBase
+{
+public:
+
+	template<uint32 Alignment = DEFAULT_ALIGNMENT> class TMemStackSetAllocator : public TSetAllocator<TSparseArrayAllocator<TMemStackAllocator<Alignment>, TMemStackAllocator<Alignment>>, TMemStackAllocator<Alignment>> {};
+
+	FChaosVDSceneStreaming();
+	virtual ~FChaosVDSceneStreaming() override;
+
+	/** Relevant streaming system data that needs processing */
+	struct FPendingTrackingOperation
+	{
+		enum class EType
+		{
+			None,
+			AddOrUpdate,
+			Remove
+		};
+
+		FBox Bounds = FBox(EForceInit::ForceInitToZero);
+		int32 ObjectID = INDEX_NONE;
+		EType OperationType = EType::None;
+	};
+
+	virtual bool Tick(float DeltaTime) override;
+
+	/** Adds an object tracking operation to the queue */
+	void EnqueuePendingTrackingOperation(const FPendingTrackingOperation&& Operation);
+	
+	/**
+	 * Creates the request object tracking operation for the provided object and adds it to the queue
+	 * @param InSceneObject Object from which create the operation from
+	 * @param Type Type of operation to perform (AddOrUpdate, Remove)
+	 */
+	void EnqueuePendingTrackingOperation(const TSharedRef<FChaosVDBaseSceneObject>& InSceneObject, FPendingTrackingOperation::EType Type);
+
+	/**
+	 * Updates the current location used to calculate what objects should be fully loaded or not
+	 * @param NewLocation New streaming source location
+	 */
+	void UpdateStreamingSourceLocation(const FVector& NewLocation);
+
+	/**
+	 * Returns true if the provided bounds are within the current calculated streaming volume
+	 * @param Bounds Bounds to evaluate
+	 */
+	bool IsInStreamingRange(const FBox& Bounds) const;
+
+	/**
+	 * Sets the correct streaming state for the provided object based on its bounds
+	 * @param InSceneObject Object to evaluate
+	 */
+	void UpdateStreamingStateForObject(const TSharedRef<FChaosVDBaseSceneObject>& InSceneObject) const;
+
+	/**
+	 * Creates an object tracking operation structure with the necessary data to update the streaming system state for this object
+	 * @param InSceneObject Object from which get the necessary data
+	 * @param Type Operation type (AddOrUpdate, Remove)
+	 */
+	FPendingTrackingOperation CreateStreamingTrackingOperation(const TSharedRef<FChaosVDBaseSceneObject>& InSceneObject, FPendingTrackingOperation::EType Type);
+
+	/**
+	 * Resets the streaming system state
+	 */
+	void Reset();
+
+	/**
+	 * Sets a ptr to an array with all the objects managed by this streaming system
+	 * @param InStreamingDataSource Object that provides access to the data this system will manage
+	 */
+	void SetStreamingDataSource(IChaosVDStreamingDataSource* InStreamingDataSource);
+
+	/** Changes the enabled state of this streaming system, and updates the world acordingly
+	 * @param bNewEnabled New enabled state
+	 */
+	void SetStreamingEnabled(bool bNewEnabled);
+
+	/**
+	 * Sets a weak ptr to the CVD scene where all the objects managed by this system live
+	 * @param InSceneWeakPtr 
+	 */
+	void SetScene(const TWeakPtr<FChaosVDScene>& InSceneWeakPtr)
+	{
+		SceneWeakPtr = InSceneWeakPtr;
+	}
+
+	/**
+	 * Handles any streaming settings changes done by the user and updates the world accordingly
+	 * @param SettingsObject Object containing the new settings data
+	 */
+	void HandleSettingsChanged(UObject* SettingsObject);
+
+
+	/** Returns true if this system is enabled */
+	bool IsEnabled() const
+	{
+		return bIsStreamingSystemEnabled;
+	}
+
+protected:
+
+	void SetDirtyFlags(EChaosVDStreamingDirtyFlags Flag);
+	void RemoveDirtyFlag(EChaosVDStreamingDirtyFlags Flag);
+
+	void ProcessPendingOperations();
+
+	void UpdateStreamingState();
+
+	void SetStreamingExtent(float NewExtent);
+
+	void UpdateStreamingQueryShape();
+
+	void ReBuildAccelStructureFromSourceDataArray();
+
+	void MakeEverythingVisible();
+
+	float StreamingExtent = 5000.0f;
+
+	float MovementThreshold = 10.0f;
+
+	FVector CurrentStreamingSourceLocation = FVector::ZeroVector;
+
+	TOptional<FVector> LastStreamingLocationUpdate;
+
+	FBox SteamingViewBox;
+
+	Chaos::TAABBTree<int32, Chaos::TAABBTreeLeafArray<int32>> StreamingAccelerationStructure;
+
+	IChaosVDStreamingDataSource* StreamingDataSource = nullptr;
+
+	TQueue<FPendingTrackingOperation, EQueueMode::Mpsc> PendingOperationsPerParticle;
+
+	EChaosVDStreamingDirtyFlags DirtyFlags;
+
+	bool bIsStreamingSystemEnabled = true;
+
+	bool bProcessPendingOperationsQueueInWorkerThread = true;
+
+	TWeakPtr<FChaosVDScene> SceneWeakPtr;
+
+	UE::Tasks::FTask CurrentProcessingTaskHandle;
+
+	mutable FRWLock StreamingAccelerationOperationsLock;
+
+	mutable FRWLock DirtyFlagsLock;
+
+	struct FQueryVisitor
+	{
+		FQueryVisitor(TSet<int32>& InResults) : CollectedResults(InResults)
+		{
+		
+		}
+
+		bool VisitOverlap(const Chaos::TSpatialVisitorData<int32>& Instance)
+		{
+			CollectedResults.Add(Instance.Payload);
+			return true;
+		}
+		bool VisitSweep(const Chaos::TSpatialVisitorData<int32>& Instance, Chaos::FQueryFastData& CurData)
+		{
+			check(false);
+			return true;
+		}
+		bool VisitRaycast(const Chaos::TSpatialVisitorData<int32>& Instance, Chaos::FQueryFastData& CurData)
+		{
+			check(false);
+			return true;
+		}
+
+		const void* GetQueryData() const
+		{
+			return nullptr;
+		}
+
+		const void* GetSimData() const
+		{
+			return nullptr;
+		}
+
+		const void* GetQueryPayload() const 
+		{ 
+			return nullptr; 
+		}
+
+		bool HasBlockingHit() const
+		{
+			return false;
+		}
+
+		bool ShouldIgnore(const Chaos::TSpatialVisitorData<int32>& Instance) const
+		{
+			return false;
+		}
+
+		TSet<int32>& CollectedResults;
+	};
+};
diff --git a/Engine/Plugins/ChaosVD/Source/ChaosVD/Private/ChaosVDWorldOutlinerMode.cpp b/Engine/Plugins/ChaosVD/Source/ChaosVD/Private/ChaosVDWorldOutlinerMode.cpp
--- a/Engine/Plugins/ChaosVD/Source/ChaosVD/Private/ChaosVDWorldOutlinerMode.cpp
+++ b/Engine/Plugins/ChaosVD/Source/ChaosVD/Private/ChaosVDWorldOutlinerMode.cpp
@@ -15,6 +15,7 @@
 #include "Elements/Framework/EngineElementsLibrary.h"
 #include "Elements/Framework/TypedElementSelectionSet.h"
 #include "Elements/Interfaces/TypedElementDataStorageCompatibilityInterface.h"
+#include "TEDS/ChaosVDParticleEditorDataFactory.h"
 #include "TEDS/ChaosVDSelectionInterface.h"
 #include "TEDS/ChaosVDStructTypedElementData.h"
 #include "TEDS/ChaosVDTedsUtils.h"
@@ -109,7 +110,6 @@ bool FChaosVDActorOutlinerSelector::operator()(const TWeakPtr<ISceneOutlinerTree
 	return false;
 }
 
-
 FChaosVDWorldOutlinerMode::FChaosVDWorldOutlinerMode(const UE::Editor::Outliner::FTedsOutlinerParams& InModeParams, TWeakPtr<FChaosVDScene> InScene, TWeakPtr<FChaosVDPlaybackController> InPlaybackController)
 	: FTedsOutlinerMode(InModeParams)
 	, CVDScene(InScene)
@@ -121,11 +121,65 @@ FChaosVDWorldOutlinerMode::FChaosVDWorldOutlinerMode(const UE::Editor::Outliner:
 		return;
 	}
 
+	if (SceneOutliner)
+	{
+		TAttribute<bool> ConditionalEnabledAttribute;
+		ConditionalEnabledAttribute.BindRaw(this, &FChaosVDWorldOutlinerMode::CanPopulate);
+
+		SceneOutliner->SetEnabled(ConditionalEnabledAttribute);
+	}
+	
 	RegisterSelectionSetObject(ScenePtr->GetElementSelectionSet());
+
+	if (IEditorDataStorageProvider* DataStorage = TedsOutlinerImpl ? TedsOutlinerImpl->GetStorage() : nullptr)
+	{
+		using namespace UE::Editor::DataStorage;
+		using namespace UE::Editor::DataStorage::Queries;
+		using namespace UE::Editor::Outliner;
+
+		FQueryDescription RowAdditionQueryDescription =
+			Select(
+				TEXT("Add CVD Row to Outliner"),
+				FObserver::OnAdd<FChaosVDActiveObjectTag>().SetExecutionMode(EExecutionMode::GameThread),
+				[this](IQueryContext& Context, RowHandle Row)
+				{
+					RowsPendingAddition.Emplace(Row);
+					RowsPendingRemoval.Remove(Row);
+				})
+			.Compile();
+
+		CustomRowAdditionQueryHandle = DataStorage->RegisterQuery(MoveTemp(RowAdditionQueryDescription));
+	
+		// Row to track removal of rows from the Outliner
+		FQueryDescription RowRemovalQueryDescription =
+			Select(
+					TEXT("Remove CVD Row from Outliner"),
+					FObserver::OnRemove<FChaosVDActiveObjectTag>().SetExecutionMode(EExecutionMode::GameThread),
+					[this](IQueryContext& Context, RowHandle Row)
+					{
+						RowsPendingAddition.Remove(Row);
+						RowsPendingRemoval.Emplace(Row);
+					})
+				.Compile();
+
+		CustomRowRemovalQueryHandle = DataStorage->RegisterQuery(MoveTemp(RowRemovalQueryDescription));
+	}
+
 }
 
 FChaosVDWorldOutlinerMode::~FChaosVDWorldOutlinerMode()
 {
+	if (SceneOutliner)
+	{
+		TAttribute<bool> EmptyConditionalEnabledAttribute;
+		SceneOutliner->SetEnabled(EmptyConditionalEnabledAttribute);
+	}
+
+	if (IEditorDataStorageProvider* DataStorage = TedsOutlinerImpl ? TedsOutlinerImpl->GetStorage() : nullptr)
+	{
+		DataStorage->UnregisterQuery(CustomRowAdditionQueryHandle);
+		DataStorage->UnregisterQuery(CustomRowRemovalQueryHandle);
+	}
 }
 
 void FChaosVDWorldOutlinerMode::OnItemSelectionChanged(FSceneOutlinerTreeItemPtr Item, ESelectInfo::Type SelectionType, const FSceneOutlinerItemSelection& Selection)
@@ -171,7 +225,7 @@ void FChaosVDWorldOutlinerMode::OnItemDoubleClick(FSceneOutlinerTreeItemPtr Item
 
 	if (FChaosVDSceneParticle* ParticleInstance = Chaos::VD::SceneOutlinerUtils::Private::GetParticleInstanceFromOutlinerItem(Item))
 	{
-		ScenePtr->OnFocusRequest().Broadcast(ParticleInstance->GetComponentsBoundingBox());
+		ScenePtr->OnFocusRequest().Broadcast(ParticleInstance->GetBoundingBox());
 	}
 	else if (AActor* ActorItem = Chaos::VD::SceneOutlinerUtils::Private::GetActorFromOutlinerItem(Item))
 	{
@@ -193,7 +247,7 @@ bool FChaosVDWorldOutlinerMode::CanPopulate() const
 	{
 		// Updating the scene outliner during playback it is very expensive and can tank framerate,
 		// as it need to re-build the hierarchy when things are added in ad removed. So if we are playing we want to pause any updates to the outliner
-		return !PlaybackControllerPtr->IsPlaying();
+		return !PlaybackControllerPtr->IsPlaying() && !PlaybackControllerPtr->IsScrubbingTimeline();
 	}
 	return true;
 }
@@ -223,7 +277,70 @@ void FChaosVDWorldOutlinerMode::HandlePostSelectionChange(const UTypedElementSel
 
 		// Add a selection column in TEDS
 		TedsOutlinerImpl->SetSelection({ RowHandle });
+
+		if (FSceneOutlinerTreeItemPtr TreeItem = SceneOutliner->GetTreeItem(RowHandle, true))
+		{
+			SceneOutliner->ScrollItemIntoView(TreeItem);
+		}
+	}
+}
+
+void FChaosVDWorldOutlinerMode::ProcessPendingActionsList(double TimeBudgetInSeconds, TSet<UE::Editor::DataStorage::RowHandle>& PendingItemsToProcess, EPendingOperationType Type)
+{
+	if (!TedsOutlinerImpl)
+	{
+		return;
+	}
+
+	const double StartTimeSeconds = FPlatformTime::Seconds();
+	double CurrentTimeSpentSeconds = 0.0;
+	int32 CurrentTasksProcessedNum = 0;
+
+	FSceneOutlinerHierarchyChangedData EventData;
+	EventData.Type = Type == EPendingOperationType::Add ? FSceneOutlinerHierarchyChangedData::Added : FSceneOutlinerHierarchyChangedData::Removed;
+
+	for (TSet<UE::Editor::DataStorage::RowHandle>::TIterator PendingItemsIDIterator(PendingItemsToProcess); PendingItemsIDIterator; ++PendingItemsIDIterator)
+	{
+		// Only check the budget every 5 tasks as Getting the current time is a syscall and it is not free
+		if (CurrentTasksProcessedNum % 5 == 0)
+		{
+			CurrentTimeSpentSeconds += FPlatformTime::Seconds() - StartTimeSeconds;
+
+			if (CurrentTimeSpentSeconds > TimeBudgetInSeconds)
+			{
+				break;
+			}
+		}
+
+		if (Type == EPendingOperationType::Add)
+		{
+			EventData.ItemIDs.Add(*PendingItemsIDIterator);
+		}
+		else
+		{
+			EventData.Items.Add(CreateItemFor<UE::Editor::Outliner::FTedsOutlinerTreeItem>(UE::Editor::Outliner::FTedsOutlinerTreeItem(*PendingItemsIDIterator, TedsOutlinerImpl.ToSharedRef())));
+		}
+
+		PendingItemsIDIterator.RemoveCurrent();
+		CurrentTasksProcessedNum++;
+	}
+
+	GetHierarchy()->OnHierarchyChanged().Broadcast(MoveTemp(EventData));
+}
+
+bool FChaosVDWorldOutlinerMode::Tick(float DeltaTime)
+{
+	if (!CanPopulate())
+	{
+		SceneOutliner->SetNextUIRefreshDelay(1.0f);
 	}
+	else
+	{
+		constexpr double TimeBudgetInSeconds = 0.002;
+		ProcessPendingActionsList(TimeBudgetInSeconds, RowsPendingAddition, EPendingOperationType::Add);
+		ProcessPendingActionsList(TimeBudgetInSeconds, RowsPendingRemoval, EPendingOperationType::Remove);
+	}
+	return true;
 }
 
 #undef LOCTEXT_NAMESPACE
diff --git a/Engine/Plugins/ChaosVD/Source/ChaosVD/Private/ChaosVDWorldOutlinerMode.h b/Engine/Plugins/ChaosVD/Source/ChaosVD/Private/ChaosVDWorldOutlinerMode.h
--- a/Engine/Plugins/ChaosVD/Source/ChaosVD/Private/ChaosVDWorldOutlinerMode.h
+++ b/Engine/Plugins/ChaosVD/Source/ChaosVD/Private/ChaosVDWorldOutlinerMode.h
@@ -33,13 +33,12 @@ struct FChaosVDActorOutlinerSelector
 	bool operator()(const TWeakPtr<ISceneOutlinerTreeItem>& Item, AActor*& ActorPtrOut) const;
 };
 
-
 /**
  * Scene outliner mode used to represent a CVD (Chaos Visual Debugger) world
  * It has a more limited view compared to the normal outliner, hiding features we don't support,
  * and it is integrated with the CVD local selection system
  */
-class FChaosVDWorldOutlinerMode : public UE::Editor::Outliner::FTedsOutlinerMode, public FChaosVDSceneSelectionObserver
+class FChaosVDWorldOutlinerMode : public UE::Editor::Outliner::FTedsOutlinerMode, public FChaosVDSceneSelectionObserver, public FTSTickerObjectBase
 {
 public:
 
@@ -78,10 +77,25 @@ public:
 		return SNullWidget::NullWidget;
 	}
 
+	virtual bool Tick(float DeltaTime) override;
+
 private:
 
 	virtual void HandlePostSelectionChange(const UTypedElementSelectionSet* ChangesSelectionSet) override;
 
+	enum class EPendingOperationType
+	{
+		Add,
+		Remove
+	};
+	void ProcessPendingActionsList(double TimeBudgetInSeconds, TSet<UE::Editor::DataStorage::RowHandle>& PendingItemsToProcess, EPendingOperationType Type);
+
+	TSet<UE::Editor::DataStorage::RowHandle> RowsPendingRemoval; 
+	TSet<UE::Editor::DataStorage::RowHandle> RowsPendingAddition;
+
+	UE::Editor::DataStorage::QueryHandle CustomRowAdditionQueryHandle;
+	UE::Editor::DataStorage::QueryHandle CustomRowRemovalQueryHandle;
+
 	TWeakPtr<FChaosVDScene> CVDScene;
 	TWeakPtr<FChaosVDPlaybackController> PlaybackController;
 
diff --git a/Engine/Plugins/ChaosVD/Source/ChaosVD/Private/ChaosVDWorldOutlinerTab.cpp b/Engine/Plugins/ChaosVD/Source/ChaosVD/Private/ChaosVDWorldOutlinerTab.cpp
--- a/Engine/Plugins/ChaosVD/Source/ChaosVD/Private/ChaosVDWorldOutlinerTab.cpp
+++ b/Engine/Plugins/ChaosVD/Source/ChaosVD/Private/ChaosVDWorldOutlinerTab.cpp
@@ -94,7 +94,7 @@ void FChaosVDWorldOutlinerTab::CreateWorldOutlinerWidget()
 
 	Params.SelectionSetOverride = FName("CVDSelection");
 
-	InitOptions.ModeFactory = FCreateSceneOutlinerMode::CreateLambda([this, PlaybackController, &Params, &InitOptions](SSceneOutliner* Outliner)
+	InitOptions.ModeFactory = FCreateSceneOutlinerMode::CreateLambda([this, PlaybackController, &Params](SSceneOutliner* Outliner)
 	{
 		Params.SceneOutliner = Outliner;
 
diff --git a/Engine/Plugins/ChaosVD/Source/ChaosVD/Private/Components/ChaosVDInstancedStaticMeshComponent.cpp b/Engine/Plugins/ChaosVD/Source/ChaosVD/Private/Components/ChaosVDInstancedStaticMeshComponent.cpp
--- a/Engine/Plugins/ChaosVD/Source/ChaosVD/Private/Components/ChaosVDInstancedStaticMeshComponent.cpp
+++ b/Engine/Plugins/ChaosVD/Source/ChaosVD/Private/Components/ChaosVDInstancedStaticMeshComponent.cpp
@@ -4,73 +4,33 @@
 
 #include "ChaosVDGeometryBuilder.h"
 #include "ChaosVDModule.h"
-#include "Materials/MaterialInterface.h"
-#include "Materials/MaterialInstanceDynamic.h"
+#include "Actors/ChaosVDGeometryContainer.h"
 
-void UChaosVDInstancedStaticMeshComponent::UpdateInstanceVisibility(const TSharedPtr<FChaosVDMeshDataInstanceHandle>& InInstanceHandle, bool bIsVisible)
+void UChaosVDInstancedStaticMeshComponent::UpdateVisibilityForInstance(const TSharedRef<FChaosVDInstancedMeshData>& InInstanceHandle)
 {
-	if (!InInstanceHandle.IsValid())
-	{
-		UE_LOG(LogChaosVDEditor, Warning, TEXT("[%s] Attempted to update a mesh instance using an invalid handle. No instances were updated"), ANSI_TO_TCHAR(__FUNCTION__));
-		return;
-	}
-
 	if (InInstanceHandle->GetMeshComponent() != this)
 	{
 		UE_LOG(LogChaosVDEditor, Error, TEXT("[%s] Attempted to update a mesh instance using a handle from another component. No instances were updated | Handle Component [%s] | Current Component [%s]"), ANSI_TO_TCHAR(__FUNCTION__), *GetNameSafe(InInstanceHandle->GetMeshComponent()), *GetNameSafe(this));
 		return;
 	}
 
-	const int32 InstanceIndex = InInstanceHandle->GetMeshInstanceIndex();
-
-	FTransform Transform = InInstanceHandle->GetWorldTransform();
-	if (!bIsVisible)
-	{
-		// Setting the scale to 0 will hide this instance while keeping it on the component
-		Transform.SetScale3D(FVector::ZeroVector);
-	}
-
-	constexpr bool bIsWorldSpaceTransform = true;
-	constexpr bool bMarkRenderDirty = false;
-	constexpr bool bTeleport = true;
-
-	UpdateInstanceTransform(InstanceIndex, Transform, bIsWorldSpaceTransform, bMarkRenderDirty, bTeleport);
-
-	bHadTransformsUpdate = true;
+	EnqueueMeshInstanceOperation(InInstanceHandle, EChaosVDMeshInstanceOperationsFlags::TransformUpdate);
 }
 
-void UChaosVDInstancedStaticMeshComponent::SetIsSelected(const TSharedPtr<FChaosVDMeshDataInstanceHandle>& InInstanceHandle, bool bIsSelected)
+void UChaosVDInstancedStaticMeshComponent::UpdateSelectionStateForInstance(const TSharedRef<FChaosVDInstancedMeshData>& InInstanceHandle)
 {
-	if (!InInstanceHandle.IsValid())
-	{
-		UE_LOG(LogChaosVDEditor, Warning, TEXT("[%s] Attempted to update a mesh instance using an invalid handle. No instances were updated"), ANSI_TO_TCHAR(__FUNCTION__));
-		return;
-	}
-
 	if (InInstanceHandle->GetMeshComponent() != this)
 	{
 		UE_LOG(LogChaosVDEditor, Error, TEXT("[%s] Attempted to update a mesh instance using a handle from another component. No instances were updated | Handle Component [%s] | Current Component [%s]"), ANSI_TO_TCHAR(__FUNCTION__), *GetNameSafe(InInstanceHandle->GetMeshComponent()), *GetNameSafe(this));
 		return;
 	}
 
-	const int32 InstanceIndex = InInstanceHandle->GetMeshInstanceIndex();
-
-	if (!ensure(IsValidInstance(InstanceIndex)))
-	{
-		UE_LOG(LogChaosVDEditor, Error, TEXT("[%s] Attempted to update a mesh instance using a handle with an invalid instance index | Handle Instance Index [%d] | Current Instance Conut [%d]"), ANSI_TO_TCHAR(__FUNCTION__), InstanceIndex, GetInstanceCount());
-		return;
-	}
-
-	NotifySMInstanceSelectionChanged({this, InstanceIndex}, bIsSelected);
+	EnqueueMeshInstanceOperation(InInstanceHandle, EChaosVDMeshInstanceOperationsFlags::SelectionUpdate);
 }
 
-void UChaosVDInstancedStaticMeshComponent::UpdateInstanceColor(const TSharedPtr<FChaosVDMeshDataInstanceHandle>& InInstanceHandle, FLinearColor NewColor)
+void UChaosVDInstancedStaticMeshComponent::UpdateColorForInstance(const TSharedRef<FChaosVDInstancedMeshData>& InInstanceHandle)
 {
-	if (CurrentMaterial == nullptr)
-	{
-		CurrentMaterial = GetMaterial(0);
-	}
-
+	FLinearColor NewColor = InInstanceHandle->GetInstanceColor();
 	// Check that this mesh component supports the intended visualization
 	// We can't change the material of Instanced mesh components because we might have other instances that are not intended to be translucent (or the other way around).
 	// The Mesh handle instance system should have detected we need to migrate the instance to another component before ever reaching this point
@@ -90,60 +50,34 @@ void UChaosVDInstancedStaticMeshComponent::UpdateInstanceColor(const TSharedPtr<
 			UE_LOG(LogChaosVDEditor, Warning, TEXT("[%s] Desired Color is not supported in this mesh component [%s]..."), ANSI_TO_TCHAR(__FUNCTION__), *NewColor.ToString());
 		}
 	}
-	
-	SetNumCustomDataFloats(4);
-	SetCustomData(InInstanceHandle->GetMeshInstanceIndex(), TArrayView<float>(reinterpret_cast<float*>(&NewColor), 4));
+
+	EnqueueMeshInstanceOperation(InInstanceHandle, EChaosVDMeshInstanceOperationsFlags::ColorUpdate);
 }
 
-void UChaosVDInstancedStaticMeshComponent::UpdateInstanceWorldTransform(const TSharedPtr<FChaosVDMeshDataInstanceHandle>& InInstanceHandle, const FTransform& InTransform)
+void UChaosVDInstancedStaticMeshComponent::UpdateWorldTransformForInstance(const TSharedRef<FChaosVDInstancedMeshData>& InInstanceHandle)
 {
-	if (!InInstanceHandle.IsValid())
-	{
-		UE_LOG(LogChaosVDEditor, Warning, TEXT("[%s] Attempted to update a mesh instance using an invalid handle. No instances were updated"), ANSI_TO_TCHAR(__FUNCTION__));
-		return;
-	}
-
 	if (InInstanceHandle->GetMeshComponent() != this)
 	{
 		UE_LOG(LogChaosVDEditor, Error, TEXT("[%s] Attempted to update a mesh instance using a handle from another component. No instances were updated | Handle Component [%s] | Current Component [%s]"), ANSI_TO_TCHAR(__FUNCTION__), *GetNameSafe(InInstanceHandle->GetMeshComponent()), *GetNameSafe(this));
 		return;
 	}
 
-	const int32 InstanceIndex = InInstanceHandle->GetMeshInstanceIndex();
-
-	FTransform NewTransform = InTransform;
-
-	if (!InInstanceHandle->GetVisibility())
-	{
-		NewTransform.SetScale3D(FVector::ZeroVector);
-	}
-
-	constexpr bool bIsWorldSpaceTransform = true;
-	constexpr bool bMarkRenderDirty = false;
-	constexpr bool bTeleport = true;
-
-	UpdateInstanceTransform(InstanceIndex, NewTransform, bIsWorldSpaceTransform, bMarkRenderDirty, bTeleport);
-
-	bHadTransformsUpdate = true;
-}
-
-TArrayView<TSharedPtr<FChaosVDMeshDataInstanceHandle>> UChaosVDInstancedStaticMeshComponent::GetMeshDataInstanceHandles()
-{
-	return CurrentInstanceHandles;
+	EnqueueMeshInstanceOperation(InInstanceHandle, EChaosVDMeshInstanceOperationsFlags::TransformUpdate);
 }
 
 void UChaosVDInstancedStaticMeshComponent::Reset()
 {
 	bIsMeshReady = false;
+	bIsDestroyed = false;
 	MeshReadyDelegate = FChaosVDMeshReadyDelegate();
 	ComponentEmptyDelegate = FChaosVDMeshComponentEmptyDelegate();
-	CurrentMaterial = nullptr;
 
 	SetStaticMesh(nullptr);
 
 	EmptyOverrideMaterials();
 
-	CurrentInstanceHandles.Reset();
+	ClearInstanceHandles();
+
 	CurrentGeometryKey = 0;
 }
 
@@ -157,10 +91,14 @@ void UChaosVDInstancedStaticMeshComponent::Initialize()
 	TSharedPtr<FChaosVDGeometryBuilder> GeometryBuilder = GeometryBuilderWeakPtr.Pin();
 	if (ensure(GeometryBuilder))
 	{
-		if (!HasOverrideMaterials())
+		if (!ensure(!HasOverrideMaterials()))
 		{
-			GeometryBuilder->RequestMaterialUpdate(this);
+			UE_LOG(LogChaosVDEditor, Error, TEXT("[%s] | Component [%s] already had a material, which is not expected!. It likely it was modified after being disposed."), ANSI_TO_TCHAR(__FUNCTION__), *GetNameSafe(this));
+			
+			EmptyOverrideMaterials();
 		}
+
+		GeometryBuilder->RequestMaterialUpdate(this);
 	}
 }
 
@@ -173,7 +111,7 @@ EChaosVDMaterialType UChaosVDInstancedStaticMeshComponent::GetMaterialType() con
 {
 	if (EnumHasAnyFlags(MeshComponentAttributeFlags, EChaosVDMeshAttributesFlags::TranslucentGeometry))
 	{
-		return  EChaosVDMaterialType::ISMCTranslucent;
+		return EChaosVDMaterialType::ISMCTranslucent;
 	}
 
 	return EChaosVDMaterialType::ISMCOpaque;
@@ -182,6 +120,8 @@ EChaosVDMaterialType UChaosVDInstancedStaticMeshComponent::GetMaterialType() con
 void UChaosVDInstancedStaticMeshComponent::OnDisposed()
 {
 	Reset();
+
+	bIsDestroyed = true;
 	
 	SetRelativeTransform(FTransform::Identity);
 
@@ -194,6 +134,14 @@ void UChaosVDInstancedStaticMeshComponent::OnDisposed()
 	{
 		Owner->RemoveOwnedComponent(this);
 	}
+
+	FTSTicker::GetCoreTicker().RemoveTicker(ExternalTickerHandle);
+}
+
+bool UChaosVDInstancedStaticMeshComponent::ExternalTick(float DeltaTime)
+{
+	ProcessChanges();
+	return true;
 }
 
 bool UChaosVDInstancedStaticMeshComponent::UpdateGeometryKey(const uint32 NewHandleGeometryKey)
@@ -213,13 +161,9 @@ bool UChaosVDInstancedStaticMeshComponent::UpdateGeometryKey(const uint32 NewHan
 	return true;
 }
 
-TSharedPtr<FChaosVDMeshDataInstanceHandle> UChaosVDInstancedStaticMeshComponent::AddMeshInstance(const FTransform InstanceTransform, bool bIsWorldSpace, const TSharedRef<FChaosVDExtractedGeometryDataHandle>& InGeometryHandle, int32 ParticleID, int32 SolverID)
+TSharedPtr<FChaosVDInstancedMeshData> UChaosVDInstancedStaticMeshComponent::AddMeshInstance(const FTransform InstanceTransform, bool bIsWorldSpace, const TSharedRef<FChaosVDExtractedGeometryDataHandle>& InGeometryHandle, int32 ParticleID, int32 SolverID)
 {
-	int32 InstanceIndex = AddInstance(InstanceTransform, bIsWorldSpace);
-
-	check(InstanceIndex != INDEX_NONE);
-
-	const TSharedPtr<FChaosVDMeshDataInstanceHandle> InstanceHandle = MakeShared<FChaosVDMeshDataInstanceHandle>(InstanceIndex, this, ParticleID, SolverID, InGeometryHandle);
+	const TSharedPtr<FChaosVDInstancedMeshData> InstanceHandle = MakeShared<FChaosVDInstancedMeshData>(INDEX_NONE, this, ParticleID, SolverID, InGeometryHandle);
 
 	const uint32 NewHandleGeometryKey = InGeometryHandle->GetGeometryKey();
 
@@ -227,81 +171,124 @@ TSharedPtr<FChaosVDMeshDataInstanceHandle> UChaosVDInstancedStaticMeshComponent:
 	{
 		return nullptr;
 	}
-	
-	CurrentInstanceHandles.Add(InstanceHandle);
+
+	EnqueueMeshInstanceOperation(InstanceHandle.ToSharedRef(), EChaosVDMeshInstanceOperationsFlags::Add);
 
 	return InstanceHandle;
 }
 
-void UChaosVDInstancedStaticMeshComponent::AddMeshInstanceForHandle(TSharedPtr<FChaosVDMeshDataInstanceHandle> MeshDataHandle, const FTransform InstanceTransform, bool bIsWorldSpace, const TSharedRef<FChaosVDExtractedGeometryDataHandle>& InGeometryHandle, int32 ParticleID, int32 SolverID)
+void UChaosVDInstancedStaticMeshComponent::AddExistingMeshInstance(const TSharedRef<FChaosVDInstancedMeshData>& InMeshDataHandle)
 {
-	const int32 InstanceIndex = AddInstance(InstanceTransform, bIsWorldSpace);
-
-	check(InstanceIndex != INDEX_NONE);
-
-	const uint32 NewHandleGeometryKey = InGeometryHandle->GetGeometryKey();
+	const uint32 NewHandleGeometryKey = InMeshDataHandle->ExtractedGeometryHandle->GetGeometryKey();
 
 	if (!UpdateGeometryKey(NewHandleGeometryKey))
 	{
 		return;
 	}
 
-	MeshDataHandle->SetMeshInstanceIndex(InstanceIndex);
-	MeshDataHandle->SetMeshComponent(this);
-	
-	CurrentInstanceHandles.Add(MeshDataHandle);
+	InMeshDataHandle->SetMeshInstanceIndex(INDEX_NONE);
+	InMeshDataHandle->SetMeshComponent(this);
+
+	EnqueueMeshInstanceOperation(InMeshDataHandle, EChaosVDMeshInstanceOperationsFlags::Add);
 }
 
-void UChaosVDInstancedStaticMeshComponent::RemoveMeshInstance(TSharedPtr<FChaosVDMeshDataInstanceHandle> InHandleToRemove)
+void UChaosVDInstancedStaticMeshComponent::RemoveMeshInstance(const TSharedRef<FChaosVDInstancedMeshData>& InHandleToRemove, ERemovalMode Mode)
 {
-	if (!InHandleToRemove.IsValid())
-	{
-		UE_LOG(LogChaosVDEditor, Warning, TEXT("[%s] Attempted to remove a mesh instace using an invalid handle. No instanced were removed"), ANSI_TO_TCHAR(__FUNCTION__));
-		return;
-	}
-
 	if (InHandleToRemove->GetMeshComponent() != this)
 	{
 		UE_LOG(LogChaosVDEditor, Error, TEXT("[%s] Attempted to remove a mesh instace using a handle from another component. No instanced were removed | Handle Component [%s] | Current Component [%s]"), ANSI_TO_TCHAR(__FUNCTION__), *GetNameSafe(InHandleToRemove->GetMeshComponent()), *GetNameSafe(this));
 		return;
 	}
 
-	if (!ensure(InHandleToRemove->GetMeshInstanceIndex() != INDEX_NONE))
+	if (Mode == ERemovalMode::Deferred)
 	{
-		return;
+		EnqueueMeshInstanceOperation(InHandleToRemove, EChaosVDMeshInstanceOperationsFlags::Remove);
 	}
-
+	else
 	{
-		const int32 CurrentInstanceCount =  GetInstanceCount();
+		PendingOperationsByInstance.Remove(InHandleToRemove);
+
+		int32 InstanceIndex = InHandleToRemove->GetMeshInstanceIndex();
 
-		if (InHandleToRemove->GetMeshInstanceIndex() > CurrentInstanceCount - 1)
+		if (InstanceIndex == INDEX_NONE)
 		{
-			ensure(false);
+			// The mesh instance wasn't added yet
+			return;
+		}
 
-			UE_LOG(LogChaosVDEditor, Error, TEXT("[%s] Mesh Data Handle has an invalid instance index. No instanced were removed | Handle Instance Index [%d] | Current Instance Count [%d]"), ANSI_TO_TCHAR(__FUNCTION__), InHandleToRemove->GetMeshInstanceIndex(), CurrentInstanceCount);
+		if (!ensure(!InHandleToRemove->IsPendingDestroy()))
+		{
 			return;
 		}
+		
+		{
+			const int32 CurrentInstanceCount = GetInstanceCount();
+			if (!ensure(IsValidInstance(InstanceIndex)))
+			{
+				UE_LOG(LogChaosVDEditor, Error, TEXT("[%s] Mesh Data Handle has an invalid instance index. No instanced were removed | Handle Instance Index [%d] | Current Instance Count [%d]"), ANSI_TO_TCHAR(__FUNCTION__), InHandleToRemove->GetMeshInstanceIndex(), CurrentInstanceCount);
+				return;
+			}
+		}
+
+		RemoveInstance(InstanceIndex);
+		InHandleToRemove->SetMeshInstanceIndex(INDEX_NONE);
 	}
+}
 
-	RemoveInstance(InHandleToRemove->GetMeshInstanceIndex());
+void UChaosVDInstancedStaticMeshComponent::UpdateInstanceHandle(int32 OldIndex, int32 NewIndex)
+{
+	TSharedPtr<FChaosVDInstancedMeshData> RelocatedHandle;
+	ensure(CurrentInstanceHandlesByIndex.RemoveAndCopyValue(OldIndex, RelocatedHandle));
 
-	auto FindPredicate = [InHandleToRemove](const TSharedPtr<FChaosVDMeshDataInstanceHandle>& InstanceHandle)
+	if (!ensure(RelocatedHandle->GetMeshComponent() == this))
 	{
-		return InstanceHandle->GetMeshInstanceIndex() == InHandleToRemove->GetMeshInstanceIndex();
-	};
+		UE_LOG(LogChaosVDEditor, Error, TEXT("[%s] Trying to update an instance from another component | Handle Component [%s] | This component [%s]"), ANSI_TO_TCHAR(__FUNCTION__), *GetNameSafe(RelocatedHandle->GetMeshComponent()), *GetNameSafe(this));
+	}
+
+    if (RelocatedHandle)
+    {
+        RelocatedHandle->SetMeshInstanceIndex(NewIndex);
 
-	// TODO: Evaluate if keeping the instance handles array sorted on any change is worth the cost to allow a binary search here. Same question for a Set
-	const int32 InstanceHandleIndex = CurrentInstanceHandles.IndexOfByPredicate(FindPredicate);
+        if (NewIndex != INDEX_NONE)
+        {
+        	CurrentInstanceHandlesByIndex.Add(NewIndex, RelocatedHandle);
+        }
+    }
+}
 
-	if (InstanceHandleIndex != INDEX_NONE)
+void UChaosVDInstancedStaticMeshComponent::HandleInstanceIndexUpdated(TArrayView<const FInstancedStaticMeshDelegates::FInstanceIndexUpdateData> InIndexUpdates)
+{
+	bool bComponentCleared = false;
+	for (const FInstancedStaticMeshDelegates::FInstanceIndexUpdateData& IndexUpdateData : InIndexUpdates)
 	{
-		CurrentInstanceHandles.RemoveAtSwap(InstanceHandleIndex);
-	}
+		switch (IndexUpdateData.Type)
+		{
+		case FInstancedStaticMeshDelegates::EInstanceIndexUpdateType::Added:
+			break; // We don't need to process 'Added' updates as they can't affect existing IDs
+		case FInstancedStaticMeshDelegates::EInstanceIndexUpdateType::Relocated:
+			{
+				UpdateInstanceHandle(IndexUpdateData.OldIndex, IndexUpdateData.Index);
+				break;
+			}
+		case FInstancedStaticMeshDelegates::EInstanceIndexUpdateType::Removed:
+			{
+				UpdateInstanceHandle(IndexUpdateData.Index, INDEX_NONE);
+				break;
+			}
+		case FInstancedStaticMeshDelegates::EInstanceIndexUpdateType::Destroyed:
+		case FInstancedStaticMeshDelegates::EInstanceIndexUpdateType::Cleared:
+			{
+				bComponentCleared = true;
+				break;
+			}
+		default:
+			break;
+		}
+	};
 
-	// If the component becomes "empty", allow it to go back to the pool
-	if (GetInstanceCount() == 0)
+	if (bComponentCleared)
 	{
-		ComponentEmptyDelegate.Broadcast(this);
+		ClearInstanceHandles();
 	}
 }
 
@@ -316,33 +303,225 @@ bool UChaosVDInstancedStaticMeshComponent::IsNavigationRelevant() const
 	return false;
 }
 
-void UChaosVDInstancedStaticMeshComponent::TickComponent(float DeltaTime, enum ELevelTick TickType, FActorComponentTickFunction* ThisTickFunction)
+void UChaosVDInstancedStaticMeshComponent::ClearInstanceHandles()
+{
+	CurrentInstanceHandlesByIndex.Reset();
+	PendingOperationsByInstance.Reset();
+}
+
+void UChaosVDInstancedStaticMeshComponent::OnAcquired()
 {
-	Super::TickComponent(DeltaTime, TickType, ThisTickFunction);
+	bIsDestroyed = false;
+
+	ExternalTickerHandle = FTSTicker::GetCoreTicker().AddTicker(FTickerDelegate::CreateUObject(this, &UChaosVDInstancedStaticMeshComponent::ExternalTick));
+}
+
+void UChaosVDInstancedStaticMeshComponent::EnqueueMeshInstanceOperation(const TSharedRef<FChaosVDInstancedMeshData>& MeshInstanceHandle, EChaosVDMeshInstanceOperationsFlags Operation)
+{
+	if (!ensure(!bIsDestroyed))
+	{
+		return;
+	}
+	
+	if (!ensure(!MeshInstanceHandle->IsPendingDestroy()))
+	{
+		return;
+	}
 
-	if (bHadTransformsUpdate)
+	if (EChaosVDMeshInstanceOperationsFlags* PendingOperation = PendingOperationsByInstance.Find(MeshInstanceHandle))
+	{
+		if (EnumHasAnyFlags(Operation, EChaosVDMeshInstanceOperationsFlags::Remove))
+		{
+			const int32 CurrentInstanceCount = GetInstanceCount();
+
+			if (CurrentInstanceCount == 0)
+			{
+				// If the instance count is 0 it means we have pending operations but we didn't process them, therefore there is nothing to do
+				PendingOperationsByInstance.Remove(MeshInstanceHandle);
+			}
+			else
+			{
+				if (EnumHasAnyFlags(*PendingOperation, EChaosVDMeshInstanceOperationsFlags::Add))
+				{
+					ensure(MeshInstanceHandle->GetMeshInstanceIndex() == INDEX_NONE);
+					PendingOperationsByInstance.Remove(MeshInstanceHandle);
+					return;
+				}
+
+				// If this new operation is a removal, discard all other operations
+				*PendingOperation = Operation;
+			}
+		}
+		else if (EnumHasAnyFlags(Operation, EChaosVDMeshInstanceOperationsFlags::Add))
+		{
+			EnumRemoveFlags(*PendingOperation, EChaosVDMeshInstanceOperationsFlags::Remove);
+		}
+		else
+		{
+			EnumRemoveFlags(*PendingOperation, EChaosVDMeshInstanceOperationsFlags::Remove);
+			EnumAddFlags(*PendingOperation, Operation);
+		}
+	}
+	else
 	{
-		MarkRenderStateDirty();
-		bHadTransformsUpdate = false;
+		PendingOperationsByInstance.Add(MeshInstanceHandle, Operation);
 	}
 }
 
-uint32 UChaosVDInstancedStaticMeshComponent::GetGeometryKey() const
+void UChaosVDInstancedStaticMeshComponent::CancelMeshInstanceOperation(const TSharedRef<FChaosVDInstancedMeshData>& MeshInstanceHandle, EChaosVDMeshInstanceOperationsFlags Operation)
 {
-	return CurrentGeometryKey;
+	if (EChaosVDMeshInstanceOperationsFlags* PendingOperation = PendingOperationsByInstance.Find(MeshInstanceHandle))
+	{
+		EnumRemoveFlags(*PendingOperation, Operation);
+
+		if (*PendingOperation == EChaosVDMeshInstanceOperationsFlags::None)
+		{
+			PendingOperationsByInstance.Remove(MeshInstanceHandle);
+		}
+	}
 }
 
-TSharedPtr<FChaosVDMeshDataInstanceHandle> UChaosVDInstancedStaticMeshComponent::GetMeshDataInstanceHandle(int32 InstanceIndex) const
+bool UChaosVDInstancedStaticMeshComponent::CanExecuteOperation(EChaosVDMeshInstanceOperationsFlags Operation, EChaosVDMeshInstanceOperationsFlags PendingOperations)
 {
-	auto FindPredicate = [InstanceIndex](const TSharedPtr<FChaosVDMeshDataInstanceHandle>& InstanceHandle)
+	if (!EnumHasAnyFlags(PendingOperations, Operation))
 	{
-		return InstanceHandle->GetMeshInstanceIndex() == InstanceIndex;
-	};
+		return false;
+	}
+
+	if (EnumHasAnyFlags(Operation, EChaosVDMeshInstanceOperationsFlags::ColorUpdate) ||
+		EnumHasAnyFlags(Operation, EChaosVDMeshInstanceOperationsFlags::SelectionUpdate) ||
+		EnumHasAnyFlags(Operation, EChaosVDMeshInstanceOperationsFlags::TransformUpdate))
+	{
+		return !EnumHasAnyFlags(PendingOperations, EChaosVDMeshInstanceOperationsFlags::Add);
+	}
+
+	return true;
+}
 
-	if (const TSharedPtr<FChaosVDMeshDataInstanceHandle>* FoundHandle = CurrentInstanceHandles.FindByPredicate(FindPredicate))
+void UChaosVDInstancedStaticMeshComponent::ProcessChanges()
+{
+	if (bIsMeshReady)
+	{
+		bool bHasSelectionChange = false;
+		for (TMap<TSharedPtr<FChaosVDInstancedMeshData>, EChaosVDMeshInstanceOperationsFlags>::TIterator HandleOperationsIterator = PendingOperationsByInstance.CreateIterator(); HandleOperationsIterator; ++HandleOperationsIterator)
+		{
+			TPair<TSharedPtr<FChaosVDInstancedMeshData>, EChaosVDMeshInstanceOperationsFlags>& OperationHandle = *HandleOperationsIterator;
+			
+			if (CanExecuteOperation(EChaosVDMeshInstanceOperationsFlags::Add, OperationHandle.Value))
+			{
+				int32 InstanceIndex = AddInstance(OperationHandle.Key->GetWorldTransform(), true);
+				OperationHandle.Key->SetMeshInstanceIndex(InstanceIndex);
+
+				CurrentInstanceHandlesByIndex.Add(InstanceIndex, OperationHandle.Key);
+
+				EnumRemoveFlags(OperationHandle.Value, EChaosVDMeshInstanceOperationsFlags::Add);
+			}
+			
+			if (CanExecuteOperation(EChaosVDMeshInstanceOperationsFlags::TransformUpdate, OperationHandle.Value))
+			{
+				constexpr bool bIsWorldSpaceTransform = true;
+				constexpr bool bMarkRenderDirty = true;
+				constexpr bool bTeleport = true;
+
+				FTransform Transform = OperationHandle.Key->GetWorldTransform();
+				if (!OperationHandle.Key->GetVisibility())
+				{
+					// Setting the scale to 0 will hide this instance while keeping it on the component
+					Transform.SetScale3D(FVector::ZeroVector);
+				}
+	
+				UpdateInstanceTransform(OperationHandle.Key->GetMeshInstanceIndex(), Transform, bIsWorldSpaceTransform, bMarkRenderDirty, bTeleport);
+
+				EnumRemoveFlags(OperationHandle.Value, EChaosVDMeshInstanceOperationsFlags::TransformUpdate);
+			}
+
+			if (CanExecuteOperation(EChaosVDMeshInstanceOperationsFlags::ColorUpdate, OperationHandle.Value))
+			{
+				SetNumCustomDataFloats(4);
+				FLinearColor NewColor = OperationHandle.Key->GetInstanceColor();
+				SetCustomData(OperationHandle.Key->GetMeshInstanceIndex(), TArrayView<float>(reinterpret_cast<float*>(&NewColor), 4));
+
+				EnumRemoveFlags(OperationHandle.Value, EChaosVDMeshInstanceOperationsFlags::ColorUpdate);
+			}
+
+			if (CanExecuteOperation(EChaosVDMeshInstanceOperationsFlags::SelectionUpdate, OperationHandle.Value))
+			{
+				const int32 InstanceIndex = OperationHandle.Key->GetMeshInstanceIndex();
+
+				if (!ensure(IsValidInstance(InstanceIndex)))
+				{
+					UE_LOG(LogChaosVDEditor, Error, TEXT("[%s] Attempted to update a mesh instance using a handle with an invalid instance index | Handle Instance Index [%d] | Current Instance Conut [%d]"), ANSI_TO_TCHAR(__FUNCTION__), InstanceIndex, GetInstanceCount());
+					EnumRemoveFlags(OperationHandle.Value, EChaosVDMeshInstanceOperationsFlags::SelectionUpdate);
+				}
+				else
+				{
+					NotifySMInstanceSelectionChanged({this, InstanceIndex}, OperationHandle.Key->InstanceState.bIsSelected);
+					EnumRemoveFlags(OperationHandle.Value, EChaosVDMeshInstanceOperationsFlags::SelectionUpdate);
+					bHasSelectionChange = true;
+				}
+			}
+
+			if (CanExecuteOperation(EChaosVDMeshInstanceOperationsFlags::Remove, OperationHandle.Value))
+			{
+				int32 InstanceIndex = OperationHandle.Key->GetMeshInstanceIndex();
+				const int32 CurrentInstanceCount = GetInstanceCount();
+
+				if (!ensure(IsValidInstance(InstanceIndex)))
+				{
+					UE_LOG(LogChaosVDEditor, Error, TEXT("[%s] Mesh Data Handle has an invalid instance index. No instanced were removed | Handle Instance Index [%d] | Current Instance Count [%d]"), ANSI_TO_TCHAR(__FUNCTION__), OperationHandle.Key->GetMeshInstanceIndex(), CurrentInstanceCount);
+				}
+				else
+				{
+					RemoveInstance(InstanceIndex);
+					OperationHandle.Key->SetMeshInstanceIndex(INDEX_NONE);
+				}
+
+				// We just removed this instance, clear all flags
+				OperationHandle.Value = EChaosVDMeshInstanceOperationsFlags::None;
+			}
+
+			if (OperationHandle.Value == EChaosVDMeshInstanceOperationsFlags::None)
+			{
+				HandleOperationsIterator.RemoveCurrent();
+			}
+		}
+
+		if (bHasSelectionChange)
+		{
+			if (AChaosVDGeometryContainer* GeometryContainer = Cast<AChaosVDGeometryContainer>(GetOwner()))
+			{
+				if (TSharedPtr<FChaosVDScene> ScenePtr = GeometryContainer->GetScene())
+				{
+					ScenePtr->RequestUpdate();
+				}
+			}
+		}
+
+		if (CurrentInstanceHandlesByIndex.Num() == 0 && PendingOperationsByInstance.Num() == 0)
+		{
+			ensure(GetInstanceCount() == 0);
+
+			ComponentEmptyDelegate.Broadcast(this);
+		}
+	}
+}
+
+uint32 UChaosVDInstancedStaticMeshComponent::GetGeometryKey() const
+{
+	return CurrentGeometryKey;
+}
+
+TSharedPtr<FChaosVDInstancedMeshData> UChaosVDInstancedStaticMeshComponent::GetMeshDataInstanceHandle(int32 InstanceIndex) const
+{
+	if (const TSharedPtr<FChaosVDInstancedMeshData>* FoundHandle = CurrentInstanceHandlesByIndex.Find(InstanceIndex))
 	{
 		return *FoundHandle;
 	}
 
 	return nullptr;
 }
+
+void UChaosVDInstancedStaticMeshComponent::SetIsMeshReady(bool bIsReady)
+{
+	bIsMeshReady = bIsReady;
+}
diff --git a/Engine/Plugins/ChaosVD/Source/ChaosVD/Private/Components/ChaosVDInstancedStaticMeshComponent.h b/Engine/Plugins/ChaosVD/Source/ChaosVD/Private/Components/ChaosVDInstancedStaticMeshComponent.h
--- a/Engine/Plugins/ChaosVD/Source/ChaosVD/Private/Components/ChaosVDInstancedStaticMeshComponent.h
+++ b/Engine/Plugins/ChaosVD/Source/ChaosVD/Private/Components/ChaosVDInstancedStaticMeshComponent.h
@@ -4,10 +4,24 @@
 
 #include "Components/InstancedStaticMeshComponent.h"
 #include "ChaosVDGeometryDataComponent.h"
+#include "Components/HierarchicalInstancedStaticMeshComponent.h"
+#include "Containers/Ticker.h"
 #include "Interfaces/ChaosVDPooledObject.h"
 #include "ChaosVDInstancedStaticMeshComponent.generated.h"
 
-/** CVD version of a Instance Static Mesh Component that holds additional CVD data */
+UENUM()
+enum class EChaosVDMeshInstanceOperationsFlags
+{
+	None = 0,
+	Add = 1 << 0,
+	Remove = 1 << 1,
+	ColorUpdate = 1 << 2,
+	SelectionUpdate = 1 << 3,
+	TransformUpdate = 1 << 4,
+};
+ENUM_CLASS_FLAGS(EChaosVDMeshInstanceOperationsFlags);
+
+/** CVD version of an Instance Static Mesh Component that holds additional CVD data */
 UCLASS(HideCategories=("Transform"), MinimalAPI)
 class UChaosVDInstancedStaticMeshComponent : public UInstancedStaticMeshComponent, public IChaosVDGeometryComponent, public IChaosVDPooledObject
 {
@@ -24,34 +38,62 @@ public:
 		WireframeColorOverride = FColor::White;
 	}
 
-	virtual void TickComponent(float DeltaTime, enum ELevelTick TickType, FActorComponentTickFunction* ThisTickFunction) override;
+	void EnqueueMeshInstanceOperation(const TSharedRef<FChaosVDInstancedMeshData>& MeshInstanceHandle, EChaosVDMeshInstanceOperationsFlags Operation);
+	void CancelMeshInstanceOperation(const TSharedRef<FChaosVDInstancedMeshData>& MeshInstanceHandle, EChaosVDMeshInstanceOperationsFlags Operation);
+
+	bool CanExecuteOperation(EChaosVDMeshInstanceOperationsFlags Operation, EChaosVDMeshInstanceOperationsFlags PendingOperations);
+
+	void ProcessChanges();
 	
 	// BEGIN IChaosVDGeometryDataComponent Interface
 	
 	virtual uint32 GetGeometryKey() const override;
 
-	virtual TSharedPtr<FChaosVDMeshDataInstanceHandle> GetMeshDataInstanceHandle(int32 InstanceIndex) const override;
+	virtual TSharedPtr<FChaosVDInstancedMeshData> GetMeshDataInstanceHandle(int32 InstanceIndex) const override;
 
-	virtual bool IsMeshReady() const override { return bIsMeshReady; }
+	virtual bool IsMeshReady() const override
+	{
+		return bIsMeshReady;
+	}
 	
-	virtual void SetIsMeshReady(bool bIsReady) override { bIsMeshReady = bIsReady; }
+	virtual void SetIsMeshReady(bool bIsReady) override;
 
-	virtual FChaosVDMeshReadyDelegate* OnMeshReady() override { return &MeshReadyDelegate; }
+	virtual FChaosVDMeshReadyDelegate* OnMeshReady() override
+	{
+		return &MeshReadyDelegate;
+	}
 
-	virtual FChaosVDMeshComponentEmptyDelegate* OnComponentEmpty() override { return &ComponentEmptyDelegate; }
+	virtual FChaosVDMeshComponentEmptyDelegate* OnComponentEmpty() override
+	{
+		return &ComponentEmptyDelegate;
+	}
 
-	virtual void UpdateInstanceVisibility(const TSharedPtr<FChaosVDMeshDataInstanceHandle>& InInstanceHandle, bool bIsVisible) override;
+	virtual void UpdateVisibilityForInstance(const TSharedRef<FChaosVDInstancedMeshData>& InInstanceHandle) override;
 
-	virtual void SetIsSelected(const TSharedPtr<FChaosVDMeshDataInstanceHandle>& InInstanceHandle, bool bIsSelected) override;
+	virtual void UpdateSelectionStateForInstance(const TSharedRef<FChaosVDInstancedMeshData>& InInstanceHandle) override;
 
-	virtual void UpdateInstanceColor(const TSharedPtr<FChaosVDMeshDataInstanceHandle>& InInstanceHandle, FLinearColor NewColor) override;
+	virtual void UpdateColorForInstance(const TSharedRef<FChaosVDInstancedMeshData>& InInstanceHandle) override;
 	
-	virtual void UpdateInstanceWorldTransform(const TSharedPtr<FChaosVDMeshDataInstanceHandle>& InInstanceHandle, const FTransform& InTransform) override;
+	virtual void UpdateWorldTransformForInstance(const TSharedRef<FChaosVDInstancedMeshData>& InInstanceHandle) override;
 
-	virtual void SetMeshComponentAttributeFlags(EChaosVDMeshAttributesFlags Flags) override { MeshComponentAttributeFlags = Flags; }
-	virtual EChaosVDMeshAttributesFlags GetMeshComponentAttributeFlags() const override { return MeshComponentAttributeFlags; };
+	virtual void SetMeshComponentAttributeFlags(EChaosVDMeshAttributesFlags Flags) override
+	{
+		MeshComponentAttributeFlags = Flags;
+	}
+	virtual EChaosVDMeshAttributesFlags GetMeshComponentAttributeFlags() const override
+	{
+		return MeshComponentAttributeFlags;
+	}
 
-	virtual TArrayView<TSharedPtr<FChaosVDMeshDataInstanceHandle>> GetMeshDataInstanceHandles() override;
+	virtual bool GetIsDestroyed() const override
+	{
+		return bIsDestroyed;
+	}
+
+	virtual void SetIsDestroyed(bool bNewIsDestroyed) override
+	{
+		bIsDestroyed = bNewIsDestroyed;
+	}
 
 	virtual void Reset() override;
 
@@ -59,22 +101,31 @@ public:
 
 	virtual void SetGeometryBuilder(TWeakPtr<FChaosVDGeometryBuilder> GeometryBuilder) override;
 	
-	virtual EChaosVDMaterialType GetMaterialType() const;
+	virtual EChaosVDMaterialType GetMaterialType() const override;
 	// END IChaosVDGeometryDataComponent Interface
 
-	virtual TSharedPtr<FChaosVDMeshDataInstanceHandle> AddMeshInstance(const FTransform InstanceTransform, bool bIsWorldSpace, const TSharedRef<FChaosVDExtractedGeometryDataHandle>& InGeometryHandle, int32 ParticleID, int32 SolverID) override;
-	virtual void AddMeshInstanceForHandle(TSharedPtr<FChaosVDMeshDataInstanceHandle> MeshDataHandle, const FTransform InstanceTransform, bool bIsWorldSpace, const TSharedRef<FChaosVDExtractedGeometryDataHandle>& InGeometryHandle, int32 ParticleID, int32 SolverID) override;
-	virtual void RemoveMeshInstance(TSharedPtr<FChaosVDMeshDataInstanceHandle> InHandleToRemove) override;
+	virtual TSharedPtr<FChaosVDInstancedMeshData> AddMeshInstance(const FTransform InstanceTransform, bool bIsWorldSpace, const TSharedRef<FChaosVDExtractedGeometryDataHandle>& InGeometryHandle, int32 ParticleID, int32 SolverID) override;
+	virtual void AddExistingMeshInstance(const TSharedRef<FChaosVDInstancedMeshData>& InMeshDataHandle) override;
+	virtual void RemoveMeshInstance(const TSharedRef<FChaosVDInstancedMeshData>& InHandleToRemove, ERemovalMode Mode = ERemovalMode::Deferred) override;
+
+	void UpdateInstanceHandle(int32 OldIndex, int32 NewIndex);
+
+	/** Handles a mesh instance index update reported by the mesh component used to render this mesh instance */
+	void HandleInstanceIndexUpdated(TArrayView<const FInstancedStaticMeshDelegates::FInstanceIndexUpdateData> InIndexUpdates);
 
 	virtual bool Modify(bool bAlwaysMarkDirty) override;
 
 	virtual bool IsNavigationRelevant() const override;
 
+	void ClearInstanceHandles();
+
 	// BEGIN IChaosVDPooledObject Interface
-	virtual void OnAcquired() override{};
+	virtual void OnAcquired() override;
 	virtual void OnDisposed() override;
 	// END IChaosVDPooledObject Interface
 
+	bool ExternalTick(float DeltaTime);
+
 protected:
 	
 	bool UpdateGeometryKey(uint32 NewHandleGeometryKey);
@@ -83,20 +134,17 @@ protected:
 
 	uint32 CurrentGeometryKey = 0;
 
-	bool bIsMeshReady = false;
-
-	UPROPERTY(Transient)
-	TObjectPtr<UMaterialInterface> CurrentMaterial = nullptr;
-
-	UPROPERTY(Transient)
-	TMap<EChaosVDMaterialType, TObjectPtr<UMaterialInstanceDynamic>> CachedMaterialInstancesByID;
+	uint8 bIsMeshReady:1 = false;	
+	uint8 bIsDestroyed:1 = false;
 
 	FChaosVDMeshReadyDelegate MeshReadyDelegate;
 	FChaosVDMeshComponentEmptyDelegate ComponentEmptyDelegate;
-
-	TArray<TSharedPtr<FChaosVDMeshDataInstanceHandle>> CurrentInstanceHandles;
 	
 	TWeakPtr<FChaosVDGeometryBuilder> GeometryBuilderWeakPtr;
 
-	bool bHadTransformsUpdate = false;
+	TMap<TSharedPtr<FChaosVDInstancedMeshData>, EChaosVDMeshInstanceOperationsFlags> PendingOperationsByInstance;
+
+	TMap<int32, TSharedPtr<FChaosVDInstancedMeshData>> CurrentInstanceHandlesByIndex;
+
+	FTSTicker::FDelegateHandle ExternalTickerHandle;
 };
diff --git a/Engine/Plugins/ChaosVD/Source/ChaosVD/Private/Components/ChaosVDParticleDataComponent.cpp b/Engine/Plugins/ChaosVD/Source/ChaosVD/Private/Components/ChaosVDParticleDataComponent.cpp
--- a/Engine/Plugins/ChaosVD/Source/ChaosVD/Private/Components/ChaosVDParticleDataComponent.cpp
+++ b/Engine/Plugins/ChaosVD/Source/ChaosVD/Private/Components/ChaosVDParticleDataComponent.cpp
@@ -15,6 +15,7 @@
 #include "Elements/Interfaces/TypedElementDataStorageInterface.h"
 #include "Misc/ScopedSlowTask.h"
 #include "Settings/ChaosVDParticleVisualizationSettings.h"
+#include "TEDS/ChaosVDParticleEditorDataFactory.h"
 #include "TEDS/ChaosVDSelectionInterface.h"
 #include "TEDS/ChaosVDStructTypedElementData.h"
 #include "TEDS/ChaosVDTedsUtils.h"
@@ -48,6 +49,15 @@ UChaosVDParticleDataComponent::UChaosVDParticleDataComponent()
 	{
 		ColorVisualizationSettings->OnSettingsChanged().AddUObject(this, &UChaosVDParticleDataComponent::HandleColorsSettingsUpdated);
 	}
+
+	StreamingSystem.SetStreamingDataSource(this);
+}
+
+void UChaosVDParticleDataComponent::SetScene(const TWeakPtr<FChaosVDScene>& InSceneWeakPtr)
+{
+	Super::SetScene(InSceneWeakPtr);
+	
+	StreamingSystem.SetScene(InSceneWeakPtr);
 }
 
 void UChaosVDParticleDataComponent::ClearData()
@@ -67,6 +77,9 @@ void UChaosVDParticleDataComponent::ClearData()
 	{
 		ScenePtr->ClearSelectionAndNotify();
 	}
+
+	StreamingSystem.Reset();
+	StreamingSystem.SetStreamingDataSource(nullptr);
 	
 	{
 		FReadScopeLock ReadLock(ParticleSceneObjectsLock);
@@ -83,9 +96,11 @@ void UChaosVDParticleDataComponent::ClearData()
 			Chaos::VD::TedsUtils::RemoveObjectToFromDataStorage(ContainerWithID.Value.Get());
 		}
 	}
+
 	{
 		FWriteScopeLock WriteScopeLock(ParticleSceneObjectsLock);
 		SolverParticlesByID.Reset();
+		SolverParticlesArray.Reset();
 	}
 }
 
@@ -107,13 +122,19 @@ FChaosVDSceneParticle* UChaosVDParticleDataComponent::GetSelectedParticle() cons
 	return SelectedParticlesHandles.Num() > 0 ? GetStructDataFromTypedElementHandle<FChaosVDSceneParticle>(SelectedParticlesHandles[0]) : nullptr;
 }
 
+void UChaosVDParticleDataComponent::HandleWorldStreamingLocationUpdated(const FVector& InLocation)
+{
+	StreamingSystem.UpdateStreamingSourceLocation(InLocation);
+}
+
 int32 UChaosVDParticleDataComponent::GetMaxElementCountForBatchTEDSUpdate_AssumesLocked() const
 {
 	// All Particles + their container objects + this component owner
 	return SolverParticlesByID.Num() + ParticleSceneContainersByType.Num() + 1;
 }
 
-void UChaosVDParticleDataComponent::ProcessRemovedParticles_AssumesLocked(const TSharedRef<FChaosVDScene>& SceneRef, const FChaosVDSolverFrameData& InSolverFrameData, const FChaosVDFrameStageData& InSolverFrameStageData, const TSet<int32>& ParticlesUpdatedIDsForKeyFrameDiff)
+void UChaosVDParticleDataComponent::ProcessRemovedParticles_AssumessLocked(const TSharedRef<FChaosVDScene>& SceneRef, const FChaosVDSolverFrameData& InSolverFrameData,
+	const FChaosVDFrameStageData& InSolverFrameStageData, const TSet<int32>& ParticlesUpdatedIDsForKeyFrameDiff, FTedsHandlesForBatchArray& OutRemovedTedsHandles)
 {
 	// TODO: We will need handle multi selection in the future, we can do by checking if each particle is selected directly with
 	// the selection set, but given that we know now that we only have one particle, we can just compare against that as this is faster
@@ -126,8 +147,6 @@ void UChaosVDParticleDataComponent::ProcessRemovedParticles_AssumesLocked(const
 		Chaos::VD::TypedElementDataUtil::DestroyTypedElementHandleForStruct(SelectedParticleInstance);
 	}
 
-	FMemMark StackMarker(FMemStack::Get());
-
 	// Evaluating and setting a particle as inactive can be done in parallel, but updating visibility and TEDS based on the new state, needs to be done in the GT,
 	TQueue<TSharedPtr<FChaosVDSceneParticle>, EQueueMode::Mpsc> DestroyedParticlesToProcessInGT;
 
@@ -146,6 +165,7 @@ void UChaosVDParticleDataComponent::ProcessRemovedParticles_AssumesLocked(const
 				if (TSharedRef<FChaosVDSceneParticle>* FoundParticle = SolverParticlesByID.Find(ParticleID))
 				{
 					(*FoundParticle)->SetIsActive(false);
+
 					DestroyedParticlesToProcessInGT.Enqueue(*FoundParticle);
 				}
 			}
@@ -165,35 +185,28 @@ void UChaosVDParticleDataComponent::ProcessRemovedParticles_AssumesLocked(const
 			if (TSharedRef<FChaosVDSceneParticle>* FoundParticle = SolverParticlesByID.Find(DestroyedParticlesID))
 			{
 				(*FoundParticle)->SetIsActive(false);
+
 				DestroyedParticlesToProcessInGT.Enqueue(*FoundParticle);
 			}
 		});
 	}
 
-	using namespace UE::Editor::DataStorage;
-	IEditorDataStorageCompatibilityProvider* Compatibility = GetMutableDataStorageFeature<IEditorDataStorageCompatibilityProvider>(CompatibilityFeatureName);
-	IEditorDataStorageProvider* DataStorage = GetMutableDataStorageFeature<IEditorDataStorageProvider>(StorageFeatureName);
-
 	while (!DestroyedParticlesToProcessInGT.IsEmpty())
 	{
 		TSharedPtr<FChaosVDSceneParticle> ParticleToProcessInGT;
 		DestroyedParticlesToProcessInGT.Dequeue(ParticleToProcessInGT);
 
-		ParticleToProcessInGT->SyncTEDSRegistration(DataStorage, Compatibility);
+		if (EnumHasAnyFlags(ParticleToProcessInGT->DirtyFlags, EChaosVDSceneParticleDirtyFlags::Active))
+		{
+			OutRemovedTedsHandles.Emplace(ParticleToProcessInGT->GetTedsRowHandle());
+			EnumRemoveFlags(ParticleToProcessInGT->DirtyFlags, EChaosVDSceneParticleDirtyFlags::Active);
+		}
+
 		ParticleToProcessInGT->UpdateGeometryComponentsVisibility();
 	}
 }
 
-void UChaosVDParticleDataComponent::BatchTEDSSyncFromWorldTagUpdate(TArray<UE::Editor::DataStorage::RowHandle, TInlineAllocator<256, TMemStackAllocator<>>>& ParticlesPendingTEDSTagSync)
-{
-	using namespace UE::Editor::DataStorage;
-
-	const IEditorDataStorageCompatibilityProvider* Compatibility = GetDataStorageFeature<IEditorDataStorageCompatibilityProvider>(CompatibilityFeatureName);
-	IEditorDataStorageProvider* DataStorage = GetMutableDataStorageFeature<IEditorDataStorageProvider>(StorageFeatureName);
-	BatchTEDSSyncFromWorldTagUpdate(ParticlesPendingTEDSTagSync, Compatibility, DataStorage);
-}
-
-void UChaosVDParticleDataComponent::BatchTEDSSyncFromWorldTagUpdate(TArray<UE::Editor::DataStorage::RowHandle, TInlineAllocator<256, TMemStackAllocator<>>>& ParticlesPendingTEDSTagSync,
+void UChaosVDParticleDataComponent::BatchTEDSSyncFromWorldTag(FTedsHandlesForBatchArray& ParticlesPendingTEDSTagSync,
 	const IEditorDataStorageCompatibilityProvider* CompatibilityStorage, IEditorDataStorageProvider* DataStorage)
 {
 	if (CompatibilityStorage && DataStorage)
@@ -211,6 +224,17 @@ void UChaosVDParticleDataComponent::BatchTEDSSyncFromWorldTagUpdate(TArray<UE::E
 	}
 }
 
+void UChaosVDParticleDataComponent::BatchTEDSAddRemoveTags(IEditorDataStorageProvider* EditorDataStorage, FTedsHandlesForBatchArray& ParticlesPendingTEDSTagSync,
+	TConstArrayView<const UScriptStruct*> ColumnsToAdd, TConstArrayView<const UScriptStruct*> ColumnsToRemove)
+{
+	if (!EditorDataStorage)
+	{
+		return;
+	}
+	
+	EditorDataStorage->BatchAddRemoveColumns(ParticlesPendingTEDSTagSync, ColumnsToAdd, ColumnsToRemove);
+}
+
 void UChaosVDParticleDataComponent::UpdateFromNewSolverStageData(const FChaosVDSolverFrameData& InSolverFrameData, const FChaosVDFrameStageData& InSolverFrameStageData)
 {
 	// Units of work representing each phase of the update loop (Pre-Processing Particle Data, Applying Particle data, Update TEDS Tags, and Removing Particles).
@@ -229,10 +253,14 @@ void UChaosVDParticleDataComponent::UpdateFromNewSolverStageData(const FChaosVDS
 	}
 
 	TQueue<TSharedPtr<FChaosVDSceneParticle>, EQueueMode::Mpsc> SceneParticleToProcessInGT;
-	
+
 	{
+		using namespace UE::Editor::DataStorage;
+		IEditorDataStorageCompatibilityProvider* Compatibility = GetMutableDataStorageFeature<IEditorDataStorageCompatibilityProvider>(CompatibilityFeatureName);
+		IEditorDataStorageProvider* DataStorage = GetMutableDataStorageFeature<IEditorDataStorageProvider>(StorageFeatureName);
+
 		// Make a pre-pass in parallel and process all the data that can be updated in parallel. This step will mostly copy data, and set flags of what needs to be updated in the Game Thread
-		ParallelFor(InSolverFrameStageData.RecordedParticlesData.Num(),[this, &SceneParticleToProcessInGT, &InSolverFrameStageData, &InSolverFrameData](int32 ParticleIndex)
+		ParallelFor(InSolverFrameStageData.RecordedParticlesData.Num(),[this, &DataStorage, &Compatibility, &SceneParticleToProcessInGT, &InSolverFrameStageData, &InSolverFrameData](int32 ParticleIndex)
 		{
 			const TSharedPtr<FChaosVDParticleDataWrapper>& Particle = InSolverFrameStageData.RecordedParticlesData[ParticleIndex];
 			if (!Particle)
@@ -268,37 +296,41 @@ void UChaosVDParticleDataComponent::UpdateFromNewSolverStageData(const FChaosVDS
 			{
 				{
 					FWriteScopeLock WriteLock(ParticleSceneObjectsLock);
-					ParticleInstanceToUpdate = CreateSceneParticle_AssumesLocked(Particle.ToSharedRef(), InSolverFrameData);
-					SolverParticlesByID.Add(ParticleVDInstanceID, ParticleInstanceToUpdate.ToSharedRef());
+
+					TSharedRef<FChaosVDSceneParticle> NewParticleRef = CreateSceneParticle_AssumesLocked(Particle.ToSharedRef(), InSolverFrameData);
+
+					NewParticleRef->SetTedsRowHandle(Chaos::VD::TedsUtils::AddObjectToDataStorage(&NewParticleRef.Get(), DataStorage, Compatibility));
+
+					ParticleInstanceToUpdate = NewParticleRef;
 				}
 
 				ParticleInstanceToUpdate->PreUpdateFromRecordedParticleData(Particle, InSolverFrameData.SimulationTransform);
+
+				StreamingSystem.EnqueuePendingTrackingOperation(ParticleInstanceToUpdate.ToSharedRef(), FChaosVDSceneStreaming::FPendingTrackingOperation::EType::AddOrUpdate);
 			}
 
 			SceneParticleToProcessInGT.Enqueue(ParticleInstanceToUpdate);
 		});
 	}
 
-	FMemMark StackMarker(FMemStack::Get());
-
 	using namespace UE::Editor::DataStorage;
 
+	FMemMark StackMarker(FMemStack::Get());
+
 	// Currently, updating TEDS tags is expensive, but the cost is reduced if we batch the update
 	// therefore we will use this to gather all TEDS rows that need updating and do that in a single call later on
-	TArray<RowHandle, TInlineAllocator<256, TMemStackAllocator<>>> ParticlesPendingTEDSTagSync;
-	ParticlesPendingTEDSTagSync.Reserve(InSolverFrameStageData.RecordedParticlesData.Num());
+	TArray<RowHandle, TInlineAllocator<256, TMemStackAllocator<>>> ParticlesPendingTEDSSyncFromWorldTagSync;
+	ParticlesPendingTEDSSyncFromWorldTagSync.Reserve(InSolverFrameStageData.RecordedParticlesData.Num());
 
 	TSet<int32> ParticlesUpdatedIDsForKeyFrameDiff;
 	if (Chaos::VisualDebugger::Cvars::bUnloadParticleDataUsingKeyFrameDiff)
 	{
 		ParticlesUpdatedIDsForKeyFrameDiff.Reserve(InSolverFrameStageData.RecordedParticlesData.Num());
 	}
+
+	FTedsHandlesForBatchArray ParticlesPendingTEDSActiveStateTagAddition;
 	
 	{
-		using namespace UE::Editor::DataStorage;
-		IEditorDataStorageCompatibilityProvider* Compatibility = GetMutableDataStorageFeature<IEditorDataStorageCompatibilityProvider>(CompatibilityFeatureName);
-		IEditorDataStorageProvider* DataStorage = GetMutableDataStorageFeature<IEditorDataStorageProvider>(StorageFeatureName);
-
 		const FText UpdatingSceneMessage = FText::Format(FTextFormat(LOCTEXT("UpdatingSceneMessage", "Updating Scene for {0} Solver with ID {1} ...")), FText::FromString(GetSolverName()), FText::AsNumber(SolverID));
 		UpdatingParticleData.EnterProgressFrame(1.0f, UpdatingSceneMessage);
 
@@ -310,8 +342,19 @@ void UChaosVDParticleDataComponent::UpdateFromNewSolverStageData(const FChaosVDS
 			TSharedPtr<FChaosVDSceneParticle> ParticleToProcessInGT;
 			SceneParticleToProcessInGT.Dequeue(ParticleToProcessInGT);
 
+			if (EnumHasAnyFlags(ParticleToProcessInGT->DirtyFlags, EChaosVDSceneParticleDirtyFlags::Transform))
+			{
+				StreamingSystem.EnqueuePendingTrackingOperation(ParticleToProcessInGT.ToSharedRef(), FChaosVDSceneStreaming::FPendingTrackingOperation::EType::AddOrUpdate);	
+			}
+
 			ParticleToProcessInGT->ProcessPendingParticleDataUpdates();
-			ParticleToProcessInGT->SyncTEDSRegistration(DataStorage, Compatibility);
+
+			if (EnumHasAnyFlags(ParticleToProcessInGT->DirtyFlags, EChaosVDSceneParticleDirtyFlags::Active))
+			{
+				// We could check if the particle is active before doing this, but at this point only newly active particles can be in the queue
+				ParticlesPendingTEDSActiveStateTagAddition.Emplace(ParticleToProcessInGT->GetTedsRowHandle());
+				EnumRemoveFlags(ParticleToProcessInGT->DirtyFlags, EChaosVDSceneParticleDirtyFlags::Active);
+			}
 
 			if (Chaos::VisualDebugger::Cvars::bUnloadParticleDataUsingKeyFrameDiff)
 			{
@@ -323,26 +366,22 @@ void UChaosVDParticleDataComponent::UpdateFromNewSolverStageData(const FChaosVDS
 
 			if (EnumHasAnyFlags(ParticleToProcessInGT->DirtyFlags, EChaosVDSceneParticleDirtyFlags::TEDS))
 			{
-				ParticlesPendingTEDSTagSync.Emplace(ParticleToProcessInGT->GetTedsRowHandle());
+				ParticlesPendingTEDSSyncFromWorldTagSync.Emplace(ParticleToProcessInGT->GetTedsRowHandle());
 
 				EnumRemoveFlags(ParticleToProcessInGT->DirtyFlags, EChaosVDSceneParticleDirtyFlags::TEDS);
 			}
 		}
 	}
 
-	{
-		const FText UpdatingDataStorageTagsMessage = FText::Format(FTextFormat(LOCTEXT("UpdatingDataStorageTagsMessage", "Updating Editor Data Storage Tags for {0} Solver with ID {1} ...")), FText::FromString(GetSolverName()), FText::AsNumber(SolverID));
-		UpdatingParticleData.EnterProgressFrame(1.0f, UpdatingDataStorageTagsMessage);
+	FTedsHandlesForBatchArray ParticlesPendingTEDSActiveStateTagRemoval;
+	ProcessRemovedParticles_AssumessLocked(ScenePtr.ToSharedRef(), InSolverFrameData, InSolverFrameStageData, ParticlesUpdatedIDsForKeyFrameDiff, ParticlesPendingTEDSActiveStateTagRemoval);
 
-		BatchTEDSSyncFromWorldTagUpdate(ParticlesPendingTEDSTagSync);
-	}
-
-	{
-		const FText RemovingDestroyedParticlesMessage = FText::Format(FTextFormat(LOCTEXT("RemovingDestroyedParticlesMessage", "Removing Destroyed Particles for {0} Solver with ID {1} ...")), FText::FromString(GetSolverName()), FText::AsNumber(SolverID));
-		UpdatingParticleData.EnterProgressFrame(1.0f, RemovingDestroyedParticlesMessage);
-
-		ProcessRemovedParticles_AssumesLocked(ScenePtr.ToSharedRef(), InSolverFrameData, InSolverFrameStageData, ParticlesUpdatedIDsForKeyFrameDiff);
-	}
+	const IEditorDataStorageCompatibilityProvider* EditorDataStorageCompatibility = GetDataStorageFeature<IEditorDataStorageCompatibilityProvider>(CompatibilityFeatureName);
+	IEditorDataStorageProvider* EditorDataStorage = GetMutableDataStorageFeature<IEditorDataStorageProvider>(StorageFeatureName);
+	
+	BatchTEDSSyncFromWorldTag(ParticlesPendingTEDSSyncFromWorldTagSync, EditorDataStorageCompatibility, EditorDataStorage);
+	BatchTEDSAddRemoveTags(EditorDataStorage, ParticlesPendingTEDSActiveStateTagAddition, { FChaosVDActiveObjectTag::StaticStruct() } , {});
+	BatchTEDSAddRemoveTags(EditorDataStorage, ParticlesPendingTEDSActiveStateTagRemoval, {}, { FChaosVDActiveObjectTag::StaticStruct() });
 }
 
 TSharedRef<FChaosVDSceneParticle> UChaosVDParticleDataComponent::CreateSceneParticle_AssumesLocked(const TSharedRef<FChaosVDParticleDataWrapper>& InParticleData, const FChaosVDSolverFrameData& InFrameData)
@@ -359,6 +398,14 @@ TSharedRef<FChaosVDSceneParticle> UChaosVDParticleDataComponent::CreateScenePart
 	NewRepresentation->SetIsActive(true);
 	NewRepresentation->SetIsServerParticle(IsServerData());
 
+	if (!StreamingSystem.IsEnabled())
+	{
+		NewRepresentation->SetStreamingState(FChaosVDBaseSceneObject::EStreamingState::Visible);
+	}
+
+	SolverParticlesByID.Emplace(InParticleData->ParticleIndex, NewRepresentation);
+	SolverParticlesArray.Emplace(StaticCastSharedRef<FChaosVDBaseSceneObject>(NewRepresentation));
+
 	return NewRepresentation;
 }
 
@@ -434,7 +481,14 @@ void UChaosVDParticleDataComponent::BatchApplySolverVisibilityToParticle(bool bN
 		ParticleWithIDPair.Value->UpdateGeometryComponentsVisibility();
 	}
 
-	BatchTEDSSyncFromWorldTagUpdate(ParticlesPendingTEDSTagSync);
+	const IEditorDataStorageCompatibilityProvider* EditorDataStorageCompatibility = GetDataStorageFeature<IEditorDataStorageCompatibilityProvider>(CompatibilityFeatureName);
+	IEditorDataStorageProvider* EditorDataStorage = GetMutableDataStorageFeature<IEditorDataStorageProvider>(StorageFeatureName);
+	BatchTEDSSyncFromWorldTag(ParticlesPendingTEDSTagSync, EditorDataStorageCompatibility, EditorDataStorage);
+	
+	if (TSharedPtr<FChaosVDScene> ScenePtr = SceneWeakPtr.Pin())
+	{
+		ScenePtr->RequestUpdate();
+	}
 }
 
 void UChaosVDParticleDataComponent::SetVisibility(bool bNewIsVisible)
diff --git a/Engine/Plugins/ChaosVD/Source/ChaosVD/Private/Components/ChaosVDParticleDataComponent.h b/Engine/Plugins/ChaosVD/Source/ChaosVD/Private/Components/ChaosVDParticleDataComponent.h
--- a/Engine/Plugins/ChaosVD/Source/ChaosVD/Private/Components/ChaosVDParticleDataComponent.h
+++ b/Engine/Plugins/ChaosVD/Source/ChaosVD/Private/Components/ChaosVDParticleDataComponent.h
@@ -4,31 +4,37 @@
 
 #include "ChaosVDScene.h"
 #include "ChaosVDSceneParticle.h"
+#include "ChaosVDSceneStreaming.h"
+#include "Chaos/AABBTree.h"
 #include "Components/ChaosVDSolverDataComponent.h"
+#include "Math/GenericOctree.h"
 #include "TEDS/ChaosVDSelectionInterface.h"
 #include "TEDS/ChaosVDStructTypedElementData.h"
 #include "ChaosVDParticleDataComponent.generated.h"
 
+class FChaosVDSceneStreaming;
 struct FChaosVDParticleDataWrapper;
 struct FChaosVDSceneParticle;
+
 /**
- * Component that will reference all particle data for a specific solver for the current frame.
- * Currently it is empty and only used to interface with the component visualizer system for particle data,
- * but when we eventually remove the dependency on actors to represent particles as planned, the particle data will be referenced in this component.\
- * (similarly to how the Scene Query Data and Collision Data is referenced in their respective components)
+ * Component that references all particle data for a specific solver for the current frame, and handles how the visualization
+ * is updated based on that data
  */
 UCLASS()
-class UChaosVDParticleDataComponent : public UChaosVDSolverDataComponent
+class UChaosVDParticleDataComponent : public UChaosVDSolverDataComponent, public IChaosVDStreamingDataSource
 {
 public:
 
 	UChaosVDParticleDataComponent();
 
+	typedef TArray<UE::Editor::DataStorage::RowHandle, TInlineAllocator<256, TMemStackAllocator<>>> FTedsHandlesForBatchArray;
+
 private:
 	GENERATED_BODY()
 
 public:
 
+	virtual void SetScene(const TWeakPtr<FChaosVDScene>& InSceneWeakPtr) override;
 	virtual void ClearData() override;
 
 	virtual void UpdateFromNewSolverStageData(const FChaosVDSolverFrameData& InSolverFrameData, const FChaosVDFrameStageData& InSolverFrameStageData) override;
@@ -37,6 +43,9 @@ public:
 	template <class TCallback>
 	void VisitSelectedParticleData(TCallback VisitCallback) const;
 
+	template <class TCallback>
+	void VisitAllParticleInstances(TCallback VisitCallback) const;
+
 	template <class TCallback>
 	void VisitAllParticleData(TCallback VisitCallback) const;
 
@@ -55,14 +64,30 @@ public:
 
 	FChaosVDSceneParticle* GetSelectedParticle() const;
 
+	virtual void HandleWorldStreamingLocationUpdated(const FVector& InLocation) override;
+
+	// BEGIN IChaosVDStreamingDataSource
+	virtual TConstArrayView<TSharedRef<FChaosVDBaseSceneObject>> GetStreamableSceneObjects() const override
+	{
+		return SolverParticlesArray;
+	}
+	
+	virtual FRWLock& GetObjectsLock() const override
+	{
+		return ParticleSceneObjectsLock;
+	}
+	// END IChaosVDStreamingDataSource
+
 protected:
 
 	int32 GetMaxElementCountForBatchTEDSUpdate_AssumesLocked() const;
 
-	void BatchTEDSSyncFromWorldTagUpdate(TArray<UE::Editor::DataStorage::RowHandle, TInlineAllocator<256, TMemStackAllocator<>>>& ParticlesPendingTEDSTagSync);
-	void BatchTEDSSyncFromWorldTagUpdate(TArray<UE::Editor::DataStorage::RowHandle, TInlineAllocator<256, TMemStackAllocator<>>>& ParticlesPendingTEDSTagSync, 	const IEditorDataStorageCompatibilityProvider* CompatibilityStorage, IEditorDataStorageProvider* DataStorage);
+	void BatchTEDSSyncFromWorldTag(FTedsHandlesForBatchArray& ParticlesPendingTEDSTagSync, const IEditorDataStorageCompatibilityProvider* CompatibilityStorage, IEditorDataStorageProvider* DataStorage);
+
+	void BatchTEDSAddRemoveTags(IEditorDataStorageProvider* EditorDataStorage, FTedsHandlesForBatchArray& ParticlesPendingTEDSTagSync, TConstArrayView<const UScriptStruct*> ColumnsToAdd, TConstArrayView<const UScriptStruct*> ColumnsToRemove);
 
-	void ProcessRemovedParticles_AssumesLocked(const TSharedRef<FChaosVDScene>& SceneRef, const FChaosVDSolverFrameData& InSolverFrameData, const FChaosVDFrameStageData& InSolverFrameStageData, const TSet<int32>& ParticlesUpdatedIDsForKeyFrameDiff);
+	void ProcessRemovedParticles_AssumessLocked(const TSharedRef<FChaosVDScene>& SceneRef, const FChaosVDSolverFrameData& InSolverFrameData,
+		const FChaosVDFrameStageData& InSolverFrameStageData, const TSet<int32>& ParticlesUpdatedIDsForKeyFrameDiff, FTedsHandlesForBatchArray& OutRemovedTedsHandles);
 
 	void BatchApplySolverVisibilityToParticle(bool bNewIsVisible);
 
@@ -80,11 +105,12 @@ protected:
 
 	TSortedMap<EChaosVDParticleType, TSharedPtr<FChaosVDBaseSceneObject>> ParticleSceneContainersByType;
 
-	TSortedMap<EChaosVDParticleType, FName> FolderPathByParticlePath;
-
 	TMap<int32, TSharedRef<FChaosVDSceneParticle>> SolverParticlesByID;
+	TArray<TSharedRef<FChaosVDBaseSceneObject>> SolverParticlesArray;
 
 	mutable FRWLock ParticleSceneObjectsLock;
+
+	FChaosVDSceneStreaming StreamingSystem;
 };
 
 template <typename TCallback>
@@ -124,6 +150,20 @@ void UChaosVDParticleDataComponent::VisitSelectedParticleData(TCallback VisitCal
 	}
 }
 
+template <typename TCallback>
+void UChaosVDParticleDataComponent::VisitAllParticleInstances(TCallback VisitCallback) const
+{
+	FReadScopeLock ReadLock(ParticleSceneObjectsLock);
+
+	for (const TPair<int32, TSharedRef<FChaosVDSceneParticle>>& ParticleWithIDPair : SolverParticlesByID)
+	{
+		if (!VisitCallback(ParticleWithIDPair.Value))
+		{
+			return;
+		}
+	}
+}
+
 template <typename TCallback>
 void UChaosVDParticleDataComponent::VisitAllParticleData(TCallback VisitCallback) const
 {
diff --git a/Engine/Plugins/ChaosVD/Source/ChaosVD/Private/Components/ChaosVDStaticMeshComponent.cpp b/Engine/Plugins/ChaosVD/Source/ChaosVD/Private/Components/ChaosVDStaticMeshComponent.cpp
--- a/Engine/Plugins/ChaosVD/Source/ChaosVD/Private/Components/ChaosVDStaticMeshComponent.cpp
+++ b/Engine/Plugins/ChaosVD/Source/ChaosVD/Private/Components/ChaosVDStaticMeshComponent.cpp
@@ -14,26 +14,20 @@ uint32 UChaosVDStaticMeshComponent::GetGeometryKey() const
 	return CurrentGeometryKey;
 }
 
-void UChaosVDStaticMeshComponent::UpdateInstanceVisibility(const TSharedPtr<FChaosVDMeshDataInstanceHandle>& InInstanceHandle, bool bIsVisible)
+void UChaosVDStaticMeshComponent::UpdateVisibilityForInstance(const TSharedRef<FChaosVDInstancedMeshData>& InInstanceHandle)
 {
-	if (!InInstanceHandle.IsValid())
-	{
-		UE_LOG(LogChaosVDEditor, Warning, TEXT("[%s] Attempted to update a mesh instance using an invalid handle. No instances were updated"), ANSI_TO_TCHAR(__FUNCTION__));
-		return;
-	}
-
 	if (InInstanceHandle->GetMeshComponent() != this)
 	{
 		UE_LOG(LogChaosVDEditor, Error, TEXT("[%s] Attempted to update a mesh instance using a handle from another component. No instances were updated | Handle Component [%s] | Current Component [%s]"), ANSI_TO_TCHAR(__FUNCTION__), *GetNameSafe(InInstanceHandle->GetMeshComponent()), *GetNameSafe(this));
 		return;
 	}
 
-	SetVisibility(bIsVisible);
+	SetVisibility(InInstanceHandle->GetState().bIsVisible);
 }
 
-void UChaosVDStaticMeshComponent::SetIsSelected(const TSharedPtr<FChaosVDMeshDataInstanceHandle>& InInstanceHandle, bool bIsSelected)
+void UChaosVDStaticMeshComponent::UpdateSelectionStateForInstance(const TSharedRef<FChaosVDInstancedMeshData>& InInstanceHandle)
 {
-	bIsOwningParticleSelected = bIsSelected;
+	bIsOwningParticleSelected = InInstanceHandle->GetState().bIsSelected;
 
 	PushSelectionToProxy();
 }
@@ -43,20 +37,16 @@ bool UChaosVDStaticMeshComponent::ShouldRenderSelected() const
 	return bIsOwningParticleSelected;
 }
 
-void UChaosVDStaticMeshComponent::UpdateInstanceColor(const TSharedPtr<FChaosVDMeshDataInstanceHandle>& InInstanceHandle, FLinearColor NewColor)
+void UChaosVDStaticMeshComponent::UpdateColorForInstance(const TSharedRef<FChaosVDInstancedMeshData>& InInstanceHandle)
 {
-	if (!InInstanceHandle.IsValid())
-	{
-		UE_LOG(LogChaosVDEditor, Warning, TEXT("[%s] Attempted to update a mesh instance using an invalid handle. No instances were updated"), ANSI_TO_TCHAR(__FUNCTION__));
-		return;
-	}
-
 	if (InInstanceHandle->GetMeshComponent() != this)
 	{
 		UE_LOG(LogChaosVDEditor, Error, TEXT("[%s] Attempted to update a mesh instance using a handle from another component. No instances were updated | Handle Component [%s] | Current Component [%s]"), ANSI_TO_TCHAR(__FUNCTION__), *GetNameSafe(InInstanceHandle->GetMeshComponent()), *GetNameSafe(this));
 		return;
 	}
-	
+
+	FLinearColor NewColor = InInstanceHandle->GetInstanceColor();
+
 	const bool bIsSolidColor = FMath::IsNearlyEqual(NewColor.A, 1.0f);
 
 	bool bHasOpaqueMaterial = !EnumHasAnyFlags(MeshComponentAttributeFlags, EChaosVDMeshAttributesFlags::TranslucentGeometry);
@@ -79,26 +69,21 @@ void UChaosVDStaticMeshComponent::UpdateInstanceColor(const TSharedPtr<FChaosVDM
 	}
 }
 
-void UChaosVDStaticMeshComponent::UpdateInstanceWorldTransform(const TSharedPtr<FChaosVDMeshDataInstanceHandle>& InInstanceHandle, const FTransform& InTransform)
+void UChaosVDStaticMeshComponent::UpdateWorldTransformForInstance(const TSharedRef<FChaosVDInstancedMeshData>& InInstanceHandle)
 {
-	if (!InInstanceHandle.IsValid())
-	{
-		UE_LOG(LogChaosVDEditor, Warning, TEXT("[%s] Attempted to update a mesh instance using an invalid handle. No instances were updated"), ANSI_TO_TCHAR(__FUNCTION__));
-		return;
-	}
-
 	if (InInstanceHandle->GetMeshComponent() != this)
 	{
 		UE_LOG(LogChaosVDEditor, Error, TEXT("[%s] Attempted to update a mesh instance using a handle from another component. No instances were updated | Handle Component [%s] | Current Component [%s]"), ANSI_TO_TCHAR(__FUNCTION__), *GetNameSafe(InInstanceHandle->GetMeshComponent()), *GetNameSafe(this));
 		return;
 	}
 
-	SetWorldTransform(InTransform);
+	SetWorldTransform(InInstanceHandle->GetWorldTransform());
 }
 
 void UChaosVDStaticMeshComponent::Reset()
 {
 	bIsMeshReady = false;
+	bIsDestroyed = false;
 	MeshReadyDelegate = FChaosVDMeshReadyDelegate();
 	ComponentEmptyDelegate = FChaosVDMeshComponentEmptyDelegate();
 
@@ -108,12 +93,12 @@ void UChaosVDStaticMeshComponent::Reset()
 	CurrentGeometryKey = 0;
 }
 
-TSharedPtr<FChaosVDMeshDataInstanceHandle> UChaosVDStaticMeshComponent::AddMeshInstance(const FTransform InstanceTransform, bool bIsWorldSpace, const TSharedRef<FChaosVDExtractedGeometryDataHandle>& InGeometryHandle, int32 ParticleID, int32 SolverID)
+TSharedPtr<FChaosVDInstancedMeshData> UChaosVDStaticMeshComponent::AddMeshInstance(const FTransform InstanceTransform, bool bIsWorldSpace, const TSharedRef<FChaosVDExtractedGeometryDataHandle>& InGeometryHandle, int32 ParticleID, int32 SolverID)
 {
 	// Static meshes only have one instance
 	constexpr int32 StaticMeshInstanceIndex = 0;
 
-	CurrentMeshDataHandle = MakeShared<FChaosVDMeshDataInstanceHandle>(StaticMeshInstanceIndex, this, ParticleID, SolverID, InGeometryHandle);
+	CurrentMeshDataHandle = MakeShared<FChaosVDInstancedMeshData>(StaticMeshInstanceIndex, this, ParticleID, SolverID, InGeometryHandle);
 
 	const uint32 NewHandleGeometryKey = InGeometryHandle->GetGeometryKey();
 
@@ -127,24 +112,22 @@ TSharedPtr<FChaosVDMeshDataInstanceHandle> UChaosVDStaticMeshComponent::AddMeshI
 	return CurrentMeshDataHandle;
 }
 
-void UChaosVDStaticMeshComponent::AddMeshInstanceForHandle(TSharedPtr<FChaosVDMeshDataInstanceHandle> MeshDataHandle, const FTransform InstanceTransform, bool bIsWorldSpace, const TSharedRef<FChaosVDExtractedGeometryDataHandle>& InGeometryHandle, int32 ParticleID, int32 SolverID)
+void UChaosVDStaticMeshComponent::AddExistingMeshInstance(const TSharedRef<FChaosVDInstancedMeshData>& InMeshDataHandle)
 {
 	// Static meshes only have one instance
 	constexpr int32 StaticMeshInstanceIndex = 0;
 
-	const uint32 NewHandleGeometryKey = InGeometryHandle->GetGeometryKey();
+	const uint32 NewHandleGeometryKey = InMeshDataHandle->ExtractedGeometryHandle->GetGeometryKey();
 	if (!UpdateGeometryKey(NewHandleGeometryKey))
 	{
 		return;
 	}
 
-	MeshDataHandle->SetMeshInstanceIndex(StaticMeshInstanceIndex);
-	MeshDataHandle->SetMeshComponent(this);
-	
-	CurrentMeshDataHandle->SetWorldTransform(InstanceTransform);
+	InMeshDataHandle->SetMeshInstanceIndex(StaticMeshInstanceIndex);
+	InMeshDataHandle->SetMeshComponent(this);
 }
 
-void UChaosVDStaticMeshComponent::RemoveMeshInstance(TSharedPtr<FChaosVDMeshDataInstanceHandle> InHandleToRemove)
+void UChaosVDStaticMeshComponent::RemoveMeshInstance(const TSharedRef<FChaosVDInstancedMeshData>& InHandleToRemove, ERemovalMode Mode)
 {
 	SetStaticMesh(nullptr);
 	CurrentMeshDataHandle = nullptr;
@@ -170,6 +153,8 @@ void UChaosVDStaticMeshComponent::OnDisposed()
 {
 	Reset();
 
+	bIsDestroyed = true;
+
 	SetRelativeTransform(FTransform::Identity);
 
 	if (IsRegistered())
@@ -200,16 +185,11 @@ bool UChaosVDStaticMeshComponent::UpdateGeometryKey(uint32 NewHandleGeometryKey)
 	return true;
 }
 
-TSharedPtr<FChaosVDMeshDataInstanceHandle> UChaosVDStaticMeshComponent::GetMeshDataInstanceHandle(int32 InstanceIndex) const
+TSharedPtr<FChaosVDInstancedMeshData> UChaosVDStaticMeshComponent::GetMeshDataInstanceHandle(int32 InstanceIndex) const
 {
 	return CurrentMeshDataHandle;
 }
 
-TArrayView<TSharedPtr<FChaosVDMeshDataInstanceHandle>> UChaosVDStaticMeshComponent::GetMeshDataInstanceHandles()
-{
-	return TArrayView<TSharedPtr<FChaosVDMeshDataInstanceHandle>>(&CurrentMeshDataHandle, 1);
-}
-
 void UChaosVDStaticMeshComponent::Initialize()
 {
 	TSharedPtr<FChaosVDGeometryBuilder> GeometryBuilder = GeometryBuilderWeakPtr.Pin();
diff --git a/Engine/Plugins/ChaosVD/Source/ChaosVD/Private/Components/ChaosVDStaticMeshComponent.h b/Engine/Plugins/ChaosVD/Source/ChaosVD/Private/Components/ChaosVDStaticMeshComponent.h
--- a/Engine/Plugins/ChaosVD/Source/ChaosVD/Private/Components/ChaosVDStaticMeshComponent.h
+++ b/Engine/Plugins/ChaosVD/Source/ChaosVD/Private/Components/ChaosVDStaticMeshComponent.h
@@ -34,36 +34,49 @@ public:
 	virtual FChaosVDMeshComponentEmptyDelegate* OnComponentEmpty() override { return &ComponentEmptyDelegate; }
 
 	virtual uint32 GetGeometryKey() const override;
-	virtual void UpdateInstanceVisibility(const TSharedPtr<FChaosVDMeshDataInstanceHandle>& InInstanceHandle, bool bIsVisible) override;
+	virtual void UpdateVisibilityForInstance(const TSharedRef<FChaosVDInstancedMeshData>& InInstanceHandle) override;
 
-	virtual void SetIsSelected(const TSharedPtr<FChaosVDMeshDataInstanceHandle>& InInstanceHandle, bool bIsSelected) override;
+	virtual void UpdateSelectionStateForInstance(const TSharedRef<FChaosVDInstancedMeshData>& InInstanceHandle) override;
 
 	virtual bool ShouldRenderSelected() const override;
 
-	virtual void UpdateInstanceColor(const TSharedPtr<FChaosVDMeshDataInstanceHandle>& InInstanceHandle, FLinearColor NewColor) override;
-	virtual void UpdateInstanceWorldTransform(const TSharedPtr<FChaosVDMeshDataInstanceHandle>& InInstanceHandle, const FTransform& InTransform) override;
+	virtual void UpdateColorForInstance(const TSharedRef<FChaosVDInstancedMeshData>& InInstanceHandle) override;
+	virtual void UpdateWorldTransformForInstance(const TSharedRef<FChaosVDInstancedMeshData>& InInstanceHandle) override;
 
 	virtual void SetMeshComponentAttributeFlags(EChaosVDMeshAttributesFlags Flags) override { MeshComponentAttributeFlags = Flags; };
 	virtual EChaosVDMeshAttributesFlags GetMeshComponentAttributeFlags() const override { return MeshComponentAttributeFlags; };
 
-	virtual TSharedPtr<FChaosVDMeshDataInstanceHandle> GetMeshDataInstanceHandle(int32 InstanceIndex) const override;
-	virtual TArrayView<TSharedPtr<FChaosVDMeshDataInstanceHandle>> GetMeshDataInstanceHandles() override;
+	virtual TSharedPtr<FChaosVDInstancedMeshData> GetMeshDataInstanceHandle(int32 InstanceIndex) const override;
 
 	virtual void Initialize() override;
 	virtual void Reset() override;
 
-	virtual TSharedPtr<FChaosVDMeshDataInstanceHandle> AddMeshInstance(const FTransform InstanceTransform, bool bIsWorldSpace, const TSharedRef<FChaosVDExtractedGeometryDataHandle>& InGeometryHandle, int32 ParticleID, int32 SolverID) override;
-	virtual void AddMeshInstanceForHandle(TSharedPtr<FChaosVDMeshDataInstanceHandle> MeshDataHandle, const FTransform InstanceTransform, bool bIsWorldSpace, const TSharedRef<FChaosVDExtractedGeometryDataHandle>& InGeometryHandle, int32 ParticleID, int32 SolverID) override;
-	virtual void RemoveMeshInstance(TSharedPtr<FChaosVDMeshDataInstanceHandle> InHandleToRemove) override;
+	virtual TSharedPtr<FChaosVDInstancedMeshData> AddMeshInstance(const FTransform InstanceTransform, bool bIsWorldSpace, const TSharedRef<FChaosVDExtractedGeometryDataHandle>& InGeometryHandle, int32 ParticleID, int32 SolverID) override;
+	virtual void AddExistingMeshInstance(const TSharedRef<FChaosVDInstancedMeshData>& InMeshDataHandle) override;
+	virtual void RemoveMeshInstance(const TSharedRef<FChaosVDInstancedMeshData>& InHandleToRemove, ERemovalMode Mode = ERemovalMode::Deferred) override;
 
 	virtual void SetGeometryBuilder(TWeakPtr<FChaosVDGeometryBuilder> GeometryBuilder) override;
 
 	virtual EChaosVDMaterialType GetMaterialType() const override;
 
+	virtual bool GetIsDestroyed() const override
+	{
+		return bIsDestroyed;
+	}
+
+	virtual void SetIsDestroyed(bool bNewIsPendingDestroy) override
+	{
+		bIsDestroyed = bNewIsPendingDestroy;
+	}
+
 	// END IChaosVDGeometryDataComponent Interface
 
 	// BEGIN IChaosVDPooledObject Interface
-	virtual void OnAcquired() override {};
+	virtual void OnAcquired() override
+	{
+		bIsDestroyed = false;
+	}
+
 	virtual void OnDisposed() override;
 	// END IChaosVDPooledObject Interface
 
@@ -77,8 +90,10 @@ protected:
 	bool bIsOwningParticleSelected = false;
 	FChaosVDMeshReadyDelegate MeshReadyDelegate;
 	FChaosVDMeshComponentEmptyDelegate ComponentEmptyDelegate;
+	
+	bool bIsDestroyed = false;
 
-	TSharedPtr<FChaosVDMeshDataInstanceHandle> CurrentMeshDataHandle = nullptr;
+	TSharedPtr<FChaosVDInstancedMeshData> CurrentMeshDataHandle = nullptr;
 
 	TSharedPtr<FChaosVDExtractedGeometryDataHandle> CurrentGeometryHandle = nullptr;
 	
diff --git a/Engine/Plugins/ChaosVD/Source/ChaosVD/Private/Settings/ChaosVDGeneralSettings.cpp b/Engine/Plugins/ChaosVD/Source/ChaosVD/Private/Settings/ChaosVDGeneralSettings.cpp
new file mode 100644
--- /dev/null
+++ b/Engine/Plugins/ChaosVD/Source/ChaosVD/Private/Settings/ChaosVDGeneralSettings.cpp
@@ -0,0 +1,3 @@
+// Copyright Epic Games, Inc. All Rights Reserved.
+
+#include "ChaosVDGeneralSettings.h"
\ No newline at end of file
diff --git a/Engine/Plugins/ChaosVD/Source/ChaosVD/Private/Settings/ChaosVDGeneralSettings.h b/Engine/Plugins/ChaosVD/Source/ChaosVD/Private/Settings/ChaosVDGeneralSettings.h
new file mode 100644
--- /dev/null
+++ b/Engine/Plugins/ChaosVD/Source/ChaosVD/Private/Settings/ChaosVDGeneralSettings.h
@@ -0,0 +1,28 @@
+// Copyright Epic Games, Inc. All Rights Reserved.
+
+#pragma once
+
+#include "Settings/ChaosVDCoreSettings.h"
+#include "ChaosVDGeneralSettings.generated.h"
+
+/**
+ * General settings that controls how CVD behaves
+ */
+UCLASS()
+class CHAOSVD_API UChaosVDGeneralSettings : public UChaosVDSettingsObjectBase
+{
+	GENERATED_BODY()
+
+public:
+	/** If true, CVD will only load collision geometry that is visible */
+	UPROPERTY(Config, EditAnywhere, Category = Streaming)
+	bool bStreamingSystemEnabled = true;
+
+	/** Extent size of the box used for calculate what should be streamed in */
+	UPROPERTY(Config, EditAnywhere, Category = Streaming)
+	float StreamingBoxExtentSize = 10000.0f;
+
+	// If set to true CVD will process any updates to the streaming accel structure in worker threads, in between streaming updates
+	UPROPERTY(Config, EditAnywhere, Category = Streaming)
+	bool bProcessPendingOperationsQueueInWorkerThread = true;
+};
diff --git a/Engine/Plugins/ChaosVD/Source/ChaosVD/Private/Settings/ChaosVDParticleVisualizationSettings.h b/Engine/Plugins/ChaosVD/Source/ChaosVD/Private/Settings/ChaosVDParticleVisualizationSettings.h
--- a/Engine/Plugins/ChaosVD/Source/ChaosVD/Private/Settings/ChaosVDParticleVisualizationSettings.h
+++ b/Engine/Plugins/ChaosVD/Source/ChaosVD/Private/Settings/ChaosVDParticleVisualizationSettings.h
@@ -41,7 +41,10 @@ enum class EChaosVDParticleDataVisualizationFlags : uint32
 	ClusterConnectivityEdge	= 1 << 6,
 	CenterOfMass	= 1 << 7,
 	DrawDataOnlyForSelectedParticle	= 1 << 8,
-	EnableDraw	= 1 << 9,
+	
+	Bounds = 1 << 10,
+
+	EnableDraw	= 1 << 9
 };
 ENUM_CLASS_FLAGS(EChaosVDParticleDataVisualizationFlags);
 
diff --git a/Engine/Plugins/ChaosVD/Source/ChaosVD/Private/TEDS/ChaosVDDataStorageVisibilityQueries.cpp b/Engine/Plugins/ChaosVD/Source/ChaosVD/Private/TEDS/ChaosVDDataStorageVisibilityQueries.cpp
--- a/Engine/Plugins/ChaosVD/Source/ChaosVD/Private/TEDS/ChaosVDDataStorageVisibilityQueries.cpp
+++ b/Engine/Plugins/ChaosVD/Source/ChaosVD/Private/TEDS/ChaosVDDataStorageVisibilityQueries.cpp
@@ -87,8 +87,8 @@ void UChaosVDDataStorageVisibilityQueries::RegisterVisibilityColumnToParticleQue
 			{
 				Particle->AddHiddenFlag(EChaosVDHideParticleFlags::HiddenBySceneOutliner);
 			}
-			Particle->UpdateGeometryComponentsVisibility();
-			
+
+			Particle->UpdateGeometryComponentsVisibility(EChaosVDParticleVisibilityUpdateFlags::DirtyScene);	
 		}
 		FChaosVDSceneParticle* Particle;
 		bool bIsVisible;
diff --git a/Engine/Plugins/ChaosVD/Source/ChaosVD/Private/TEDS/ChaosVDParticleEditorDataFactory.cpp b/Engine/Plugins/ChaosVD/Source/ChaosVD/Private/TEDS/ChaosVDParticleEditorDataFactory.cpp
--- a/Engine/Plugins/ChaosVD/Source/ChaosVD/Private/TEDS/ChaosVDParticleEditorDataFactory.cpp
+++ b/Engine/Plugins/ChaosVD/Source/ChaosVD/Private/TEDS/ChaosVDParticleEditorDataFactory.cpp
@@ -37,11 +37,11 @@ void UChaosVDParticleEditorDataFactory::RegisterTables(IEditorDataStorageProvide
 	using namespace UE::Editor::DataStorage;
 
 	TableHandle ObjectTable = DataStorage.RegisterTable(DataStorage.FindTable(FName("Editor_StandardExternalObjectTable")),
-		TTypedElementColumnTypeList<FChaosVDObjectDataTag, FTypedElementFromCVDWorldTag, FTypedElementLabelColumn, FTypedElementLabelHashColumn, FVisibleInEditorColumn, FChaosVDTableRowParentColumn>{},
+		TTypedElementColumnTypeList<FChaosVDObjectDataTag, FTypedElementFromCVDWorldTag, FTypedElementLabelColumn, FTypedElementLabelHashColumn, FVisibleInEditorColumn, FChaosVDTableRowParentColumn, FChaosVDActiveObjectTag>{},
 		FName("CVD_SceneObjectDataTable"));
 
 	TableHandle ActorTable = DataStorage.RegisterTable(DataStorage.FindTable(FName("Editor_StandardActorTable")),
-	TTypedElementColumnTypeList<FChaosVDObjectDataTag, FTypedElementFromCVDWorldTag, FVisibleInEditorColumn, FChaosVDTableRowParentColumn>{},
+	TTypedElementColumnTypeList<FChaosVDObjectDataTag, FTypedElementFromCVDWorldTag, FVisibleInEditorColumn, FChaosVDTableRowParentColumn, FChaosVDActiveObjectTag>{},
 	FName("CVD_ActorDataTable"));
 
 	DataStorageCompatibility.RegisterTypeTableAssociation(FChaosVDBaseSceneObject::StaticStruct(), ObjectTable);
diff --git a/Engine/Plugins/ChaosVD/Source/ChaosVD/Private/TEDS/ChaosVDParticleEditorDataFactory.h b/Engine/Plugins/ChaosVD/Source/ChaosVD/Private/TEDS/ChaosVDParticleEditorDataFactory.h
--- a/Engine/Plugins/ChaosVD/Source/ChaosVD/Private/TEDS/ChaosVDParticleEditorDataFactory.h
+++ b/Engine/Plugins/ChaosVD/Source/ChaosVD/Private/TEDS/ChaosVDParticleEditorDataFactory.h
@@ -28,6 +28,15 @@ struct FTypedElementFromCVDWorldTag final : public FEditorDataStorageTag
 	GENERATED_BODY()
 };
 
+/**
+ * TEDS tag added to any object that belongs to a CVD World, and is active (visible in the scene outliner, and with valid data)
+ */
+USTRUCT(meta = (DisplayName = "CVD Active Object"))
+struct FChaosVDActiveObjectTag final : public FEditorDataStorageTag
+{
+	GENERATED_BODY()
+};
+
 UCLASS()
 class UChaosVDParticleEditorDataFactory : public UEditorDataStorageFactory
 {
diff --git a/Engine/Plugins/ChaosVD/Source/ChaosVD/Private/TEDS/ChaosVDTedsUtils.cpp b/Engine/Plugins/ChaosVD/Source/ChaosVD/Private/TEDS/ChaosVDTedsUtils.cpp
new file mode 100644
--- /dev/null
+++ b/Engine/Plugins/ChaosVD/Source/ChaosVD/Private/TEDS/ChaosVDTedsUtils.cpp
@@ -0,0 +1,3 @@
+// Copyright Epic Games, Inc. All Rights Reserved.
+
+#include "ChaosVDTedsUtils.h"
diff --git a/Engine/Plugins/ChaosVD/Source/ChaosVD/Private/Trace/DataProcessors/ChaosVDTraceParticleDataProcessor.cpp b/Engine/Plugins/ChaosVD/Source/ChaosVD/Private/Trace/DataProcessors/ChaosVDTraceParticleDataProcessor.cpp
--- a/Engine/Plugins/ChaosVD/Source/ChaosVD/Private/Trace/DataProcessors/ChaosVDTraceParticleDataProcessor.cpp
+++ b/Engine/Plugins/ChaosVD/Source/ChaosVD/Private/Trace/DataProcessors/ChaosVDTraceParticleDataProcessor.cpp
@@ -33,6 +33,9 @@ bool FChaosVDTraceParticleDataProcessor::ProcessRawData(const TArray<uint8>& InD
 		{
 			int32 NewParticleDataIndex = INDEX_NONE;
 
+			FString& DebugName = ParticleData->DebugName;
+			ParticleData->DebugNameHash = CityHash64(reinterpret_cast<const char*>(*DebugName), DebugName.Len() * sizeof(TCHAR));
+
 			// Non-Explicit stages might have duplicated data if a particle was traced twice. As CVD only shows the last state per stage, we need to just keep the last seen state
 			// We do this on trace analysis to not do the de-duplication process during playback
 			if (!EnumHasAnyFlags(CurrentSolverStage->StageFlags, EChaosVDSolverStageFlags::ExplicitStage))
diff --git a/Engine/Plugins/ChaosVD/Source/ChaosVD/Private/Visualizers/ChaosVDParticleDataComponentVisualizer.cpp b/Engine/Plugins/ChaosVD/Source/ChaosVD/Private/Visualizers/ChaosVDParticleDataComponentVisualizer.cpp
--- a/Engine/Plugins/ChaosVD/Source/ChaosVD/Private/Visualizers/ChaosVDParticleDataComponentVisualizer.cpp
+++ b/Engine/Plugins/ChaosVD/Source/ChaosVD/Private/Visualizers/ChaosVDParticleDataComponentVisualizer.cpp
@@ -150,37 +150,28 @@ void FChaosVDParticleDataComponentVisualizer::DrawVisualization(const UActorComp
 	VisualizationContext.DebugDrawSettings = VisualizationSettings;
 	VisualizationContext.SolverDataSelectionObject = CVDScene->GetSolverDataSelectionObject().Pin();
 
-	
-
 	if (!VisualizationContext.IsVisualizationFlagEnabled(EChaosVDParticleDataVisualizationFlags::EnableDraw))
 	{
 		return;
 	}
 
+	FChaosVDSceneParticle* SelectedParticle = ParticleDataComponent->GetSelectedParticle();
+
 	if (VisualizationContext.IsVisualizationFlagEnabled(EChaosVDParticleDataVisualizationFlags::DrawDataOnlyForSelectedParticle))
 	{
-		VisualizationContext.bIsSelectedData = true;
-		ParticleDataComponent->VisitSelectedParticleData([this, PDI, View, &VisualizationContext, Component](const TSharedPtr<const FChaosVDParticleDataWrapper>& InParticleDataViewer)
+		if (SelectedParticle)
 		{
-			DrawVisualizationForParticleData(Component, PDI, View, VisualizationContext, InParticleDataViewer);
-
-			return true;
-		});
+			VisualizationContext.bIsSelectedData = true;
+			DrawVisualizationForParticleData(Component, PDI, View, VisualizationContext, *SelectedParticle);
+		}
 	}
 	else
 	{
-		int32 SelectedParticleID = INDEX_NONE;
-
-		FChaosVDSceneParticle* SelectedParticle = ParticleDataComponent->GetSelectedParticle();
-		if (TSharedPtr<const FChaosVDParticleDataWrapper> SelectedParticleData = SelectedParticle ? SelectedParticle->GetParticleData() : nullptr)
+		ParticleDataComponent->VisitAllParticleInstances([this, PDI, View, &VisualizationContext, Component, SelectedParticle](const TSharedRef<FChaosVDSceneParticle>& InParticleDataViewer)
 		{
-			SelectedParticleID = SelectedParticleData->ParticleIndex;
-		}
-
-		ParticleDataComponent->VisitAllParticleData([this, PDI, View, &VisualizationContext, Component, SelectedParticleID](const TSharedPtr<const FChaosVDParticleDataWrapper>& InParticleDataViewer)
-		{
-			VisualizationContext.bIsSelectedData = InParticleDataViewer && InParticleDataViewer->ParticleIndex == SelectedParticleID;
-			DrawVisualizationForParticleData(Component, PDI, View, VisualizationContext, InParticleDataViewer);
+			FChaosVDSceneParticle& ParticleInstanceRef = InParticleDataViewer.Get();
+			VisualizationContext.bIsSelectedData = SelectedParticle && (SelectedParticle == &ParticleInstanceRef);
+			DrawVisualizationForParticleData(Component, PDI, View, VisualizationContext, ParticleInstanceRef);
 
 			// If we reach the debug draw limit for this frame, there is no need to continue processing particles
 			return FChaosVDDebugDrawUtils::CanDebugDraw();
@@ -237,7 +228,7 @@ void FChaosVDParticleDataComponentVisualizer::DrawParticleVector(FPrimitiveDrawI
 	FChaosVDDebugDrawUtils::DrawArrowVector(PDI, StartLocation, StartLocation +  InVisualizationContext.DebugDrawSettings->GetScaleFortDataID(VectorID) * InVector, FText::AsCultureInvariant(DebugText), InVisualizationContext.DebugDrawSettings->ColorSettings.GetColorForDataID(VectorID, InVisualizationContext.bIsSelectedData),  InVisualizationContext.DebugDrawSettings->DepthPriority, LineThickness);
 }
 
-void FChaosVDParticleDataComponentVisualizer::DrawVisualizationForParticleData(const UActorComponent* Component, FPrimitiveDrawInterface* PDI, const FSceneView* View, const FChaosVDParticleDataVisualizationContext& InVisualizationContext, const TSharedPtr<const FChaosVDParticleDataWrapper>& InParticleDataViewer)
+void FChaosVDParticleDataComponentVisualizer::DrawVisualizationForParticleData(const UActorComponent* Component, FPrimitiveDrawInterface* PDI, const FSceneView* View, const FChaosVDParticleDataVisualizationContext& InVisualizationContext, const FChaosVDSceneParticle& InParticleInstance)
 {
 	using namespace Chaos::VisualDebugger::ParticleDataUnitsStrings;
 
@@ -256,41 +247,56 @@ void FChaosVDParticleDataComponentVisualizer::DrawVisualizationForParticleData(c
 		return;
 	}
 
-	const FVector& OwnerLocation = InVisualizationContext.SpaceTransform.TransformPosition(InParticleDataViewer->ParticlePositionRotation.MX);
 
-	// TODO: See how expensive is get the bounds. It is not something we have recorded
-	constexpr float VisibleRadius = 50.0f;
-	if (!View->ViewFrustum.IntersectSphere(OwnerLocation, VisibleRadius))
+
+	FBox ParticleBounds = InParticleInstance.GetBoundingBox();
+
+	if (!View->ViewFrustum.IntersectBox(ParticleBounds.GetCenter(), ParticleBounds.GetExtent()))
 	{
 		// If this particle location is not even visible, just ignore it.
 		return;
 	}
 
-	const FQuat& OwnerRotation =  InVisualizationContext.SpaceTransform.TransformRotation(InParticleDataViewer->ParticlePositionRotation.MR);
-	const FVector OwnerCoMLocation = InVisualizationContext.SpaceTransform.TransformPosition(InParticleDataViewer->ParticlePositionRotation.MX + (InParticleDataViewer->ParticlePositionRotation.MR *  InParticleDataViewer->ParticleMassProps.MCenterOfMass));
+	TSharedPtr<const FChaosVDParticleDataWrapper> ParticleData = InParticleInstance.GetParticleData();
+
+	if (!ParticleData)
+	{
+		return;
+	}
+
+	const FVector& OwnerLocation = InVisualizationContext.SpaceTransform.TransformPosition(ParticleData->ParticlePositionRotation.MX);
+
+	const FQuat& OwnerRotation =  InVisualizationContext.SpaceTransform.TransformRotation(ParticleData->ParticlePositionRotation.MR);
+	const FVector OwnerCoMLocation = InVisualizationContext.SpaceTransform.TransformPosition(ParticleData->ParticlePositionRotation.MX + (ParticleData->ParticlePositionRotation.MR *  ParticleData->ParticleMassProps.MCenterOfMass));
 	
-	FChaosVDScopedParticleHitProxy ScopedHitProxy(PDI, new HChaosVDComponentVisProxy(Component, InVisualizationContext.SolverDataSelectionObject->MakeSelectionHandle(ConstCastSharedPtr<FChaosVDParticleDataWrapper>(InParticleDataViewer))));
+	FChaosVDScopedParticleHitProxy ScopedHitProxy(PDI, new HChaosVDComponentVisProxy(Component, InVisualizationContext.SolverDataSelectionObject->MakeSelectionHandle(ConstCastSharedPtr<FChaosVDParticleDataWrapper>(ParticleData))));
 
 	constexpr float DefaultLineThickness = 1.5f;
 	constexpr float SelectedLineThickness = 3.5f;
 	const float LineThickness = InVisualizationContext.bIsSelectedData ? SelectedLineThickness : DefaultLineThickness;
-	
 
-	if (InParticleDataViewer->ParticleVelocities.HasValidData())
+	if (InVisualizationContext.IsVisualizationFlagEnabled(EChaosVDParticleDataVisualizationFlags::Bounds))
+	{
+		FTransform Location;
+		Location.SetLocation(ParticleBounds.GetCenter());
+		FChaosVDDebugDrawUtils::DrawBox(PDI, ParticleBounds.GetExtent(), FColor::Red, Location, FText::GetEmpty(), InVisualizationContext.DebugDrawSettings->DepthPriority, LineThickness);
+	}
+
+	if (ParticleData->ParticleVelocities.HasValidData())
 	{
-		DrawParticleVector(PDI, OwnerCoMLocation, InParticleDataViewer->ParticleVelocities.MV,EChaosVDParticleDataVisualizationFlags::Velocity, InVisualizationContext, LineThickness); 
-		DrawParticleVector(PDI, OwnerCoMLocation, InParticleDataViewer->ParticleVelocities.MW,EChaosVDParticleDataVisualizationFlags::AngularVelocity, InVisualizationContext, LineThickness); 
+		DrawParticleVector(PDI, OwnerCoMLocation, ParticleData->ParticleVelocities.MV,EChaosVDParticleDataVisualizationFlags::Velocity, InVisualizationContext, LineThickness); 
+		DrawParticleVector(PDI, OwnerCoMLocation, ParticleData->ParticleVelocities.MW,EChaosVDParticleDataVisualizationFlags::AngularVelocity, InVisualizationContext, LineThickness); 
 	}
 
-	if (InParticleDataViewer->ParticleDynamics.HasValidData())
+	if (ParticleData->ParticleDynamics.HasValidData())
 	{
-		DrawParticleVector(PDI, OwnerCoMLocation, InParticleDataViewer->ParticleDynamics.MAcceleration,EChaosVDParticleDataVisualizationFlags::Acceleration, InVisualizationContext, LineThickness); 
-		DrawParticleVector(PDI, OwnerCoMLocation, InParticleDataViewer->ParticleDynamics.MAngularAcceleration,EChaosVDParticleDataVisualizationFlags::AngularAcceleration, InVisualizationContext, LineThickness); 
-		DrawParticleVector(PDI, OwnerCoMLocation, InParticleDataViewer->ParticleDynamics.MLinearImpulseVelocity,EChaosVDParticleDataVisualizationFlags::LinearImpulse, InVisualizationContext, LineThickness); 
-		DrawParticleVector(PDI, OwnerCoMLocation, InParticleDataViewer->ParticleDynamics.MAngularImpulseVelocity,EChaosVDParticleDataVisualizationFlags::AngularImpulse, InVisualizationContext, LineThickness);
+		DrawParticleVector(PDI, OwnerCoMLocation, ParticleData->ParticleDynamics.MAcceleration,EChaosVDParticleDataVisualizationFlags::Acceleration, InVisualizationContext, LineThickness); 
+		DrawParticleVector(PDI, OwnerCoMLocation, ParticleData->ParticleDynamics.MAngularAcceleration,EChaosVDParticleDataVisualizationFlags::AngularAcceleration, InVisualizationContext, LineThickness); 
+		DrawParticleVector(PDI, OwnerCoMLocation, ParticleData->ParticleDynamics.MLinearImpulseVelocity,EChaosVDParticleDataVisualizationFlags::LinearImpulse, InVisualizationContext, LineThickness); 
+		DrawParticleVector(PDI, OwnerCoMLocation, ParticleData->ParticleDynamics.MAngularImpulseVelocity,EChaosVDParticleDataVisualizationFlags::AngularImpulse, InVisualizationContext, LineThickness);
 	}
 
-	if (InParticleDataViewer->ParticleMassProps.HasValidData())
+	if (ParticleData->ParticleMassProps.HasValidData())
 	{
 		if (InVisualizationContext.IsVisualizationFlagEnabled(EChaosVDParticleDataVisualizationFlags::CenterOfMass))
 		{
@@ -306,11 +312,11 @@ void FChaosVDParticleDataComponentVisualizer::DrawVisualizationForParticleData(c
 	}
 
 	// TODO: This is a Proof of concept to test how debug draw connectivity data will look
-	if (InParticleDataViewer->ParticleCluster.HasValidData())
+	if (ParticleData->ParticleCluster.HasValidData())
 	{
 		if (InVisualizationContext.IsVisualizationFlagEnabled(EChaosVDParticleDataVisualizationFlags::ClusterConnectivityEdge))
 		{
-			for (const FChaosVDConnectivityEdge& ConnectivityEdge : InParticleDataViewer->ParticleCluster.ConnectivityEdges)
+			for (const FChaosVDConnectivityEdge& ConnectivityEdge : ParticleData->ParticleCluster.ConnectivityEdges)
 			{
 				if (const TSharedPtr<FChaosVDScene> ScenePtr = InVisualizationContext.CVDScene.Pin())
 				{
diff --git a/Engine/Plugins/ChaosVD/Source/ChaosVD/Private/Visualizers/ChaosVDParticleDataComponentVisualizer.h b/Engine/Plugins/ChaosVD/Source/ChaosVD/Private/Visualizers/ChaosVDParticleDataComponentVisualizer.h
--- a/Engine/Plugins/ChaosVD/Source/ChaosVD/Private/Visualizers/ChaosVDParticleDataComponentVisualizer.h
+++ b/Engine/Plugins/ChaosVD/Source/ChaosVD/Private/Visualizers/ChaosVDParticleDataComponentVisualizer.h
@@ -8,6 +8,7 @@
 #include "Settings/ChaosVDParticleVisualizationSettings.h"
 #include "Templates/SharedPointer.h"
 
+struct FChaosVDSceneParticle;
 class FChaosVDGeometryBuilder;
 class UChaosVDParticleVisualizationDebugDrawSettings;
 
@@ -49,5 +50,5 @@ public:
 protected:
 	
 	void DrawParticleVector(FPrimitiveDrawInterface* PDI, const FVector& StartLocation, const FVector& InVector, EChaosVDParticleDataVisualizationFlags VectorID, const FChaosVDParticleDataVisualizationContext& InVisualizationContext, float LineThickness);
-	void DrawVisualizationForParticleData(const UActorComponent* Component, FPrimitiveDrawInterface* PDI, const FSceneView* View, const FChaosVDParticleDataVisualizationContext& InVisualizationContext, const TSharedPtr<const FChaosVDParticleDataWrapper>& InParticleDataViewer);
+	void DrawVisualizationForParticleData(const UActorComponent* Component, FPrimitiveDrawInterface* PDI, const FSceneView* View, const FChaosVDParticleDataVisualizationContext& InVisualizationContext, const FChaosVDSceneParticle& InParticleInstance);
 };
diff --git a/Engine/Plugins/ChaosVD/Source/ChaosVD/Private/Widgets/SChaosVDGameFramesPlaybackControls.cpp b/Engine/Plugins/ChaosVD/Source/ChaosVD/Private/Widgets/SChaosVDGameFramesPlaybackControls.cpp
--- a/Engine/Plugins/ChaosVD/Source/ChaosVD/Private/Widgets/SChaosVDGameFramesPlaybackControls.cpp
+++ b/Engine/Plugins/ChaosVD/Source/ChaosVD/Private/Widgets/SChaosVDGameFramesPlaybackControls.cpp
@@ -16,11 +16,29 @@ void SChaosVDGameFramesPlaybackControls::Construct(const FArguments& InArgs, con
 			.MaxFrames_Raw(this, &SChaosVDGameFramesPlaybackControls::GetMaxFrames)
 			.CurrentFrame_Raw(this, &SChaosVDGameFramesPlaybackControls::GetCurrentFrame)
 			.ButtonEnabledFlags_Raw(this, &SChaosVDGameFramesPlaybackControls::GetElementEnabledFlags)
+			.OnTimelineScrubStart(this, &SChaosVDGameFramesPlaybackControls::HandleTimelineScrubStart)
+			.OnTimelineScrubEnd(this, &SChaosVDGameFramesPlaybackControls::HandleTimelineScrubEnd)
 	];
 
 	RegisterNewController(InPlaybackController);
 }
 
+void SChaosVDGameFramesPlaybackControls::HandleTimelineScrubStart()
+{
+	if (const TSharedPtr<FChaosVDPlaybackController> PlaybackControllerPtr = PlaybackController.Pin())
+	{
+		PlaybackControllerPtr->SetScrubbingTimeline(true);
+	}
+}
+
+void SChaosVDGameFramesPlaybackControls::HandleTimelineScrubEnd()
+{
+	if (const TSharedPtr<FChaosVDPlaybackController> PlaybackControllerPtr = PlaybackController.Pin())
+	{
+		PlaybackControllerPtr->SetScrubbingTimeline(false);
+	}
+}
+
 void SChaosVDGameFramesPlaybackControls::OnFrameSelectionUpdated(int32 NewFrameIndex)
 {
 	if (const TSharedPtr<FChaosVDPlaybackController> PlaybackControllerPtr = PlaybackController.Pin())
diff --git a/Engine/Plugins/ChaosVD/Source/ChaosVD/Private/Widgets/SChaosVDGameFramesPlaybackControls.h b/Engine/Plugins/ChaosVD/Source/ChaosVD/Private/Widgets/SChaosVDGameFramesPlaybackControls.h
--- a/Engine/Plugins/ChaosVD/Source/ChaosVD/Private/Widgets/SChaosVDGameFramesPlaybackControls.h
+++ b/Engine/Plugins/ChaosVD/Source/ChaosVD/Private/Widgets/SChaosVDGameFramesPlaybackControls.h
@@ -38,6 +38,9 @@ protected:
 
 	virtual void HandlePlaybackControllerDataUpdated(TWeakPtr<FChaosVDPlaybackController> InController) override;
 
+	void HandleTimelineScrubStart();
+	void HandleTimelineScrubEnd();
+
 	bool CanPlayback() const;
 
 	bool IsPlaying() const;
diff --git a/Engine/Plugins/ChaosVD/Source/ChaosVD/Private/Widgets/SChaosVDMainTab.cpp b/Engine/Plugins/ChaosVD/Source/ChaosVD/Private/Widgets/SChaosVDMainTab.cpp
--- a/Engine/Plugins/ChaosVD/Source/ChaosVD/Private/Widgets/SChaosVDMainTab.cpp
+++ b/Engine/Plugins/ChaosVD/Source/ChaosVD/Private/Widgets/SChaosVDMainTab.cpp
@@ -48,6 +48,7 @@
 
 #include "Features/IModularFeatures.h"
 #include "HAL/FileManager.h"
+#include "Settings/ChaosVDGeneralSettings.h"
 #include "Settings/ChaosVDMiscSettings.h"
 #include "Styling/StyleColors.h"
 #include "Styling/ToolBarStyle.h"
@@ -159,15 +160,23 @@ void SChaosVDMainTab::Construct(const FArguments& InArgs, TSharedPtr<FChaosVDEng
 				[
 					SNew(SHorizontalBox)
 					+SHorizontalBox::Slot()
+					.HAlign(HAlign_Left)
 					.AutoWidth()
 					[
 						GenerateMainToolbarWidget()
 					]
 					+SHorizontalBox::Slot()
+					.HAlign(HAlign_Left)
 					.AutoWidth()
 					[
 						SNew(SChaosVDRecordingControls, StaticCastSharedRef<SChaosVDMainTab>(AsShared()))
 					]
+					+SHorizontalBox::Slot()
+					.HAlign(HAlign_Right)
+					.AutoWidth()
+					[
+						GenerateSettingsMenuWidget()
+					]
 				]
 			]
 		]
@@ -479,6 +488,21 @@ TSharedRef<SWidget> SChaosVDMainTab::GenerateMainToolbarWidget()
 	return UToolMenus::Get()->GenerateWidget(MainToolBarName, MenuContext);
 }
 
+TSharedRef<SWidget> SChaosVDMainTab::GenerateSettingsMenuWidget()
+{
+	RegisterSettingsMenu();
+
+	FToolMenuContext MenuContext;
+
+	UChaosVDMainToolbarMenuContext* CommonContextObject = NewObject<UChaosVDMainToolbarMenuContext>();
+	CommonContextObject->MainTab = SharedThis(this);
+
+	MenuContext.AddObject(CommonContextObject);
+
+	const FName SettingsMenuName = FName("ChaosVDMainTabSettingsMenu");
+	return UToolMenus::Get()->GenerateWidget(SettingsMenuName, MenuContext);
+}
+
 void SChaosVDMainTab::LoadCVDFile(const FString& InFilename)
 {
 	if (ensure(InFilename.EndsWith(TEXT("utrace"))))
@@ -815,6 +839,29 @@ void SChaosVDMainTab::RegisterMainTabMenu()
 	}));
 }
 
+void SChaosVDMainTab::RegisterSettingsMenu()
+{
+	const UToolMenus* ToolMenus = UToolMenus::Get();
+	const FName SettingsMenuName = FName("ChaosVDMainTabSettingsMenu");
+	if (ToolMenus->IsMenuRegistered(SettingsMenuName))
+	{
+		return;
+	}
+
+	UToolMenu* ToolBar = UToolMenus::Get()->RegisterMenu(SettingsMenuName, NAME_None, EMultiBoxType::SlimHorizontalToolBar);
+	if (!ToolBar)
+	{
+		return;
+	}
+
+	using namespace Chaos::VisualDebugger::Utils;
+
+	FToolMenuSection& Section = ToolBar->AddSection("SettingsMenu");
+	FNewToolMenuDelegate MainSettingsMenuBuilder = FNewToolMenuDelegate::CreateStatic(&CreateMenuEntryForSettingsObject<UChaosVDGeneralSettings>, EChaosVDSaveSettingsOptions::ShowResetButton);
+	constexpr bool bOpenSubMenuOnClick = true;
+	Section.AddSubMenu(TEXT("MainSettingsMenu"), LOCTEXT("MainSettingsMenuLabel", "Settings"), LOCTEXT("MainSettingsMenuTip", "Settings that controls general behavior of CVD"), MainSettingsMenuBuilder, bOpenSubMenuOnClick,  FSlateIcon(FAppStyle::Get().GetStyleSetName(), TEXT("Icons.Toolbar.Settings")));
+}
+
 void SChaosVDMainTab::BrowseLiveSessionsFromTraceStore() const
 {
 	const TSharedRef<SChaosVDBrowseSessionsModal> SessionBrowserModal = SNew(SChaosVDBrowseSessionsModal);
diff --git a/Engine/Plugins/ChaosVD/Source/ChaosVD/Private/Widgets/SChaosVDSceneQueryDataInspector.cpp b/Engine/Plugins/ChaosVD/Source/ChaosVD/Private/Widgets/SChaosVDSceneQueryDataInspector.cpp
--- a/Engine/Plugins/ChaosVD/Source/ChaosVD/Private/Widgets/SChaosVDSceneQueryDataInspector.cpp
+++ b/Engine/Plugins/ChaosVD/Source/ChaosVD/Private/Widgets/SChaosVDSceneQueryDataInspector.cpp
@@ -446,8 +446,8 @@ FReply SChaosVDSceneQueryDataInspector::SelectParticleForCurrentQueryData() cons
 		// TODO : This will not work properly when the visited shape was a Union within a union
 		// CVD currently doesn't support multi selection, so we can't easily select all mesh instances that represent a union within a union
 		// We need to revisit this when multi selection support is added. Jira for tracking UE-212733
-		const TConstArrayView<TSharedRef<FChaosVDMeshDataInstanceHandle>> AvailableMeshInstances = ParticleActor->GetMeshInstances();
-		for (const TSharedRef<FChaosVDMeshDataInstanceHandle>& MeshInstance : AvailableMeshInstances)
+		const TConstArrayView<TSharedRef<FChaosVDInstancedMeshData>> AvailableMeshInstances = ParticleActor->GetMeshInstances();
+		for (const TSharedRef<FChaosVDInstancedMeshData>& MeshInstance : AvailableMeshInstances)
 		{
 			if (MeshInstance->GetState().ImplicitObjectInfo.ShapeInstanceIndex == ShapeInstanceIndexToSelect)
 			{
diff --git a/Engine/Plugins/ChaosVD/Source/ChaosVD/Private/Widgets/SChaosVDSolverPlaybackControls.cpp b/Engine/Plugins/ChaosVD/Source/ChaosVD/Private/Widgets/SChaosVDSolverPlaybackControls.cpp
--- a/Engine/Plugins/ChaosVD/Source/ChaosVD/Private/Widgets/SChaosVDSolverPlaybackControls.cpp
+++ b/Engine/Plugins/ChaosVD/Source/ChaosVD/Private/Widgets/SChaosVDSolverPlaybackControls.cpp
@@ -61,6 +61,8 @@ void SChaosVDSolverPlaybackControls::Construct(const FArguments& InArgs, const T
 					.CurrentFrame_Raw(this, &SChaosVDSolverPlaybackControls::GetCurrentFrame)
 					.OnFrameChanged_Raw(this, &SChaosVDSolverPlaybackControls::OnFrameSelectionUpdated)
 					.OnButtonClicked_Raw(this, &SChaosVDSolverPlaybackControls::HandleFramePlaybackButtonClicked)
+					.OnTimelineScrubStart(this, &SChaosVDSolverPlaybackControls::HandleTimelineScrubStart)
+					.OnTimelineScrubEnd(this, &SChaosVDSolverPlaybackControls::HandleTimelineScrubEnd)
 				]
 				+SHorizontalBox::Slot()
 				.Padding(6.0f,0.0f)
@@ -109,6 +111,8 @@ void SChaosVDSolverPlaybackControls::Construct(const FArguments& InArgs, const T
 					.MinFrames_Raw(this, &SChaosVDSolverPlaybackControls::GetMinSolverStage)
 					.CurrentFrame_Raw(this, &SChaosVDSolverPlaybackControls::GetCurrentSolverStage)
 					.OnButtonClicked_Raw(this, &SChaosVDSolverPlaybackControls::HandleSolverStagePlaybackButtonClicked)
+					.OnTimelineScrubStart(this, &SChaosVDSolverPlaybackControls::HandleTimelineScrubStart)
+					.OnTimelineScrubEnd(this, &SChaosVDSolverPlaybackControls::HandleTimelineScrubEnd)
 				]
 				+SHorizontalBox::Slot()
 				.AutoWidth()
@@ -138,6 +142,22 @@ SChaosVDSolverPlaybackControls::~SChaosVDSolverPlaybackControls()
 	}
 }
 
+void SChaosVDSolverPlaybackControls::HandleTimelineScrubStart()
+{
+	if (const TSharedPtr<FChaosVDPlaybackController> PlaybackControllerPtr = PlaybackController.Pin())
+	{
+		PlaybackControllerPtr->SetScrubbingTimeline(true);
+	}
+}
+
+void SChaosVDSolverPlaybackControls::HandleTimelineScrubEnd()
+{
+	if (const TSharedPtr<FChaosVDPlaybackController> PlaybackControllerPtr = PlaybackController.Pin())
+	{
+		PlaybackControllerPtr->SetScrubbingTimeline(false);
+	}
+}
+
 bool SChaosVDSolverPlaybackControls::CanPlayback() const
 {
 	bool bCanControlPlayback = false;
diff --git a/Engine/Plugins/ChaosVD/Source/ChaosVD/Private/Widgets/SChaosVDSolverPlaybackControls.h b/Engine/Plugins/ChaosVD/Source/ChaosVD/Private/Widgets/SChaosVDSolverPlaybackControls.h
--- a/Engine/Plugins/ChaosVD/Source/ChaosVD/Private/Widgets/SChaosVDSolverPlaybackControls.h
+++ b/Engine/Plugins/ChaosVD/Source/ChaosVD/Private/Widgets/SChaosVDSolverPlaybackControls.h
@@ -47,6 +47,8 @@ public:
 	void Construct(const FArguments& InArgs, const TSharedRef<const FChaosVDTrackInfo>& InSolverTrackInfo, const TWeakPtr<FChaosVDPlaybackController>& InPlaybackController);
 
 	virtual ~SChaosVDSolverPlaybackControls() override;
+	void HandleTimelineScrubStart();
+	void HandleTimelineScrubEnd();
 
 private:
 
diff --git a/Engine/Plugins/ChaosVD/Source/ChaosVD/Private/Widgets/SChaosVDTimelineWidget.cpp b/Engine/Plugins/ChaosVD/Source/ChaosVD/Private/Widgets/SChaosVDTimelineWidget.cpp
--- a/Engine/Plugins/ChaosVD/Source/ChaosVD/Private/Widgets/SChaosVDTimelineWidget.cpp
+++ b/Engine/Plugins/ChaosVD/Source/ChaosVD/Private/Widgets/SChaosVDTimelineWidget.cpp
@@ -21,6 +21,8 @@ void SChaosVDTimelineWidget::Construct(const FArguments& InArgs)
 	ButtonClickedDelegate = InArgs._OnButtonClicked;
 	ElementVisibilityFlags = InArgs._ButtonVisibilityFlags;
 	ElementEnabledFlags = InArgs._ButtonEnabledFlags;
+	TimelineScrubStartDelegate = InArgs._OnTimelineScrubStart;
+	TimelineScrubEndDelegate = InArgs._OnTimelineScrubEnd;
 	bIsPlaying = InArgs._IsPlaying;
 
 	if (!ElementVisibilityFlags.IsSet())
@@ -35,6 +37,8 @@ void SChaosVDTimelineWidget::Construct(const FArguments& InArgs)
 
 	SetCanTick(true);
 
+	PendingValueChange.Reset();
+
 	ChildSlot
 	[
 		SNew(SHorizontalBox)
@@ -119,14 +123,16 @@ void SChaosVDTimelineWidget::Construct(const FArguments& InArgs)
 			.Padding(4.0f,0.0f)
 			.FillWidth(1.0f)
 			[
-			  SAssignNew(TimelineSlider, SSlider)
-			  .Visibility_Raw(this, &SChaosVDTimelineWidget::GetElementVisibility, EChaosVDTimelineElementIDFlags::Timeline)
-			  .ToolTipText_Lambda([this]()-> FText{ return FText::AsNumber(GetCurrentFrame()); })
-			  .Value_Raw(this, &SChaosVDTimelineWidget::GetCurrentFrameAsFloat)
-			  .OnValueChanged_Raw(this, &SChaosVDTimelineWidget::SetCurrentTimelineFrame, EChaosVDSetTimelineFrameFlags::BroadcastChange)
-			  .StepSize(1)
-			  .MaxValue(0)
-			  .MinValue(0)
+				SAssignNew(TimelineSlider, SSlider)
+				.Visibility_Raw(this, &SChaosVDTimelineWidget::GetElementVisibility, EChaosVDTimelineElementIDFlags::Timeline)
+				.ToolTipText_Lambda([this]()-> FText{ return FText::AsNumber(GetCurrentFrame()); })
+				.Value_Raw(this, &SChaosVDTimelineWidget::GetCurrentFrameAsFloat)
+				.OnValueChanged_Raw(this, &SChaosVDTimelineWidget::SetCurrentTimelineFrame, EChaosVDSetTimelineFrameFlags::BroadcastChange)
+				.StepSize(1)
+				.MaxValue(0)
+				.MinValue(0)
+				.OnMouseCaptureBegin_Raw(this, &SChaosVDTimelineWidget::HandleTimelineScrubStart)
+				.OnMouseCaptureEnd_Raw(this, &SChaosVDTimelineWidget::HandleTimelineScrubEnd)
 			]
 			+SHorizontalBox::Slot()
 			.Padding(4.0f,0.0f)
@@ -173,7 +179,7 @@ void SChaosVDTimelineWidget::SetCurrentTimelineFrame(float FrameNumber, EChaosVD
 	{
 		if (EnumHasAnyFlags(Options, EChaosVDSetTimelineFrameFlags::BroadcastChange))
 		{
-			FrameChangedDelegate.ExecuteIfBound(static_cast<int32>(FrameNumber));
+			PendingValueChange = FrameNumber;
 		}
 	}
 }
@@ -190,6 +196,12 @@ void SChaosVDTimelineWidget::Tick(const FGeometry& AllottedGeometry, const doubl
 		{
 			TimelineSlider->SetMinAndMaxValues(static_cast<float>(CurrentMinFrames), static_cast<float>(CurrentMaxFrames));
 		}
+
+		if (PendingValueChange.IsSet())
+		{
+			FrameChangedDelegate.ExecuteIfBound(static_cast<int32>(PendingValueChange.GetValue()));
+			PendingValueChange.Reset();
+		}
 	}
 }
 
@@ -254,10 +266,19 @@ int32 SChaosVDTimelineWidget::GetCurrentMinFrames() const
 	return MinFrames.IsSet() ? MinFrames.Get() : INDEX_NONE;
 }
 
-
 int32 SChaosVDTimelineWidget::GetCurrentMaxFrames() const
 {
 	return MaxFrames.IsSet() ? MaxFrames.Get() : INDEX_NONE;
 }
 
+void SChaosVDTimelineWidget::HandleTimelineScrubStart()
+{
+	TimelineScrubStartDelegate.ExecuteIfBound();
+}
+
+void SChaosVDTimelineWidget::HandleTimelineScrubEnd()
+{
+	TimelineScrubEndDelegate.ExecuteIfBound();
+}
+
 #undef LOCTEXT_NAMESPACE
diff --git a/Engine/Plugins/ChaosVD/Source/ChaosVD/Private/Widgets/SChaosVDTimelineWidget.h b/Engine/Plugins/ChaosVD/Source/ChaosVD/Private/Widgets/SChaosVDTimelineWidget.h
--- a/Engine/Plugins/ChaosVD/Source/ChaosVD/Private/Widgets/SChaosVDTimelineWidget.h
+++ b/Engine/Plugins/ChaosVD/Source/ChaosVD/Private/Widgets/SChaosVDTimelineWidget.h
@@ -50,6 +50,7 @@ enum class EChaosVDTimelineElementIDFlags : uint16
 ENUM_CLASS_FLAGS(EChaosVDTimelineElementIDFlags)
 
 DECLARE_DELEGATE_OneParam(FChaosVDFrameChangedDelegate, int32)
+DECLARE_DELEGATE(FChaosVDFrameScrubDelegate)
 
 /** Simple timeline control widget */
 class SChaosVDTimelineWidget : public SCompoundWidget
@@ -65,6 +66,8 @@ public:
 		SLATE_ATTRIBUTE(bool, IsPlaying)
 		SLATE_EVENT(FChaosVDFrameChangedDelegate, OnFrameChanged)
 		SLATE_EVENT(FChaosControlButtonClicked, OnButtonClicked)
+		SLATE_EVENT(FChaosVDFrameScrubDelegate, OnTimelineScrubStart)
+		SLATE_EVENT(FChaosVDFrameScrubDelegate, OnTimelineScrubEnd)
 	SLATE_END_ARGS()
 
 	void Construct(const FArguments& InArgs);
@@ -95,6 +98,11 @@ protected:
 	int32 GetCurrentMinFrames() const;
 	int32 GetCurrentMaxFrames() const;
 
+	void HandleTimelineScrubStart();
+	void HandleTimelineScrubEnd();
+
+	TOptional<float> PendingValueChange = 0.0f;
+
 	TSharedPtr<SSlider> TimelineSlider;
 
 	TAttribute<int32> CurrentFrame = 0;
@@ -106,6 +114,8 @@ protected:
 
 	FChaosVDFrameChangedDelegate FrameChangedDelegate;
 	FChaosControlButtonClicked ButtonClickedDelegate;
+	FChaosVDFrameScrubDelegate TimelineScrubStartDelegate;
+	FChaosVDFrameScrubDelegate TimelineScrubEndDelegate;
 
 	TAttribute<EChaosVDTimelineElementIDFlags> ElementVisibilityFlags = EChaosVDTimelineElementIDFlags::All;
 
diff --git a/Engine/Plugins/ChaosVD/Source/ChaosVD/Public/Actors/ChaosVDDataContainerBaseActor.h b/Engine/Plugins/ChaosVD/Source/ChaosVD/Public/Actors/ChaosVDDataContainerBaseActor.h
--- a/Engine/Plugins/ChaosVD/Source/ChaosVD/Public/Actors/ChaosVDDataContainerBaseActor.h
+++ b/Engine/Plugins/ChaosVD/Source/ChaosVD/Public/Actors/ChaosVDDataContainerBaseActor.h
@@ -34,6 +34,8 @@ public:
 	{
 		SceneWeakPtr = InScene;
 	}
+
+	virtual void HandleWorldStreamingLocationUpdated(const FVector& InLocation);
 	
 	void SetSolverID(int32 InSolverID);
 
diff --git a/Engine/Plugins/ChaosVD/Source/ChaosVD/Public/ChaosVDGeometryBuilder.h b/Engine/Plugins/ChaosVD/Source/ChaosVD/Public/ChaosVDGeometryBuilder.h
--- a/Engine/Plugins/ChaosVD/Source/ChaosVD/Public/ChaosVDGeometryBuilder.h
+++ b/Engine/Plugins/ChaosVD/Source/ChaosVD/Public/ChaosVDGeometryBuilder.h
@@ -55,7 +55,7 @@ public:
 	}
 
 	virtual ~FChaosVDGeometryBuilder() override;
-	
+
 	void Initialize(const TWeakPtr<FChaosVDScene>& ChaosVDScene);
 	void DeInitialize();
 
@@ -137,7 +137,7 @@ public:
 	 * @param InExtractedGeometryDataHandle Handle to the extracted geometry data the new component will use
 	 * */
 	template<typename ComponentType>
-	TSharedPtr<FChaosVDMeshDataInstanceHandle> CreateMeshDataInstance(const FChaosVDParticleDataWrapper& InOwningParticleData, const TSharedPtr<FChaosVDExtractedGeometryDataHandle>& InExtractedGeometryDataHandle);
+	TSharedPtr<FChaosVDInstancedMeshData> CreateMeshDataInstance(const FChaosVDParticleDataWrapper& InOwningParticleData, const TSharedRef<FChaosVDExtractedGeometryDataHandle>& InExtractedGeometryDataHandle);
 
 	/**
 	 * Finds or creates a Mesh component compatible with the provided mesh data handle, and updates the handle to use that new component.
@@ -146,7 +146,7 @@ public:
 	 * @param MeshAttributes Attributes of the mesh that the new component needs to be compativle with
 	 * */
 	template<typename ComponentType>
-	void UpdateMeshDataInstance(TSharedPtr<FChaosVDMeshDataInstanceHandle> HandleToUpdate, EChaosVDMeshAttributesFlags MeshAttributes);
+	void UpdateMeshDataInstance(const TSharedRef<FChaosVDInstancedMeshData>& InHandleToUpdate, EChaosVDMeshAttributesFlags MeshAttributes);
 
 	/**
 	 * Destroys a Mesh component that will not longer be used.
@@ -161,13 +161,15 @@ public:
 	void RequestMaterialUpdate(UMeshComponent* MeshComponent);
 
 private:
+	
+	void CachePreBuiltMeshes();
 
 	/** Gets a ptr to a fully initialized Mesh component compatible with the provided geometry handle and mesh attribute flags, ready to accept a new mesh instance
 	 * @param GeometryDataHandle Handle with the data of the geometry to be used for the new instance
 	 * @param MeshAttributes Set of flags that the mesh component needs to be compatible with
 	 */
 	template<typename ComponentType>
-	ComponentType* GetMeshComponentForNewInstance(const TSharedPtr<FChaosVDExtractedGeometryDataHandle>& GeometryDataHandle, EChaosVDMeshAttributesFlags MeshAttributes);
+	ComponentType* GetMeshComponentForNewInstance(const TSharedRef<FChaosVDExtractedGeometryDataHandle>& GeometryDataHandle, EChaosVDMeshAttributesFlags MeshAttributes);
 
 	/** Gets a reference to the correct instanced static mesh component cache that is compatible with the provided mesh attribute flags
 	 * @param MeshAttributeFlags GeometryHandle this component will render
@@ -247,7 +249,6 @@ private:
 	/** Tick method of this Geometry builder. Used to do everything that needs to be performed in the GT, like applying the generated meshes to mesh component */
 	bool GameThreadTick(float DeltaTime);
 
-public:
 	/**
 	 * Add a mesh component to the waiting list for Geometry. This needs to be called before dispatching a generation job for new Geometry
 	 * @param GeometryKey Key of the geometry that is being generated (or will be generated)
@@ -262,13 +263,18 @@ public:
 	 */
 	void RemoveMeshComponentWaitingForGeometry(uint32 GeometryKey, TWeakObjectPtr<UMeshComponent> MeshComponent) const;
 
+	void RequestMeshForComponent(const TSharedRef<FChaosVDExtractedGeometryDataHandle>& SourceGeometry, UMeshComponent* MeshComponent);
+
 	/** Returns a reference to the Mesh components pool used by this builder */
 	FChaosVDMeshComponentPool& GetMeshComponentDataPool()
 	{
 		return ComponentMeshPool;
 	}
 
-private:
+	/** Returns true if the provided Implicit object type can use a pre-built static mesh*
+	 * @param ObjetType Type to evaluate
+	 */
+	bool UsesPreBuiltGeometry(Chaos::EImplicitObjectType ObjetType) const;
 
 	/** Handles any changes to the indexes of created instanced mesh components we are managing, making corrections/updates as needed */
 	void HandleStaticMeshComponentInstanceIndexUpdated(UInstancedStaticMeshComponent* InComponent, TArrayView<const FInstancedStaticMeshDelegates::FInstanceIndexUpdateData> InIndexUpdates);
@@ -313,6 +319,8 @@ private:
 	/** Weak Ptr to the CVD scene owning this geometry builder */
 	TWeakPtr<FChaosVDScene> SceneWeakPtr;
 
+	TObjectPtr<UStaticMesh> BasicShapesMeshes[Chaos::ImplicitObjectType::ConcreteObjectCount];	
+
 	bool bInitialized = false;
 
 	struct FSourceGeometryHashCache
@@ -334,7 +342,7 @@ private:
 					return *FoundHash;
 				}
 			}
-			
+
 			const uint32 GeometryHash = ImplicitObject->GetTypeHash();
 			CacheImplicitObjectHash(ImplicitObject, GeometryHash);
 
@@ -498,16 +506,10 @@ ComponentType* FChaosVDGeometryBuilder::GetAvailableMeshComponent(const TSharedP
 }
 
 template <typename ComponentType>
-TSharedPtr<FChaosVDMeshDataInstanceHandle> FChaosVDGeometryBuilder::CreateMeshDataInstance(const FChaosVDParticleDataWrapper& InOwningParticleData, const TSharedPtr<FChaosVDExtractedGeometryDataHandle>& InExtractedGeometryDataHandle)
+TSharedPtr<FChaosVDInstancedMeshData> FChaosVDGeometryBuilder::CreateMeshDataInstance(const FChaosVDParticleDataWrapper& InOwningParticleData, const TSharedRef<FChaosVDExtractedGeometryDataHandle>& InExtractedGeometryDataHandle)
 {
 	static_assert(std::is_base_of_v<UChaosVDStaticMeshComponent, ComponentType> || std::is_base_of_v<UChaosVDInstancedStaticMeshComponent, ComponentType>, "CreateMeshComponentsFromImplicit Only supports CVD versions of Static MeshComponent and Instanced Static Mesh Component");
 
-	if (!InExtractedGeometryDataHandle.IsValid())
-	{
-		UE_LOG(LogChaosVDEditor, Error, TEXT("[%s] Failed To Create mesh Component | Handle is invalid. "), ANSI_TO_TCHAR(__FUNCTION__));
-		return nullptr;
-	}
-
 	const FTransform ExtractedGeometryTransform = InExtractedGeometryDataHandle->GetRelativeTransform();
 
 	EChaosVDMeshAttributesFlags MeshComponentAttributeFlags = EChaosVDMeshAttributesFlags::None;
@@ -516,14 +518,14 @@ TSharedPtr<FChaosVDMeshDataInstanceHandle> FChaosVDGeometryBuilder::CreateMeshDa
 		EnumAddFlags(MeshComponentAttributeFlags, EChaosVDMeshAttributesFlags::MirroredGeometry);
 	}
 
-	TSharedPtr<FChaosVDMeshDataInstanceHandle> MeshComponentHandle;
+	TSharedPtr<FChaosVDInstancedMeshData> MeshComponentHandle;
 	ComponentType* Component = GetMeshComponentForNewInstance<ComponentType>(InExtractedGeometryDataHandle, MeshComponentAttributeFlags);
 
 	if (IChaosVDGeometryComponent* CVDGeometryComponent = Cast<IChaosVDGeometryComponent>(Component))
 	{
 		constexpr bool bIsWorldSpace = true;
 		const FTransform OwningParticleTransform(InOwningParticleData.ParticlePositionRotation.MR, InOwningParticleData.ParticlePositionRotation.MX);
-		MeshComponentHandle = CVDGeometryComponent->AddMeshInstance(OwningParticleTransform, bIsWorldSpace, InExtractedGeometryDataHandle.ToSharedRef(), InOwningParticleData.ParticleIndex, InOwningParticleData.SolverID);
+		MeshComponentHandle = CVDGeometryComponent->AddMeshInstance(OwningParticleTransform, bIsWorldSpace, InExtractedGeometryDataHandle, InOwningParticleData.ParticleIndex, InOwningParticleData.SolverID);
 		
 		MeshComponentHandle->SetGeometryBuilder(AsWeak());
 	}
@@ -532,34 +534,21 @@ TSharedPtr<FChaosVDMeshDataInstanceHandle> FChaosVDGeometryBuilder::CreateMeshDa
 }
 
 template <typename ComponentType>
-void FChaosVDGeometryBuilder::UpdateMeshDataInstance(TSharedPtr<FChaosVDMeshDataInstanceHandle> HandleToUpdate, EChaosVDMeshAttributesFlags MeshAttributes)
+void FChaosVDGeometryBuilder::UpdateMeshDataInstance(const TSharedRef<FChaosVDInstancedMeshData>& InHandleToUpdate, EChaosVDMeshAttributesFlags MeshAttributes)
 {
 	static_assert(std::is_base_of_v<UChaosVDStaticMeshComponent, ComponentType> || std::is_base_of_v<UChaosVDInstancedStaticMeshComponent, ComponentType>, "CreateMeshComponentsFromImplicit Only supports CVD versions of Static MeshComponent and Instanced Static Mesh Component");
 
-	if (!HandleToUpdate.IsValid())
-	{
-		UE_LOG(LogChaosVDEditor, Error, TEXT("[%s] Failed To Update Mesh Data Handle | Mesh Data Handle is invalid. "), ANSI_TO_TCHAR(__FUNCTION__));
-		return;
-	}
-
-	ComponentType* Component = GetMeshComponentForNewInstance<ComponentType>(HandleToUpdate->GetGeometryHandle(), MeshAttributes);
+	ComponentType* Component = GetMeshComponentForNewInstance<ComponentType>(InHandleToUpdate->GetGeometryHandle(), MeshAttributes);
 
 	if (IChaosVDGeometryComponent* CVDGeometryComponent = Cast<IChaosVDGeometryComponent>(Component))
 	{
-		constexpr bool bIsWorldSpace = true;
-		CVDGeometryComponent->AddMeshInstanceForHandle(HandleToUpdate, HandleToUpdate->GetWorldTransform(), bIsWorldSpace, HandleToUpdate->GetGeometryHandle(), HandleToUpdate->GetOwningParticleID(), HandleToUpdate->GetOwningSolverID());
+		CVDGeometryComponent->AddExistingMeshInstance(InHandleToUpdate);
 	}
 }
 
 template <typename ComponentType>
-ComponentType* FChaosVDGeometryBuilder::GetMeshComponentForNewInstance(const TSharedPtr<FChaosVDExtractedGeometryDataHandle>& GeometryDataHandle, EChaosVDMeshAttributesFlags MeshAttributes)
+ComponentType* FChaosVDGeometryBuilder::GetMeshComponentForNewInstance(const TSharedRef<FChaosVDExtractedGeometryDataHandle>& GeometryDataHandle, EChaosVDMeshAttributesFlags MeshAttributes)
 {
-	if (!GeometryDataHandle.IsValid())
-	{
-		UE_LOG(LogChaosVDEditor, Error, TEXT("[%s] Failed To obtain Mesh Component | Geometry Data Handle is invalid. "), ANSI_TO_TCHAR(__FUNCTION__));
-		return nullptr;
-	}
-
 	AActor* MeshComponentsContainerActor = nullptr;
 	
 	if (const TSharedPtr<FChaosVDScene> CVDScene = SceneWeakPtr.Pin())
@@ -577,16 +566,15 @@ ComponentType* FChaosVDGeometryBuilder::GetMeshComponentForNewInstance(const TSh
 
 	if (IChaosVDGeometryComponent* CVDGeometryComponent = Cast<IChaosVDGeometryComponent>(Component))
 	{
-		CVDGeometryComponent->SetMeshComponentAttributeFlags(MeshAttributes);
-
 		if (!CVDGeometryComponent->IsMeshReady())
 		{
-			AddMeshComponentWaitingForGeometry(GeometryDataHandle->GetGeometryKey(), Component);
+			RequestMeshForComponent(GeometryDataHandle, Component);
 		}
 
 		if (bIsNew)
 		{
 			CVDGeometryComponent->SetGeometryBuilder(AsWeak());
+			CVDGeometryComponent->SetMeshComponentAttributeFlags(MeshAttributes);
 			CVDGeometryComponent->Initialize();
 			CVDGeometryComponent->OnComponentEmpty()->AddRaw(this, &FChaosVDGeometryBuilder::DestroyMeshComponent);			
 		}
diff --git a/Engine/Plugins/ChaosVD/Source/ChaosVD/Public/ChaosVDPlaybackController.h b/Engine/Plugins/ChaosVD/Source/ChaosVD/Public/ChaosVDPlaybackController.h
--- a/Engine/Plugins/ChaosVD/Source/ChaosVD/Public/ChaosVDPlaybackController.h
+++ b/Engine/Plugins/ChaosVD/Source/ChaosVD/Public/ChaosVDPlaybackController.h
@@ -376,6 +376,13 @@ public:
 	/** Returns true if we are in continuous playback mode */
 	bool IsPlaying() const;
 
+	bool IsScrubbingTimeline() const
+	{
+		return bIsScrubbingTimeline;
+	}
+
+	void SetScrubbingTimeline(bool bNewIsScrubbingTimeline);
+
 	/** Returns the current sync mode for tracks */
 	EChaosVDSyncTimelinesMode GetTimelineSyncMode() const
 	{
@@ -420,6 +427,8 @@ protected:
 
 	TSharedRef<FChaosVDTrackInfo> CreateTrackInfo(int32 SlotIndex = INDEX_NONE);
 
+	bool bIsScrubbingTimeline = false;
+
 	/** Map containing all track info, by track type*/
 	TMap<EChaosVDTrackType, TrackInfoByIDMap> TrackInfoPerType;
 
diff --git a/Engine/Plugins/ChaosVD/Source/ChaosVD/Public/ChaosVDScene.h b/Engine/Plugins/ChaosVD/Source/ChaosVD/Public/ChaosVDScene.h
--- a/Engine/Plugins/ChaosVD/Source/ChaosVD/Public/ChaosVDScene.h
+++ b/Engine/Plugins/ChaosVD/Source/ChaosVD/Public/ChaosVDScene.h
@@ -224,8 +224,17 @@ public:
 
 	virtual bool Tick(float DeltaTime) override;
 
+	void UpdateWorldStreamingLocation(const FVector& InLocation);
+
+	const FVector& GetWorldStreamingLocation() const
+	{
+		return WorldStreamingLocation;
+	}
+
 private:
 
+	FVector WorldStreamingLocation = FVector::ZeroVector; 
+
 	void AddFromCVDWorldTagToActor(AActor* Actor);
 
 	FName TEDSSelectionSetName;
@@ -234,10 +243,7 @@ private:
 	
 	TSharedPtr<FChaosVDRecording> LoadedRecording;
 	
-	AActor* GetMeshComponentsContainerActor() const
-	{
-		return MeshComponentContainerActor;
-	}
+	AActor* GetMeshComponentsContainerActor() const;
 
 	AActor* GetSkySphereActor() const
 	{
@@ -301,7 +307,7 @@ private:
 
 	mutable AActor* SkySphere = nullptr;
 
-	AActor* MeshComponentContainerActor = nullptr;
+	AChaosVDGeometryContainer* MeshComponentContainerActor = nullptr;
 
 	AChaosVDGameFrameInfoActor* GameFrameDataInfoActor = nullptr;
 
diff --git a/Engine/Plugins/ChaosVD/Source/ChaosVD/Public/Components/ChaosVDSolverDataComponent.h b/Engine/Plugins/ChaosVD/Source/ChaosVD/Public/Components/ChaosVDSolverDataComponent.h
--- a/Engine/Plugins/ChaosVD/Source/ChaosVD/Public/Components/ChaosVDSolverDataComponent.h
+++ b/Engine/Plugins/ChaosVD/Source/ChaosVD/Public/Components/ChaosVDSolverDataComponent.h
@@ -30,6 +30,8 @@ public:
 	}
 
 	virtual void SetVisibility(bool bNewIsVisible);
+
+	virtual void HandleWorldStreamingLocationUpdated(const FVector& InLocation) {};
 	
 	void SetSolverID(int32 InSolverID)
 	{
diff --git a/Engine/Plugins/ChaosVD/Source/ChaosVD/Public/Settings/ChaosVDCoreSettings.h b/Engine/Plugins/ChaosVD/Source/ChaosVD/Public/Settings/ChaosVDCoreSettings.h
--- a/Engine/Plugins/ChaosVD/Source/ChaosVD/Public/Settings/ChaosVDCoreSettings.h
+++ b/Engine/Plugins/ChaosVD/Source/ChaosVD/Public/Settings/ChaosVDCoreSettings.h
@@ -97,4 +97,10 @@ public:
 
 	UPROPERTY(Config)
 	TSoftObjectPtr<UTextureCube> AmbientCubeMapTexture;
+
+	UPROPERTY(Config)
+	TSoftObjectPtr<UStaticMesh> BoxMesh;
+
+	UPROPERTY(Config)
+	TSoftObjectPtr<UStaticMesh> SphereMesh;
 };
diff --git a/Engine/Plugins/ChaosVD/Source/ChaosVD/Public/Widgets/SChaosVDMainTab.h b/Engine/Plugins/ChaosVD/Source/ChaosVD/Public/Widgets/SChaosVDMainTab.h
--- a/Engine/Plugins/ChaosVD/Source/ChaosVD/Public/Widgets/SChaosVDMainTab.h
+++ b/Engine/Plugins/ChaosVD/Source/ChaosVD/Public/Widgets/SChaosVDMainTab.h
@@ -118,6 +118,7 @@ private:
 	void CleanUpDisableCPUThrottlingDelegate() const;
 
 	void RegisterMainTabMenu();
+	void RegisterSettingsMenu();
 
 	void HandleTabSpawned(TSharedRef<SDockTab> Tab, FName TabID);
 	void HandleTabDestroyed(TSharedRef<SDockTab> Tab, FName TabID);
@@ -132,6 +133,7 @@ private:
 	TSharedRef<SButton> CreateSimpleButton(TFunction<FText()>&& GetTextDelegate, TFunction<FText()>&& ToolTipTextDelegate, const FSlateBrush* ButtonIcon, const UChaosVDMainToolbarMenuContext* MenuContext, const FOnClicked& InButtonClickedCallback);
 
 	TSharedRef<SWidget> GenerateMainToolbarWidget();
+	TSharedRef<SWidget> GenerateSettingsMenuWidget();
 
 	void BrowseChaosVDRecordingFromFolder(FStringView FolderPath = TEXT(""));
 
diff --git a/Engine/Source/Runtime/Experimental/ChaosVDData/Private/ChaosVDRecording.cpp b/Engine/Source/Runtime/Experimental/ChaosVDData/Private/ChaosVDRecording.cpp
--- a/Engine/Source/Runtime/Experimental/ChaosVDData/Private/ChaosVDRecording.cpp
+++ b/Engine/Source/Runtime/Experimental/ChaosVDData/Private/ChaosVDRecording.cpp
@@ -404,6 +404,9 @@ void FChaosVDRecording::CollapseSolverFramesRange_AssumesLocked(int32 SolverID,
 							continue;
 						}
 
+						// The index could have been re-used after the particle was destroyed
+						OutCollapsedFrameData.ParticlesDestroyedIDs.Remove(ParticleData->ParticleIndex);
+
 						if (TSharedPtr<FChaosVDParticleDataWrapper>* FoundParticleData = ParticlesOnCurrentGeneratedKeyframe.Find(ParticleData->ParticleIndex))
 						{
 							(*FoundParticleData) = ParticleData;
diff --git a/Engine/Source/Runtime/Experimental/ChaosVisualDebugger/Public/DataWrappers/ChaosVDParticleDataWrapper.h b/Engine/Source/Runtime/Experimental/ChaosVisualDebugger/Public/DataWrappers/ChaosVDParticleDataWrapper.h
--- a/Engine/Source/Runtime/Experimental/ChaosVisualDebugger/Public/DataWrappers/ChaosVDParticleDataWrapper.h
+++ b/Engine/Source/Runtime/Experimental/ChaosVisualDebugger/Public/DataWrappers/ChaosVDParticleDataWrapper.h
@@ -593,6 +593,9 @@ struct FChaosVDParticleDataWrapper : public FChaosVDWrapperDataBase
 	UPROPERTY(VisibleAnywhere, Category= "General")
 	FString DebugName;
 
+	/** Hash used during CVD Playback to compare names */
+	uint64 DebugNameHash = 0;
+
 	UPROPERTY(VisibleAnywhere, Category= "General")
 	int32 ParticleIndex = INDEX_NONE;
 
