[ChaosVD] Fix for CIS issue with editor tests by ensuring that CVD's streaming system is not initialized/deinitialized in the CDO, otherwise during shutdown might try to load an uobject (a settings object), causing a crash

diff --git a/Engine/Plugins/ChaosVD/Source/ChaosVD/Private/ChaosVDSceneStreaming.cpp b/Engine/Plugins/ChaosVD/Source/ChaosVD/Private/ChaosVDSceneStreaming.cpp
--- a/Engine/Plugins/ChaosVD/Source/ChaosVD/Private/ChaosVDSceneStreaming.cpp
+++ b/Engine/Plugins/ChaosVD/Source/ChaosVD/Private/ChaosVDSceneStreaming.cpp
@@ -14,7 +14,14 @@ FChaosVDSceneStreaming::FChaosVDSceneStreaming()
 {
 	UpdateStreamingQueryShape();
 	StreamingAccelerationStructure.SetTreeToDynamic();
+}
 
+FChaosVDSceneStreaming::~FChaosVDSceneStreaming()
+{
+}
+
+void FChaosVDSceneStreaming::Initialize()
+{
 	if (UChaosVDGeneralSettings* Settings = FChaosVDSettingsManager::Get().GetSettingsObject<UChaosVDGeneralSettings>())
 	{
 		Settings->OnSettingsChanged().AddRaw(this, &FChaosVDSceneStreaming::HandleSettingsChanged);
@@ -23,7 +30,7 @@ FChaosVDSceneStreaming::FChaosVDSceneStreaming()
 	}
 }
 
-FChaosVDSceneStreaming::~FChaosVDSceneStreaming()
+void FChaosVDSceneStreaming::DeInitialize()
 {
 	if (UChaosVDGeneralSettings* Settings = FChaosVDSettingsManager::Get().GetSettingsObject<UChaosVDGeneralSettings>())
 	{
diff --git a/Engine/Plugins/ChaosVD/Source/ChaosVD/Private/ChaosVDSceneStreaming.h b/Engine/Plugins/ChaosVD/Source/ChaosVD/Private/ChaosVDSceneStreaming.h
--- a/Engine/Plugins/ChaosVD/Source/ChaosVD/Private/ChaosVDSceneStreaming.h
+++ b/Engine/Plugins/ChaosVD/Source/ChaosVD/Private/ChaosVDSceneStreaming.h
@@ -60,6 +60,9 @@ public:
 		EType OperationType = EType::None;
 	};
 
+	void Initialize();
+	void DeInitialize();
+
 	virtual bool Tick(float DeltaTime) override;
 
 	/** Adds an object tracking operation to the queue */
diff --git a/Engine/Plugins/ChaosVD/Source/ChaosVD/Private/Components/ChaosVDParticleDataComponent.cpp b/Engine/Plugins/ChaosVD/Source/ChaosVD/Private/Components/ChaosVDParticleDataComponent.cpp
--- a/Engine/Plugins/ChaosVD/Source/ChaosVD/Private/Components/ChaosVDParticleDataComponent.cpp
+++ b/Engine/Plugins/ChaosVD/Source/ChaosVD/Private/Components/ChaosVDParticleDataComponent.cpp
@@ -40,6 +40,11 @@ namespace Chaos::VisualDebugger::Cvars
 
 UChaosVDParticleDataComponent::UChaosVDParticleDataComponent()
 {
+	if (HasAnyFlags(RF_ClassDefaultObject))
+	{
+		return;
+	}
+
 	if (UChaosVDParticleVisualizationSettings* ParticleVisualizationSettings = FChaosVDSettingsManager::Get().GetSettingsObject<UChaosVDParticleVisualizationSettings>())
 	{
 		ParticleVisualizationSettings->OnSettingsChanged().AddUObject(this, &UChaosVDParticleDataComponent::HandleVisibilitySettingsUpdated);
@@ -51,6 +56,8 @@ UChaosVDParticleDataComponent::UChaosVDParticleDataComponent()
 	}
 
 	StreamingSystem.SetStreamingDataSource(this);
+
+	StreamingSystem.Initialize();
 }
 
 void UChaosVDParticleDataComponent::SetScene(const TWeakPtr<FChaosVDScene>& InSceneWeakPtr)
@@ -62,6 +69,11 @@ void UChaosVDParticleDataComponent::SetScene(const TWeakPtr<FChaosVDScene>& InSc
 
 void UChaosVDParticleDataComponent::ClearData()
 {
+	if (HasAnyFlags(RF_ClassDefaultObject))
+	{
+		return;
+	}
+
 	if (UChaosVDParticleVisualizationSettings* ParticleVisualizationSettings = FChaosVDSettingsManager::Get().GetSettingsObject<UChaosVDParticleVisualizationSettings>())
 	{
 		ParticleVisualizationSettings->OnSettingsChanged().RemoveAll(this);
@@ -78,8 +90,8 @@ void UChaosVDParticleDataComponent::ClearData()
 		ScenePtr->ClearSelectionAndNotify();
 	}
 
-	StreamingSystem.Reset();
 	StreamingSystem.SetStreamingDataSource(nullptr);
+	StreamingSystem.DeInitialize();
 	
 	{
 		FReadScopeLock ReadLock(ParticleSceneObjectsLock);
