[ChaosVD] Making it so tracks are automatically re-synced if the timeline sync mode is changed

Also fixed an issue where if you manually scrub the timeline, the active track referenced was not updated. This only affected key shortcuts as they would keep acting on the old track.

diff --git a/Engine/Plugins/ChaosVD/Source/ChaosVD/Private/ChaosVDPlaybackController.cpp b/Engine/Plugins/ChaosVD/Source/ChaosVD/Private/ChaosVDPlaybackController.cpp
--- a/Engine/Plugins/ChaosVD/Source/ChaosVD/Private/ChaosVDPlaybackController.cpp
+++ b/Engine/Plugins/ChaosVD/Source/ChaosVD/Private/ChaosVDPlaybackController.cpp
@@ -1143,11 +1143,85 @@ void FChaosVDPlaybackController::SetScrubbingTimeline(bool bNewIsScrubbingTimeli
 	bIsScrubbingTimeline = bNewIsScrubbingTimeline;
 }
 
+TSharedPtr<const FChaosVDTrackInfo> FChaosVDPlaybackController::GetTrackToReSyncFromOnModeChange(EChaosVDSyncTimelinesMode SyncMode)
+{
+	switch (SyncMode)
+	{
+		case EChaosVDSyncTimelinesMode::RecordedTimestamp:
+			{
+				return ActiveTrack->IsValidTrack() ? ActiveTrack : GetTrackInfo(EChaosVDTrackType::Game, GameTrackID);
+			}
+		case EChaosVDSyncTimelinesMode::NetworkTick:
+			{
+				// Network ticks can only sync from other solver tracks
+
+				if (ActiveTrack->IsValidTrack() && ActiveTrack->TrackType != EChaosVDTrackType::Game)
+				{
+					return ActiveTrack;
+				}
+
+				TSharedPtr<const FChaosVDTrackInfo> GameThreadTrack = GetTrackInfo(EChaosVDTrackType::Game, GameTrackID);
+				if (GameThreadTrack && GameThreadTrack->IsValidTrack())
+				{
+					TArray<TSharedPtr<const FChaosVDTrackInfo>> AvailableSolverTracks;
+					GetAvailableTrackInfosAtTrackFrame(EChaosVDTrackType::Solver, GameThreadTrack.ToSharedRef(), AvailableSolverTracks);
+
+					return AvailableSolverTracks.IsEmpty() ? nullptr : AvailableSolverTracks[0];
+				}
+
+				return nullptr;
+			}
+		default:
+			return nullptr;
+	}
+}
+
+void FChaosVDPlaybackController::SetTimelineSyncMode(EChaosVDSyncTimelinesMode SyncMode)
+{
+	if (CurrentSyncMode != SyncMode)
+	{
+		CurrentSyncMode = SyncMode;
+
+		const TSharedPtr<const FChaosVDTrackInfo> TrackToReSyncFrom = GetTrackToReSyncFromOnModeChange(SyncMode);
+
+		if (TrackToReSyncFrom && TrackToReSyncFrom->IsValidTrack())
+		{
+			const TSharedRef<const FChaosVDTrackInfo> TrackToReSyncFromRef = TrackToReSyncFrom.ToSharedRef();
+
+			TrySetActiveTrack(TrackToReSyncFromRef);
+			SyncTracks(TrackToReSyncFromRef, CurrentSyncMode);
+		}
+		else
+		{
+			UE_LOG(LogChaosVDEditor, Warning, TEXT("[%s] Failed to re-sync tracks after sync mode was changed."), ANSI_TO_TCHAR(__FUNCTION__));
+		}
+	}
+}
+
 void FChaosVDPlaybackController::TrySetActiveTrack(int32 SlotIndex)
 {
 	VisitAvailableTracks([this, SlotIndex](const TSharedPtr<FChaosVDTrackInfo>& TrackInfo)
 	{
 		if (TrackInfo && TrackInfo->TrackSlot == SlotIndex)
+		{
+			if (ActiveTrack->IsValidTrack())
+			{
+				ActiveTrack->bIsPlaying = false;
+			}
+
+			ActiveTrack = TrackInfo.ToSharedRef();
+			return false;
+		}
+
+		return true;
+	});
+}
+
+void FChaosVDPlaybackController::TrySetActiveTrack(const TSharedRef<const FChaosVDTrackInfo>& NewActiveTrack)
+{
+	VisitAvailableTracks([this, NewActiveTrack](const TSharedPtr<FChaosVDTrackInfo>& TrackInfo)
+	{
+		if (TrackInfo && FChaosVDTrackInfo::AreSameTrack(TrackInfo.ToSharedRef(), NewActiveTrack))
 		{
 			ActiveTrack->bIsPlaying = false;
 			ActiveTrack = TrackInfo.ToSharedRef();
diff --git a/Engine/Plugins/ChaosVD/Source/ChaosVD/Private/Widgets/SChaosVDGameFramesPlaybackControls.cpp b/Engine/Plugins/ChaosVD/Source/ChaosVD/Private/Widgets/SChaosVDGameFramesPlaybackControls.cpp
--- a/Engine/Plugins/ChaosVD/Source/ChaosVD/Private/Widgets/SChaosVDGameFramesPlaybackControls.cpp
+++ b/Engine/Plugins/ChaosVD/Source/ChaosVD/Private/Widgets/SChaosVDGameFramesPlaybackControls.cpp
@@ -44,6 +44,7 @@ void SChaosVDGameFramesPlaybackControls::OnFrameSelectionUpdated(int32 NewFrameI
 	if (const TSharedPtr<FChaosVDPlaybackController> PlaybackControllerPtr = PlaybackController.Pin())
 	{
 		constexpr int32 SolverStage = 0;
+		PlaybackControllerPtr->TrySetActiveTrack(GameTrackInfoRef);
 		PlaybackControllerPtr->GoToTrackFrameAndSync(GetInstigatorID(), EChaosVDTrackType::Game, FChaosVDPlaybackController::GameTrackID, NewFrameIndex, SolverStage);
 	}
 }
diff --git a/Engine/Plugins/ChaosVD/Source/ChaosVD/Private/Widgets/SChaosVDSolverPlaybackControls.cpp b/Engine/Plugins/ChaosVD/Source/ChaosVD/Private/Widgets/SChaosVDSolverPlaybackControls.cpp
--- a/Engine/Plugins/ChaosVD/Source/ChaosVD/Private/Widgets/SChaosVDSolverPlaybackControls.cpp
+++ b/Engine/Plugins/ChaosVD/Source/ChaosVD/Private/Widgets/SChaosVDSolverPlaybackControls.cpp
@@ -336,6 +336,7 @@ void SChaosVDSolverPlaybackControls::OnFrameSelectionUpdated(int32 NewFrameIndex
 	{
 		// By default we always playback frames at the last recorded stage as that represents the end of frame state
 		constexpr int32 LastStepNumber = INDEX_NONE;
+		PlaybackControllerPtr->TrySetActiveTrack(SolverTrackInfoRef);
 		PlaybackControllerPtr->GoToTrackFrameAndSync(GetInstigatorID(), EChaosVDTrackType::Solver, SolverTrackInfoRef->TrackID, NewFrameIndex, LastStepNumber);
 	}
 }
diff --git a/Engine/Plugins/ChaosVD/Source/ChaosVD/Public/ChaosVDPlaybackController.h b/Engine/Plugins/ChaosVD/Source/ChaosVD/Public/ChaosVDPlaybackController.h
--- a/Engine/Plugins/ChaosVD/Source/ChaosVD/Public/ChaosVDPlaybackController.h
+++ b/Engine/Plugins/ChaosVD/Source/ChaosVD/Public/ChaosVDPlaybackController.h
@@ -389,11 +389,14 @@ public:
 		return CurrentSyncMode;
 	}
 
+	/**
+	 * Finds and returns a track instance that can be used as sync point for the new timeline mode change
+	 * @param SyncMode New timeline Sync Mode
+	 * */
+	TSharedPtr<const FChaosVDTrackInfo> GetTrackToReSyncFromOnModeChange(EChaosVDSyncTimelinesMode SyncMode);
+
 	/** Sets a new sync mode to be used between tracks */
-	void SetTimelineSyncMode(EChaosVDSyncTimelinesMode SyncMode)
-	{
-		CurrentSyncMode = SyncMode;
-	}
+	void SetTimelineSyncMode(EChaosVDSyncTimelinesMode SyncMode);
 
 	/** Attempts to set the track on the specified slot as active.
 	 * An active track is the track to all key shortcuts will be applied
@@ -401,6 +404,12 @@ public:
 	 */
 	void TrySetActiveTrack(int32 SlotIndex);
 
+	/** Attempts to set the provided track as active.
+	 * An active track is the track to all key shortcuts will be applied
+	 * @param NewActiveTrack Slot of the track we want to make active
+	 */
+	void TrySetActiveTrack(const TSharedRef<const FChaosVDTrackInfo>& NewActiveTrack);
+
 protected:
 	/** Updates (or adds) solvers data from the loaded recording to the solver tracks */
 	void UpdateSolverTracksData();
