[ChaosVD] Added a new button that will open a new (selection independent) details panel for a particle

This is useful for when you want to compare the values of the state between two or more particles.

There is a limit to 4 independent details panel at any given time because we need to declare them in the layout, but it can be expanded to more if needed in the future.

Currently the new details panel button is limited to particle data as it is added on each customization, but we can add it to scene queries or other data in the future.

The changes made in this CL are

- Created a new base class for the the Object Details tab class of CVD, that does not have any selection events related code. This base class is used for the independent details panel, and the old class now inherits from the new one and adds the required selection events handling code
- Created a new manager class for the new independent details panel, that any code with access to the Main Tab (toolkit host) can access and set a structure to inspect
- Created a new area in the layout for the new details panel slots
- Added the button to open a new details panel to the Particle Data customization class
- Ensure that if a particle is unloaded, we clear up the ptr to it (before has was handled by the selection system as we just cleared out the selection, but now we have panels not tied to that system)
- Removed dependencies on the selection system from the Particle Data Customization object.

diff --git a/Engine/Plugins/ChaosVD/Source/ChaosVD/Private/ChaosVDIndependentDetailsPanelManager.cpp b/Engine/Plugins/ChaosVD/Source/ChaosVD/Private/ChaosVDIndependentDetailsPanelManager.cpp
new file mode 100644
--- /dev/null
+++ b/Engine/Plugins/ChaosVD/Source/ChaosVD/Private/ChaosVDIndependentDetailsPanelManager.cpp
@@ -0,0 +1,61 @@
+// Copyright Epic Games, Inc. All Rights Reserved.
+
+#include "ChaosVDIndependentDetailsPanelManager.h"
+
+#include "ChaosVDObjectDetailsTab.h"
+#include "Framework/Docking/TabManager.h"
+#include "Widgets/SChaosVDMainTab.h"
+
+FChaosVDIndependentDetailsPanelManager::FChaosVDIndependentDetailsPanelManager(const TSharedRef<SChaosVDMainTab>& InMainTab) : MainTab(InMainTab)
+{
+	AvailableDetailsPanelIDs.Add(FChaosVDTabID::IndependentDetailsPanel1);
+	AvailableDetailsPanelIDs.Add(FChaosVDTabID::IndependentDetailsPanel2);
+	AvailableDetailsPanelIDs.Add(FChaosVDTabID::IndependentDetailsPanel3);
+	AvailableDetailsPanelIDs.Add(FChaosVDTabID::IndependentDetailsPanel4);
+
+	InMainTab->RegisterTabSpawner<FChaosVDStandAloneObjectDetailsTab>(FChaosVDTabID::IndependentDetailsPanel1);
+	InMainTab->RegisterTabSpawner<FChaosVDStandAloneObjectDetailsTab>(FChaosVDTabID::IndependentDetailsPanel2);
+	InMainTab->RegisterTabSpawner<FChaosVDStandAloneObjectDetailsTab>(FChaosVDTabID::IndependentDetailsPanel3);
+	InMainTab->RegisterTabSpawner<FChaosVDStandAloneObjectDetailsTab>(FChaosVDTabID::IndependentDetailsPanel4);
+}
+
+TSharedPtr<FChaosVDStandAloneObjectDetailsTab> FChaosVDIndependentDetailsPanelManager::GetAvailableStandAloneDetailsPanelTab()
+{
+	if (AvailableDetailsPanelIDs.IsEmpty())
+	{
+		return nullptr;
+	}
+
+	TSharedPtr<SChaosVDMainTab> OwningTabPtr = MainTab.Pin();
+	if (!OwningTabPtr.IsValid())
+	{
+		return nullptr;
+	}
+
+	if (const TSharedPtr<FTabManager> TabManager = OwningTabPtr->GetTabManager())
+	{
+		FName AvailableTabID = AvailableDetailsPanelIDs.Pop();
+		if (const TSharedPtr<FChaosVDStandAloneObjectDetailsTab> DetailsTab = OwningTabPtr->GetTabSpawnerInstance<FChaosVDStandAloneObjectDetailsTab>(AvailableTabID).Pin())
+		{
+			DetailsTab->OnTabDestroyed().AddSP(this, &FChaosVDIndependentDetailsPanelManager::HandleTabDestroyed, AvailableTabID);
+			TabManager->TryInvokeTab(AvailableTabID);
+
+			return DetailsTab;
+		}
+	}
+
+	return nullptr;
+}
+
+void FChaosVDIndependentDetailsPanelManager::HandleTabDestroyed(TSharedRef<SDockTab> Tab, FName TabID)
+{
+	AvailableDetailsPanelIDs.Add(TabID);
+	
+	if (TSharedPtr<SChaosVDMainTab> OwningTabPtr = MainTab.Pin())
+	{
+		if (const TSharedPtr<FChaosVDStandAloneObjectDetailsTab> DetailsTab = OwningTabPtr->GetTabSpawnerInstance<FChaosVDStandAloneObjectDetailsTab>(TabID).Pin())
+		{
+			DetailsTab->OnTabDestroyed().RemoveAll(this);
+		}
+	}
+}
diff --git a/Engine/Plugins/ChaosVD/Source/ChaosVD/Private/ChaosVDIndependentDetailsPanelManager.h b/Engine/Plugins/ChaosVD/Source/ChaosVD/Private/ChaosVDIndependentDetailsPanelManager.h
new file mode 100644
--- /dev/null
+++ b/Engine/Plugins/ChaosVD/Source/ChaosVD/Private/ChaosVDIndependentDetailsPanelManager.h
@@ -0,0 +1,29 @@
+// Copyright Epic Games, Inc. All Rights Reserved.
+
+#pragma once
+#include "ChaosVDTabsIDs.h"
+#include "Templates/SharedPointer.h"
+
+class SDockTab;
+class FChaosVDStandAloneObjectDetailsTab;
+class SChaosVDMainTab;
+
+/**
+ * Manager class that handles any selection independent details panel tab.
+ */
+class FChaosVDIndependentDetailsPanelManager : public TSharedFromThis<FChaosVDIndependentDetailsPanelManager>
+{
+public:
+	FChaosVDIndependentDetailsPanelManager(const TSharedRef<SChaosVDMainTab>& InMainTab);
+
+	/** Returns a shared ptr to a new selection independent details panel tab or null if no panels are available */
+	TSharedPtr<FChaosVDStandAloneObjectDetailsTab> GetAvailableStandAloneDetailsPanelTab();
+
+private:
+
+	void HandleTabDestroyed(TSharedRef<SDockTab> Tab, FName TabID);
+
+	TArray<FName> AvailableDetailsPanelIDs;
+	
+	TWeakPtr<SChaosVDMainTab> MainTab;
+};
diff --git a/Engine/Plugins/ChaosVD/Source/ChaosVD/Private/ChaosVDObjectDetailsTab.cpp b/Engine/Plugins/ChaosVD/Source/ChaosVD/Private/ChaosVDObjectDetailsTab.cpp
--- a/Engine/Plugins/ChaosVD/Source/ChaosVD/Private/ChaosVDObjectDetailsTab.cpp
+++ b/Engine/Plugins/ChaosVD/Source/ChaosVD/Private/ChaosVDObjectDetailsTab.cpp
@@ -23,12 +23,17 @@
 
 #define LOCTEXT_NAMESPACE "ChaosVisualDebugger"
 
-void FChaosVDObjectDetailsTab::AddUnsupportedStruct(const UStruct* Struct)
+void FChaosVDStandAloneObjectDetailsTab::AddUnsupportedStruct(const UStruct* Struct)
 {
 	UnsupportedStructs.Add(Struct);
 }
 
-TSharedRef<SDockTab> FChaosVDObjectDetailsTab::HandleTabSpawnRequest(const FSpawnTabArgs& Args)
+bool FChaosVDStandAloneObjectDetailsTab::IsSupportedStruct(const TWeakObjectPtr<const UStruct>& InWeakStructPtr)
+{
+	return !UnsupportedStructs.Contains(InWeakStructPtr);
+}
+
+TSharedRef<SDockTab> FChaosVDStandAloneObjectDetailsTab::HandleTabSpawnRequest(const FSpawnTabArgs& Args)
 {
 	TSharedRef<SDockTab> DetailsPanelTab =
 	SNew(SDockTab)
@@ -41,32 +46,12 @@ TSharedRef<SDockTab> FChaosVDObjectDetailsTab::HandleTabSpawnRequest(const FSpaw
 	AddUnsupportedStruct(FChaosVDQueryDataWrapper::StaticStruct());
 	AddUnsupportedStruct(FChaosVDParticlePairMidPhase::StaticStruct());
 
-	if (const TSharedPtr<FChaosVDScene> ScenePtr = GetChaosVDScene().Pin())
+	if (const TSharedPtr<SChaosVDMainTab> MainTabPtr = OwningTabWidget.Pin())
 	{
-		RegisterSelectionSetObject(ScenePtr->GetElementSelectionSet());
-
-		if (TSharedPtr<FChaosVDSolverDataSelection> SolverDataSelectionObject = ScenePtr->GetSolverDataSelectionObject().Pin())
-		{
-			SolverDataSelectionObject->GetDataSelectionChangedDelegate().AddSP(this, &FChaosVDObjectDetailsTab::HandleSolverDataSelectionChange);
-		}
-
-		if (const TSharedPtr<SChaosVDMainTab> MainTabPtr = OwningTabWidget.Pin())
-		{
-			DetailsPanelTab->SetContent
-			(
-				SAssignNew(DetailsPanelView, SChaosVDDetailsView, MainTabPtr.ToSharedRef())
-			);
-
-			// If we closed the tab and opened it again with an object already selected, try to restore the selected object view
-			if (DetailsPanelView.IsValid() && CurrentSelectedObject.IsValid())
-			{
-				DetailsPanelView->SetSelectedObject(CurrentSelectedObject.Get());
-			}
-		}
-		else
-		{
-			DetailsPanelTab->SetContent(GenerateErrorWidget());
-		}
+		DetailsPanelTab->SetContent
+		(
+			SAssignNew(DetailsPanelView, SChaosVDDetailsView, MainTabPtr.ToSharedRef())
+		);
 	}
 	else
 	{
@@ -80,17 +65,9 @@ TSharedRef<SDockTab> FChaosVDObjectDetailsTab::HandleTabSpawnRequest(const FSpaw
 	return DetailsPanelTab;
 }
 
-void FChaosVDObjectDetailsTab::HandleTabClosed(TSharedRef<SDockTab> InTabClosed)
+void FChaosVDStandAloneObjectDetailsTab::HandleTabClosed(TSharedRef<SDockTab> InTabClosed)
 {
 	FChaosVDTabSpawnerBase::HandleTabClosed(InTabClosed);
-
-	if (const TSharedPtr<FChaosVDScene> ScenePtr = GetChaosVDScene().Pin())
-	{
-		if (TSharedPtr<FChaosVDSolverDataSelection> SolverDataSelectionObject = ScenePtr->GetSolverDataSelectionObject().Pin())
-		{
-			SolverDataSelectionObject->GetDataSelectionChangedDelegate().RemoveAll(this);
-		}
-	}
 	
 	DetailsPanelView.Reset();
 }
@@ -115,6 +92,42 @@ void FChaosVDObjectDetailsTab::HandleActorsSelection(TArrayView<AActor*> Selecte
 	}
 }
 
+TSharedRef<SDockTab> FChaosVDObjectDetailsTab::HandleTabSpawnRequest(const FSpawnTabArgs& Args)
+{
+	TSharedRef<SDockTab> NewTab = FChaosVDStandAloneObjectDetailsTab::HandleTabSpawnRequest(Args);
+	
+	if (const TSharedPtr<FChaosVDScene> ScenePtr = GetChaosVDScene().Pin())
+	{
+		RegisterSelectionSetObject(ScenePtr->GetElementSelectionSet());
+
+		if (TSharedPtr<FChaosVDSolverDataSelection> SolverDataSelectionObject = ScenePtr->GetSolverDataSelectionObject().Pin())
+		{
+			SolverDataSelectionObject->GetDataSelectionChangedDelegate().AddSP(this, &FChaosVDObjectDetailsTab::HandleSolverDataSelectionChange);
+		}
+	}
+
+	// If we closed the tab and opened it again with an object already selected, try to restore the selected object view
+	if (DetailsPanelView.IsValid() && CurrentSelectedObject.IsValid())
+	{
+		DetailsPanelView->SetSelectedObject(CurrentSelectedObject.Get());
+	}
+	
+	return NewTab;
+}
+
+void FChaosVDObjectDetailsTab::HandleTabClosed(TSharedRef<SDockTab> InTabClosed)
+{
+	if (const TSharedPtr<FChaosVDScene> ScenePtr = GetChaosVDScene().Pin())
+	{
+		if (TSharedPtr<FChaosVDSolverDataSelection> SolverDataSelectionObject = ScenePtr->GetSolverDataSelectionObject().Pin())
+		{
+			SolverDataSelectionObject->GetDataSelectionChangedDelegate().RemoveAll(this);
+		}
+	}
+
+	FChaosVDStandAloneObjectDetailsTab::HandleTabClosed(InTabClosed);
+}
+
 void FChaosVDObjectDetailsTab::HandlePostSelectionChange(const UTypedElementSelectionSet* ChangedSelectionSet)
 {
 	TArray<AActor*> SelectedActors = ChangedSelectionSet->GetSelectedObjects<AActor>();
@@ -157,9 +170,4 @@ void FChaosVDObjectDetailsTab::HandleSolverDataSelectionChange(const TSharedPtr<
 	DetailsPanelView->SetSelectedStruct(SelectionHandle->GetCustomDataReadOnlyStructViewForDetails());
 }
 
-bool FChaosVDObjectDetailsTab::IsSupportedStruct(const TWeakObjectPtr<const UStruct>& InWeakStructPtr)
-{
-	return !UnsupportedStructs.Contains(InWeakStructPtr);
-}
-
 #undef LOCTEXT_NAMESPACE
diff --git a/Engine/Plugins/ChaosVD/Source/ChaosVD/Private/ChaosVDObjectDetailsTab.h b/Engine/Plugins/ChaosVD/Source/ChaosVD/Private/ChaosVDObjectDetailsTab.h
--- a/Engine/Plugins/ChaosVD/Source/ChaosVD/Private/ChaosVDObjectDetailsTab.h
+++ b/Engine/Plugins/ChaosVD/Source/ChaosVD/Private/ChaosVDObjectDetailsTab.h
@@ -6,6 +6,7 @@
 #include "Templates/SharedPointer.h"
 #include "ChaosVDTabSpawnerBase.h"
 #include "Delegates/IDelegateInstance.h"
+#include "Widgets/SChaosVDDetailsView.h"
 
 struct FChaosVDSolverDataSelectionHandle;
 struct FChaosVDParticleDebugData;
@@ -17,18 +18,52 @@ class FTabManager;
 class SChaosVDDetailsView;
 class SDockTab;
 
-/** Spawns and handles and instance for the visual debugger details panel */
-class FChaosVDObjectDetailsTab : public FChaosVDTabSpawnerBase, public FChaosVDSceneSelectionObserver, public TSharedFromThis<FChaosVDObjectDetailsTab>
+/** Spawns and handles an instance for a selection independent details panel */
+class FChaosVDStandAloneObjectDetailsTab : public FChaosVDTabSpawnerBase, public TSharedFromThis<FChaosVDStandAloneObjectDetailsTab>
 {
 public:
 
-	FChaosVDObjectDetailsTab(const FName& InTabID, TSharedPtr<FTabManager> InTabManager, TWeakPtr<SChaosVDMainTab> InOwningTabWidget) : FChaosVDTabSpawnerBase(InTabID, InTabManager, InOwningTabWidget)
+	FChaosVDStandAloneObjectDetailsTab(const FName& InTabID, TSharedPtr<FTabManager> InTabManager, TWeakPtr<SChaosVDMainTab> InOwningTabWidget) : FChaosVDTabSpawnerBase(InTabID, InTabManager, InOwningTabWidget)
 	{
 	}
 	
-	TSharedPtr<SChaosVDDetailsView> GetDetailsPanel() { return DetailsPanelView; }
+	TSharedPtr<SChaosVDDetailsView> GetDetailsPanel()
+	{
+		return DetailsPanelView;
+	}
 
 	void AddUnsupportedStruct(const UStruct* Struct);
+	
+	/** Updates the current object this view details is viewing */
+	template<typename TStruct>
+	void SetStructToInspect(TStruct* NewStruct);
+
+protected:
+
+	virtual TSharedRef<SDockTab> HandleTabSpawnRequest(const FSpawnTabArgs& Args) override;
+	virtual void HandleTabClosed(TSharedRef<SDockTab> InTabClosed) override;
+
+	bool IsSupportedStruct(const TWeakObjectPtr<const UStruct>& InWeakStructPtr);
+
+	TSharedPtr<SChaosVDDetailsView> DetailsPanelView;
+
+	TSet<TWeakObjectPtr<const UStruct>> UnsupportedStructs;
+};
+
+template <typename TStruct>
+void FChaosVDStandAloneObjectDetailsTab::SetStructToInspect(TStruct* NewStruct)
+{
+	DetailsPanelView->SetSelectedStruct(NewStruct);
+}
+
+/** Spawns and handles an instance for the visual debugger details panel */
+class FChaosVDObjectDetailsTab : public FChaosVDStandAloneObjectDetailsTab, public FChaosVDSceneSelectionObserver
+{
+public:
+
+	FChaosVDObjectDetailsTab(const FName& InTabID, TSharedPtr<FTabManager> InTabManager, TWeakPtr<SChaosVDMainTab> InOwningTabWidget) : FChaosVDStandAloneObjectDetailsTab(InTabID, InTabManager, InOwningTabWidget)
+	{
+	}
 
 protected:
 
@@ -41,11 +76,6 @@ protected:
 
 	void HandleSolverDataSelectionChange(const TSharedPtr<FChaosVDSolverDataSelectionHandle>& SelectionHandle);
 
-	bool IsSupportedStruct(const TWeakObjectPtr<const UStruct>& InWeakStructPtr);
-
 	FDelegateHandle SelectionDelegateHandle;
-	TSharedPtr<SChaosVDDetailsView> DetailsPanelView;
 	TWeakObjectPtr<UObject> CurrentSelectedObject = nullptr;
-
-	TSet<TWeakObjectPtr<const UStruct>> UnsupportedStructs;
-};
+};
\ No newline at end of file
diff --git a/Engine/Plugins/ChaosVD/Source/ChaosVD/Private/ChaosVDSceneParticle.cpp b/Engine/Plugins/ChaosVD/Source/ChaosVD/Private/ChaosVDSceneParticle.cpp
--- a/Engine/Plugins/ChaosVD/Source/ChaosVD/Private/ChaosVDSceneParticle.cpp
+++ b/Engine/Plugins/ChaosVD/Source/ChaosVD/Private/ChaosVDSceneParticle.cpp
@@ -47,6 +47,8 @@ FChaosVDSceneParticle::~FChaosVDSceneParticle()
 	MeshDataHandles.Empty();
 
 	SetIsActive(false);
+
+	ParticleDestroyedDelegate.ExecuteIfBound();
 }
 
 FChaosVDSceneParticle::FChaosVDSceneParticle()
diff --git a/Engine/Plugins/ChaosVD/Source/ChaosVD/Private/ChaosVDSceneParticle.h b/Engine/Plugins/ChaosVD/Source/ChaosVD/Private/ChaosVDSceneParticle.h
--- a/Engine/Plugins/ChaosVD/Source/ChaosVD/Private/ChaosVDSceneParticle.h
+++ b/Engine/Plugins/ChaosVD/Source/ChaosVD/Private/ChaosVDSceneParticle.h
@@ -193,6 +193,11 @@ protected:
 
 	mutable FBox CachedBounds = FBox(ForceInitToZero);
 
+	/** Called when this particle is destroyed. This is intentionally private and only intended to be used by
+	 * the customization class for this object
+	 */
+	FSimpleDelegate ParticleDestroyedDelegate;
+
 	friend class FChaosVDSceneParticleCustomization;
 	friend class UChaosVDParticleDataComponent;
 };
diff --git a/Engine/Plugins/ChaosVD/Source/ChaosVD/Private/ChaosVDSceneParticleCustomization.cpp b/Engine/Plugins/ChaosVD/Source/ChaosVD/Private/ChaosVDSceneParticleCustomization.cpp
--- a/Engine/Plugins/ChaosVD/Source/ChaosVD/Private/ChaosVDSceneParticleCustomization.cpp
+++ b/Engine/Plugins/ChaosVD/Source/ChaosVD/Private/ChaosVDSceneParticleCustomization.cpp
@@ -4,17 +4,18 @@
 
 #include "ChaosVDCollisionDataDetailsTab.h"
 #include "ChaosVDEngine.h"
+#include "ChaosVDIndependentDetailsPanelManager.h"
 #include "ChaosVDModule.h"
+#include "ChaosVDObjectDetailsTab.h"
 #include "ChaosVDSceneParticle.h"
 #include "ChaosVDScene.h"
 #include "ChaosVDTabsIDs.h"
 #include "DetailCategoryBuilder.h"
 #include "DetailLayoutBuilder.h"
 #include "DetailWidgetRow.h"
-#include "IStructureDetailsView.h"
 #include "DataWrappers/ChaosVDParticleDataWrapper.h"
 #include "DetailsCustomizations/ChaosVDDetailsCustomizationUtils.h"
-#include "TEDS/ChaosVDSelectionInterface.h"
+#include "Misc/MessageDialog.h"
 #include "Widgets/SChaosVDCollisionDataInspector.h"
 #include "Widgets/SChaosVDMainTab.h"
 #include "Widgets/SToolTip.h"
@@ -90,7 +91,7 @@ void FChaosVDSceneParticleCustomization::CustomizeDetails(IDetailLayoutBuilder&
 		
 		if (SelectedParticleInstance)
 		{
-			CurrentObservedParticle = SelectedParticleInstance;
+			UpdateObserverParticlePtr(SelectedParticleInstance);
 
 			HandleSceneUpdated();
 
@@ -111,6 +112,14 @@ void FChaosVDSceneParticleCustomization::CustomizeDetails(IDetailLayoutBuilder&
 				FChaosVDDetailsCustomizationUtils::HideInvalidCVDDataWrapperProperties({&InspectedDataPropertyHandleRef, 1}, DetailBuilder);
 			}
 
+			IDetailCategoryBuilder& ParticleDataCategoryBuilder = DetailBuilder.EditCategory(ParticleDataCategoryName);
+
+			ParticleDataCategoryBuilder.AddCustomRow(FText::GetEmpty()).
+			WholeRowContent()
+			[
+				GenerateOpenInNewDetailsPanelButton().ToSharedRef()
+			];
+
 			const FText CollisionDataRowLabel = LOCTEXT("ParticleCollisionData", "Particle Collision Data");
 			IDetailCategoryBuilder& CategoryBuilder = DetailBuilder.EditCategory("CollisionData", CollisionDataRowLabel);
 			
@@ -187,8 +196,42 @@ FReply FChaosVDSceneParticleCustomization::ShowCollisionDataForSelectedObject()
 	return FReply::Handled();
 }
 
+FReply FChaosVDSceneParticleCustomization::OpenNewDetailsPanel()
+{
+	if (!CurrentObservedParticle)
+	{
+		return FReply::Handled();
+	}
+	
+	TSharedPtr<SChaosVDMainTab> OwningTabPtr = MainTabWeakPtr.Pin();
+	if (!OwningTabPtr.IsValid())
+	{
+		return FReply::Handled();
+	}
+
+	if (TSharedPtr<FChaosVDIndependentDetailsPanelManager> IndependentDetailsPanelManager = OwningTabPtr->GetIndependentDetailsPanelManager())
+	{
+		if (TSharedPtr<FChaosVDStandAloneObjectDetailsTab> DetailsTab = IndependentDetailsPanelManager->GetAvailableStandAloneDetailsPanelTab())
+		{
+			DetailsTab->SetStructToInspect(CurrentObservedParticle);
+			return FReply::Handled();
+		}
+		else
+		{
+			FMessageDialog::Open(EAppMsgType::Ok, LOCTEXT("OpenDetailsPanelError", "No (selection independent) Details Panel slot available.\n\nPlease close a panel and try again."));
+		}
+	}
+
+	return FReply::Handled();
+}
+
 void FChaosVDSceneParticleCustomization::ResetCachedView()
 {
+	if (CurrentObservedParticle)
+	{
+		CurrentObservedParticle->ParticleDestroyedDelegate.Unbind();
+	}
+
 	CurrentObservedParticle = nullptr;
 	CachedParticleData = FChaosVDParticleDataWrapper();
 	CachedGeometryDataInstanceCopy = FChaosVDMeshDataInstanceState();
@@ -207,25 +250,12 @@ void FChaosVDSceneParticleCustomization::RegisterCVDScene(const TSharedPtr<FChao
 		if (InScene)
 		{
 			InScene->OnSceneUpdated().AddSP(this, &FChaosVDSceneParticleCustomization::HandleSceneUpdated);
-			RegisterSelectionSetObject(InScene->GetElementSelectionSet());
 		}
 
 		SceneWeakPtr = InScene;
 	}
 }
 
-void FChaosVDSceneParticleCustomization::HandlePostSelectionChange(const UTypedElementSelectionSet* ChangesSelectionSet)
-{
-	constexpr int32 MaxElements = 1;
-	TArray<FTypedElementHandle, TInlineAllocator<MaxElements>> SelectedParticlesHandles;
-	GetSelectionSetObject()->GetSelectedElementHandles(SelectedParticlesHandles, UChaosVDSelectionInterface::StaticClass());
-
-	if (SelectedParticlesHandles.Num() == 0)
-	{
-		CurrentObservedParticle = nullptr;
-	}
-}
-
 TSharedPtr<SWidget> FChaosVDSceneParticleCustomization::GenerateShowCollisionDataButton()
 {
 	TSharedPtr<SWidget> ShowCollisionButton = SNew(SHorizontalBox)
@@ -250,7 +280,7 @@ TSharedPtr<SWidget> FChaosVDSceneParticleCustomization::GenerateShowCollisionDat
 			[
 				SNew(STextBlock)
 				.TextStyle(FAppStyle::Get(), "SmallButtonText")
-				.Text(LOCTEXT("ShowCollisionDataOnInspector","Show Collison Data in Inspector"))
+				.Text(LOCTEXT("ShowCollisionDataOnInspector","Show Collision Data in Inspector"))
 			]
 		]
 	];
@@ -258,4 +288,59 @@ TSharedPtr<SWidget> FChaosVDSceneParticleCustomization::GenerateShowCollisionDat
 	return ShowCollisionButton;
 }
 
+TSharedPtr<SWidget> FChaosVDSceneParticleCustomization::GenerateOpenInNewDetailsPanelButton()
+{
+	TSharedPtr<SWidget> OpenButton = SNew(SHorizontalBox)
+	+SHorizontalBox::Slot()
+	.VAlign(VAlign_Center)
+	.HAlign(HAlign_Center)
+	.Padding(12.0f, 7.0f, 12.0f, 7.0f)
+	.FillWidth(1.0f)
+	[
+		SNew(SButton)
+		.ToolTip(SNew(SToolTip).Text(LOCTEXT("OpenDetailsPanelDesc", "Click here to open a new (selection independent) details panel for this particle.")))
+		.ContentPadding(FMargin(0, 5.f, 0, 4.f))
+		.OnClicked_Raw(this, &FChaosVDSceneParticleCustomization::OpenNewDetailsPanel)
+		.Content()
+		[
+			SNew(SHorizontalBox)
+			+SHorizontalBox::Slot()
+			.Padding(FMargin(3, 0, 0, 0))
+			.VAlign(VAlign_Center)
+			.AutoWidth()
+			[
+				SNew(STextBlock)
+				.TextStyle(FAppStyle::Get(), "SmallButtonText")
+				.Text(LOCTEXT("OpenDetailsPanelText","Show Data in New Panel"))
+			]
+		]
+	];
+
+	return OpenButton;
+}
+
+void FChaosVDSceneParticleCustomization::UpdateObserverParticlePtr(FChaosVDSceneParticle* NewObservedParticle)
+{
+	if (CurrentObservedParticle)
+	{
+		CurrentObservedParticle->ParticleDestroyedDelegate.Unbind();
+	}
+
+	if (NewObservedParticle)
+	{
+		NewObservedParticle->ParticleDestroyedDelegate.BindSP(this, &FChaosVDSceneParticleCustomization::HandleObservedParticleInstanceDestroyed);
+		CurrentObservedParticle = NewObservedParticle;	
+	}
+	else
+	{
+		ResetCachedView();
+	}
+}
+
+
+void FChaosVDSceneParticleCustomization::HandleObservedParticleInstanceDestroyed()
+{
+	ResetCachedView();
+}
+
 #undef LOCTEXT_NAMESPACE
diff --git a/Engine/Plugins/ChaosVD/Source/ChaosVD/Private/ChaosVDSceneParticleCustomization.h b/Engine/Plugins/ChaosVD/Source/ChaosVD/Private/ChaosVDSceneParticleCustomization.h
--- a/Engine/Plugins/ChaosVD/Source/ChaosVD/Private/ChaosVDSceneParticleCustomization.h
+++ b/Engine/Plugins/ChaosVD/Source/ChaosVD/Private/ChaosVDSceneParticleCustomization.h
@@ -3,7 +3,6 @@
 #pragma once
 
 #include "ChaosVDGeometryDataComponent.h"
-#include "ChaosVDSceneSelectionObserver.h"
 #include "DetailCategoryBuilder.h"
 #include "DetailLayoutBuilder.h"
 #include "Templates/SharedPointer.h"
@@ -16,7 +15,7 @@ class SChaosVDMainTab;
 struct FChaosVDParticleDataWrapper;
 
 /** Custom details panel for the ChaosVD Particle Actor */
-class FChaosVDSceneParticleCustomization : public IDetailCustomization, public FChaosVDSceneSelectionObserver
+class FChaosVDSceneParticleCustomization : public IDetailCustomization
 {
 public:
 	FChaosVDSceneParticleCustomization(const TWeakPtr<SChaosVDMainTab>& InMainTab);
@@ -42,15 +41,19 @@ private:
 
 	void RegisterCVDScene(const TSharedPtr<FChaosVDScene>& InScene);
 
-	virtual void HandlePostSelectionChange(const UTypedElementSelectionSet* ChangesSelectionSet) override;
-
 	void HandleSceneUpdated();
 
 	bool GetCollisionDataButtonEnabled() const;
 
 	FReply ShowCollisionDataForSelectedObject();
+	FReply OpenNewDetailsPanel();
 
 	TSharedPtr<SWidget> GenerateShowCollisionDataButton();
+	TSharedPtr<SWidget> GenerateOpenInNewDetailsPanelButton();
+
+	void UpdateObserverParticlePtr(FChaosVDSceneParticle* NewObservedParticle);
+
+	void HandleObservedParticleInstanceDestroyed();
 
 	/* Copy of the last known geometry shape data structure of a selected particle and mesh instance -  Used to avoid rebuild the layout every time we change frame in CVD */
 	FChaosVDParticleDataWrapper CachedParticleData;
diff --git a/Engine/Plugins/ChaosVD/Source/ChaosVD/Private/ChaosVDTabsIDs.cpp b/Engine/Plugins/ChaosVD/Source/ChaosVD/Private/ChaosVDTabsIDs.cpp
--- a/Engine/Plugins/ChaosVD/Source/ChaosVD/Private/ChaosVDTabsIDs.cpp
+++ b/Engine/Plugins/ChaosVD/Source/ChaosVD/Private/ChaosVDTabsIDs.cpp
@@ -6,6 +6,10 @@ const FName FChaosVDTabID::ChaosVisualDebuggerTab = TEXT("ChaosVisualDebuggerTab
 const FName FChaosVDTabID::PlaybackViewport = TEXT("PlaybackViewport");
 const FName FChaosVDTabID::WorldOutliner = TEXT("WorldOutliner");
 const FName FChaosVDTabID::DetailsPanel = TEXT("DetailsPanel");
+const FName FChaosVDTabID::IndependentDetailsPanel1 = TEXT("IndependentDetailsPanel1");
+const FName FChaosVDTabID::IndependentDetailsPanel2 = TEXT("IndependentDetailsPanel2");
+const FName FChaosVDTabID::IndependentDetailsPanel3 = TEXT("IndependentDetailsPanel3");
+const FName FChaosVDTabID::IndependentDetailsPanel4 = TEXT("IndependentDetailsPanel4");
 const FName FChaosVDTabID::OutputLog = TEXT("ChaosVDOutputLog");
 const FName FChaosVDTabID::SolversTrack = TEXT("SolversTrack");
 const FName FChaosVDTabID::StatusBar = TEXT("StatusBar");
diff --git a/Engine/Plugins/ChaosVD/Source/ChaosVD/Private/Widgets/SChaosVDMainTab.cpp b/Engine/Plugins/ChaosVD/Source/ChaosVD/Private/Widgets/SChaosVDMainTab.cpp
--- a/Engine/Plugins/ChaosVD/Source/ChaosVD/Private/Widgets/SChaosVDMainTab.cpp
+++ b/Engine/Plugins/ChaosVD/Source/ChaosVD/Private/Widgets/SChaosVDMainTab.cpp
@@ -13,6 +13,7 @@
 #include "ChaosVDSolversTracksTab.h"
 #include "ChaosVDCollisionDataDetailsTab.h"
 #include "ChaosVDConstraintDataInspectorTab.h"
+#include "ChaosVDIndependentDetailsPanelManager.h"
 #include "ChaosVDSceneParticleCustomization.h"
 #include "ChaosVDRecordedLogTab.h"
 #include "ChaosVDSceneQueryDataInspectorTab.h"
@@ -103,6 +104,8 @@ void SChaosVDMainTab::Construct(const FArguments& InArgs, TSharedPtr<FChaosVDEng
 	RegisterTabSpawner<FChaosVDSceneQueryBrowserTab>(FChaosVDTabID::SceneQueryBrowser);
 	RegisterTabSpawner<FChaosVDRecordedLogTab>(FChaosVDTabID::RecordedOutputLog);
 
+	IndependentDetailsPanelManager = MakeShared<FChaosVDIndependentDetailsPanelManager>(StaticCastSharedRef<SChaosVDMainTab>(AsShared()));
+
 	FChaosVDExtensionsManager::Get().EnumerateExtensions([this](const TSharedRef<FChaosVDExtension>& Extension)
 	{
 		Extension->RegisterCustomTabSpawners(StaticCastSharedRef<SChaosVDMainTab>(AsShared()));
@@ -267,6 +270,11 @@ void SChaosVDMainTab::RegisterComponentVisualizer(FName ClassName, const TShared
 	}
 }
 
+const TSharedPtr<FChaosVDIndependentDetailsPanelManager>& SChaosVDMainTab::GetIndependentDetailsPanelManager()
+{
+	return IndependentDetailsPanelManager;
+}
+
 void SChaosVDMainTab::HandlePostInitializationExtensionRegistered(const TSharedRef<FChaosVDExtension>& NewExtension)
 {	
 	NewExtension->RegisterCustomTabSpawners(StaticCastSharedRef<SChaosVDMainTab>(AsShared()));
@@ -355,6 +363,18 @@ TSharedRef<FTabManager::FLayout> SChaosVDMainTab::GenerateMainLayout()
 		FTabManager::NewStack()
 		->SetSizeCoefficient(1.0f)
 		->AddTab(FChaosVDTabID::SceneQueryBrowser, ETabState::ClosedTab)
+	))
+	->AddArea(
+	FTabManager::NewArea(800.0f * DPIScaleFactor, 600.0f * DPIScaleFactor)
+	->SetOrientation(Orient_Vertical)
+	->Split
+	(
+		FTabManager::NewStack()
+		->SetSizeCoefficient(1.0f)
+		->AddTab(FChaosVDTabID::IndependentDetailsPanel1, ETabState::ClosedTab)
+		->AddTab(FChaosVDTabID::IndependentDetailsPanel2, ETabState::ClosedTab)
+		->AddTab(FChaosVDTabID::IndependentDetailsPanel3, ETabState::ClosedTab)
+		->AddTab(FChaosVDTabID::IndependentDetailsPanel4, ETabState::ClosedTab)
 	));
 }
 
diff --git a/Engine/Plugins/ChaosVD/Source/ChaosVD/Public/ChaosVDTabsIDs.h b/Engine/Plugins/ChaosVD/Source/ChaosVD/Public/ChaosVDTabsIDs.h
--- a/Engine/Plugins/ChaosVD/Source/ChaosVD/Public/ChaosVDTabsIDs.h
+++ b/Engine/Plugins/ChaosVD/Source/ChaosVD/Public/ChaosVDTabsIDs.h
@@ -13,6 +13,10 @@ public:
 	static const FName PlaybackViewport;
 	static const FName WorldOutliner;
 	static const FName DetailsPanel;
+	static const FName IndependentDetailsPanel1;
+	static const FName IndependentDetailsPanel2;
+	static const FName IndependentDetailsPanel3;
+	static const FName IndependentDetailsPanel4;
 	static const FName OutputLog;
 	static const FName SolversTrack;
 	static const FName StatusBar;
diff --git a/Engine/Plugins/ChaosVD/Source/ChaosVD/Public/Widgets/SChaosVDMainTab.h b/Engine/Plugins/ChaosVD/Source/ChaosVD/Public/Widgets/SChaosVDMainTab.h
--- a/Engine/Plugins/ChaosVD/Source/ChaosVD/Public/Widgets/SChaosVDMainTab.h
+++ b/Engine/Plugins/ChaosVD/Source/ChaosVD/Public/Widgets/SChaosVDMainTab.h
@@ -19,9 +19,11 @@ class FChaosVDTabSpawnerBase;
 class FChaosVDEditorVisualizationSettingsTab;
 class FChaosVDSolversTracksTab;
 class FChaosVDEngine;
+class FChaosVDIndependentDetailsPanelManager;
 class FChaosVDOutputLogTab;
 class FChaosVDPlaybackViewportTab;
 class FChaosVDObjectDetailsTab;
+class FChaosVDStandAloneObjectDetailsTab;
 class FChaosVDWorldOutlinerTab;
 class FStructOnScope;
 class FUICommandList;
@@ -54,29 +56,48 @@ public:
 
 	virtual ~SChaosVDMainTab() override;
 
-	TSharedRef<FChaosVDEngine> GetChaosVDEngineInstance() const { return ChaosVDEngine.ToSharedRef(); }
+	TSharedRef<FChaosVDEngine> GetChaosVDEngineInstance() const
+	{
+		return ChaosVDEngine.ToSharedRef();
+	}
 
 	template<typename TabType>
 	TWeakPtr<TabType> GetTabSpawnerInstance(FName TabID);
 
 	// BEGIN ITOOLKITHOST Interface
-	virtual TSharedRef<SWidget> GetParentWidget() override { return AsShared(); }
+	virtual TSharedRef<SWidget> GetParentWidget() override
+	{
+		return AsShared();
+	}
+
 	virtual void BringToFront() override;
-	virtual TSharedPtr<FTabManager> GetTabManager() const override { return TabManager; };
+	virtual TSharedPtr<FTabManager> GetTabManager() const override
+	{
+		return TabManager;
+	};
 	virtual void OnToolkitHostingStarted(const TSharedRef<IToolkit>& Toolkit) override;
 	virtual void OnToolkitHostingFinished(const TSharedRef<IToolkit>& Toolkit) override;
 	virtual UWorld* GetWorld() const override;
 	virtual FEditorModeTools& GetEditorModeManager() const override;
 
-	virtual UTypedElementCommonActions* GetCommonActions() const override { return nullptr; }
+	virtual UTypedElementCommonActions* GetCommonActions() const override
+	{
+		return nullptr;
+	}
 	virtual FName GetStatusBarName() const override { return StatusBarID; };
-	virtual FOnActiveViewportChanged& OnActiveViewportChanged() override { return ViewportChangedDelegate;};
+	virtual FOnActiveViewportChanged& OnActiveViewportChanged() override
+	{
+		return ViewportChangedDelegate;
+	};
 	// END ITOOLKITHOST Interface
 
 	TSharedPtr<FComponentVisualizer> FindComponentVisualizer(UClass* ClassPtr);
 	TSharedPtr<FComponentVisualizer> FindComponentVisualizer(FName ClassName);
 
-	TConstArrayView<TSharedPtr<FComponentVisualizer>> GetAllComponentVisualizers() { return ComponentVisualizers; }
+	TConstArrayView<TSharedPtr<FComponentVisualizer>> GetAllComponentVisualizers()
+	{
+		return ComponentVisualizers;
+	}
 	
 	bool ConnectToLiveSession(int32 SessionID, FString SessionAddress) const;
 
@@ -87,7 +108,6 @@ public:
 	TSharedRef<IStructureDetailsView> CreateStructureDetailsView(const FDetailsViewArgs& InDetailsViewArgs, const FStructureDetailsViewArgs& InStructureDetailsViewArgs, const TSharedPtr<FStructOnScope>& InStructData = nullptr, const FText& CustomName = FText::GetEmpty());
 	void ProccessKeyEventForPlaybackTrackSelector(const FKeyEvent& InKeyEvent);
 
-
 	TSharedPtr<FUICommandList> GetGlobalUICommandList() const
 	{
 		return GlobalCommandList;
@@ -108,6 +128,8 @@ public:
 
 	CHAOSVD_API void RegisterComponentVisualizer(FName ClassName, const TSharedPtr<FComponentVisualizer>& Visualizer);
 
+	const TSharedPtr<FChaosVDIndependentDetailsPanelManager>& GetIndependentDetailsPanelManager();
+
 private:
 
 	void HandlePostInitializationExtensionRegistered(const TSharedRef<FChaosVDExtension>& NewExtension);
@@ -168,13 +190,14 @@ private:
 
 	bool bShowTrackSelectorKeyShortcut = false;
 
+	TSharedPtr<FChaosVDIndependentDetailsPanelManager> IndependentDetailsPanelManager;
+
 	/** Commandlist used for any UI actions that need to be processed globally regardless on which specific widget we are in */
 	TSharedPtr<FUICommandList> GlobalCommandList;
 
 	static inline const FName MainToolBarName = FName("ChaosVD.MainToolBar");
 };
 
-
 template <typename TabType>
 TWeakPtr<TabType> SChaosVDMainTab::GetTabSpawnerInstance(FName TabID)
 {
