[ChaosVD] Fix for an issue where changes to any settings (including all in the show menu) are not added to the undo buffer, and therefore they can't be undo with CTRL + Z.

Also fixed an issue where if you change the color mode, the view port does not get automatically re-drawn. Same for the Generic Debug Draw Settings

diff --git a/Engine/Plugins/ChaosVD/Source/ChaosVD/Private/ChaosVDSettingsManager.cpp b/Engine/Plugins/ChaosVD/Source/ChaosVD/Private/ChaosVDSettingsManager.cpp
--- a/Engine/Plugins/ChaosVD/Source/ChaosVD/Private/ChaosVDSettingsManager.cpp
+++ b/Engine/Plugins/ChaosVD/Source/ChaosVD/Private/ChaosVDSettingsManager.cpp
@@ -41,6 +41,9 @@ UChaosVDSettingsObjectBase* FChaosVDSettingsManager::GetSettingsObject(UClass* S
 
 	UChaosVDSettingsObjectBase* NewSettingsObject = NewObject<UChaosVDSettingsObjectBase>(SettingsOuter, SettingsClass);
 
+	// Needed so any changes are added to the undo buffer
+	NewSettingsObject->SetFlags(RF_Transactional);
+
 	AvailableSettingsObject.Add(SettingsClass, NewSettingsObject);
 
 	return NewSettingsObject;
diff --git a/Engine/Plugins/ChaosVD/Source/ChaosVD/Private/Settings/ChaosVDCoreSettings.cpp b/Engine/Plugins/ChaosVD/Source/ChaosVD/Private/Settings/ChaosVDCoreSettings.cpp
--- a/Engine/Plugins/ChaosVD/Source/ChaosVD/Private/Settings/ChaosVDCoreSettings.cpp
+++ b/Engine/Plugins/ChaosVD/Source/ChaosVD/Private/Settings/ChaosVDCoreSettings.cpp
@@ -45,6 +45,12 @@ void UChaosVDSettingsObjectBase::BroadcastSettingsChanged()
 void UChaosVDVisualizationSettingsObjectBase::BroadcastSettingsChanged()
 {
 	Super::BroadcastSettingsChanged();
-	SChaosVDPlaybackViewport::ExecuteExternalViewportInvalidateRequest();
-}
 
+	// All geometry related operations are queued and de-duplicated at the end of the frame before being applied
+	// Therefore we need to wait one frame to invalidate the viewport
+	FTSTicker::GetCoreTicker().AddTicker(FTickerDelegate::CreateLambda([](float DeltaTime)
+	{
+		SChaosVDPlaybackViewport::ExecuteExternalViewportInvalidateRequest();
+		return false;
+	}));
+}
diff --git a/Engine/Plugins/ChaosVD/Source/ChaosVDBuiltInExtensions/Private/GenericDebugDraw/Settings/ChaosVDGenericDebugDrawSettings.h b/Engine/Plugins/ChaosVD/Source/ChaosVDBuiltInExtensions/Private/GenericDebugDraw/Settings/ChaosVDGenericDebugDrawSettings.h
--- a/Engine/Plugins/ChaosVD/Source/ChaosVDBuiltInExtensions/Private/GenericDebugDraw/Settings/ChaosVDGenericDebugDrawSettings.h
+++ b/Engine/Plugins/ChaosVD/Source/ChaosVDBuiltInExtensions/Private/GenericDebugDraw/Settings/ChaosVDGenericDebugDrawSettings.h
@@ -19,7 +19,7 @@ enum class EChaosVDGenericDebugDrawVisualizationFlags : uint32
 ENUM_CLASS_FLAGS(EChaosVDGenericDebugDrawVisualizationFlags);
 
 UCLASS(config=ChaosVD, PerObjectConfig)
-class UChaosVDGenericDebugDrawSettings : public UChaosVDSettingsObjectBase
+class UChaosVDGenericDebugDrawSettings : public UChaosVDVisualizationSettingsObjectBase
 {
 	GENERATED_BODY()
 
