[ChaosVD] Fix issue where scene queries done with a Box or Sphere shape are shown with the incorrect size

This is fallout of the refactor to de-duplicate geometry generation. As part of that work for simple shapes like spheres or boxes, we just use a pre-made static mesh of each shape, and then calculate a scaled transform to achieve the size of the shape, instead of generating different boxes/spheres for each size.

The problem is that the debug draw code for implicit objects does not use meshes but relies in the same generator code. As the debug draw path was not updated to take into account that change, we were scaling already sized simple shapes.

This CL adds an option to opt-out of the scaling behavior, which debug draw now uses.

diff --git a/Engine/Plugins/ChaosVD/Source/ChaosVD/Private/ChaosVDGeometryBuilder.cpp b/Engine/Plugins/ChaosVD/Source/ChaosVD/Private/ChaosVDGeometryBuilder.cpp
--- a/Engine/Plugins/ChaosVD/Source/ChaosVD/Private/ChaosVDGeometryBuilder.cpp
+++ b/Engine/Plugins/ChaosVD/Source/ChaosVD/Private/ChaosVDGeometryBuilder.cpp
@@ -859,7 +859,7 @@ const Chaos::FImplicitObject* FChaosVDGeometryBuilder::UnpackImplicitObject(cons
 	return nullptr;
 }
 
-void FChaosVDGeometryBuilder::AdjustedTransformForImplicit(const Chaos::FImplicitObject* InImplicit, FTransform& OutAdjustedTransform)
+void FChaosVDGeometryBuilder::AdjustedTransformForImplicit(const Chaos::FImplicitObject* InImplicit, FTransform& OutAdjustedTransform, EChaosVDGeometryTransformGeneratorFlags Options)
 {
 	using namespace Chaos;
 	
@@ -885,7 +885,7 @@ void FChaosVDGeometryBuilder::AdjustedTransformForImplicit(const Chaos::FImplici
 			break;
 	}
 
-	if (Chaos::VisualDebugger::Cvars::bDeduplicateSimpleGeometry)
+	if (EnumHasAnyFlags(Options, EChaosVDGeometryTransformGeneratorFlags::UseScaleForSize))
 	{
 		switch (InnerType)
 		{
@@ -976,7 +976,8 @@ TSharedPtr<FChaosVDExtractedGeometryDataHandle> FChaosVDGeometryBuilder::Extract
 			DispatchCreateAndCacheMeshForImplicitAsync(GeometryKey, ImplicitObjectToProcess);
 		}
 
-		AdjustedTransformForImplicit(ImplicitObjectToProcess, ExtractedTransform);
+		EChaosVDGeometryTransformGeneratorFlags TransformUpdateFlags = Chaos::VisualDebugger::Cvars::bDeduplicateSimpleGeometry ? EChaosVDGeometryTransformGeneratorFlags::UseScaleForSize : EChaosVDGeometryTransformGeneratorFlags::None;
+		AdjustedTransformForImplicit(ImplicitObjectToProcess, ExtractedTransform, TransformUpdateFlags);
 		MeshDataHandle->SetTransform(ExtractedTransform);
 
 		return MeshDataHandle;
diff --git a/Engine/Plugins/ChaosVD/Source/ChaosVD/Public/ChaosVDGeometryBuilder.h b/Engine/Plugins/ChaosVD/Source/ChaosVD/Public/ChaosVDGeometryBuilder.h
--- a/Engine/Plugins/ChaosVD/Source/ChaosVD/Public/ChaosVDGeometryBuilder.h
+++ b/Engine/Plugins/ChaosVD/Source/ChaosVD/Public/ChaosVDGeometryBuilder.h
@@ -42,6 +42,15 @@ namespace UE
 typedef TWeakObjectPtr<UMeshComponent> FMeshComponentWeakPtr;
 typedef TSharedPtr<FChaosVDExtractedGeometryDataHandle> FExtractedGeometryHandle;
 
+/** Set of flags used to control how we generate a transform from implicit object data */
+enum class EChaosVDGeometryTransformGeneratorFlags
+{
+	None = 0,
+	/** When calculating the adjusted transform, it will generate a scale to represent the actual size of the implicit object */
+	UseScaleForSize = 1 << 0,
+};
+ENUM_CLASS_FLAGS(EChaosVDGeometryTransformGeneratorFlags)
+
 /*
  * Generates Dynamic mesh components and dynamic meshes based on Chaos implicit object data
  */
@@ -221,7 +230,7 @@ public:
 	 * @param InImplicit Source implicit object
 	 * @param OutAdjustedTransform Updated transform
 	 */
-	void AdjustedTransformForImplicit(const Chaos::FImplicitObject* InImplicit, FTransform& OutAdjustedTransform);
+	void AdjustedTransformForImplicit(const Chaos::FImplicitObject* InImplicit, FTransform& OutAdjustedTransform, EChaosVDGeometryTransformGeneratorFlags Options = EChaosVDGeometryTransformGeneratorFlags::None);
 
 private:
 
