[ChaosVD] Fix issue where collision data details panel and scene query details panel is unloaded if another selection is triggered.

The issue was caused because in the last refactor we changed how changes to the geometry visualization is applied (the selection outline is one of these changes). To invalidate the viewport I was using the scene updated delegate, but that is too wide. When selection changes the scene doesn't really changes, therefore the fix is just use the dedicated callback we have to invalidate the viewport.

Also fixed a log spam that happens if you stop PIE while doing a live recording. In that scenario you can end up with a solver instance (the editor world) that does not have any active solver stage. This is a valid scenario and handled properly by the code that emits the error but given live recording is continuous playback, it ends up being spammy. The log line is still valuable in other cases, so I decided just make it verbose and we can change it during debugging if needed.

diff --git a/Engine/Plugins/ChaosVD/Source/ChaosVD/Private/ChaosVDPlaybackController.cpp b/Engine/Plugins/ChaosVD/Source/ChaosVD/Private/ChaosVDPlaybackController.cpp
--- a/Engine/Plugins/ChaosVD/Source/ChaosVD/Private/ChaosVDPlaybackController.cpp
+++ b/Engine/Plugins/ChaosVD/Source/ChaosVD/Private/ChaosVDPlaybackController.cpp
@@ -221,7 +221,7 @@ void FChaosVDPlaybackController::GoToRecordedSolverStage_AssumesLocked(const int
 				{
 					if (SolverFrameData->SolverSteps.IsEmpty() || TargetStageIndex != INDEX_NONE)
 					{
-						UE_LOG(LogChaosVDEditor, Error, TEXT("[%s] Invalid solver stage index [%d] at frame [%d] for Solver ID [%d]. We can't continue..."), ANSI_TO_TCHAR(__FUNCTION__), TargetStageIndex, FrameNumber, InTrackID);
+						UE_LOG(LogChaosVDEditor, Verbose, TEXT("[%s] Invalid solver stage index [%d] at frame [%d] for Solver ID [%d]. We can't continue..."), ANSI_TO_TCHAR(__FUNCTION__), TargetStageIndex, FrameNumber, InTrackID);
 						return;
 					}
 
diff --git a/Engine/Plugins/ChaosVD/Source/ChaosVD/Private/Components/ChaosVDInstancedStaticMeshComponent.cpp b/Engine/Plugins/ChaosVD/Source/ChaosVD/Private/Components/ChaosVDInstancedStaticMeshComponent.cpp
--- a/Engine/Plugins/ChaosVD/Source/ChaosVD/Private/Components/ChaosVDInstancedStaticMeshComponent.cpp
+++ b/Engine/Plugins/ChaosVD/Source/ChaosVD/Private/Components/ChaosVDInstancedStaticMeshComponent.cpp
@@ -5,6 +5,8 @@
 #include "ChaosVDGeometryBuilder.h"
 #include "ChaosVDModule.h"
 #include "Actors/ChaosVDGeometryContainer.h"
+#include "Widgets/SChaosVDPlaybackViewport.h"
+
 
 void UChaosVDInstancedStaticMeshComponent::UpdateVisibilityForInstance(const TSharedRef<FChaosVDInstancedMeshData>& InInstanceHandle)
 {
@@ -488,13 +490,7 @@ void UChaosVDInstancedStaticMeshComponent::ProcessChanges()
 
 		if (bHasSelectionChange)
 		{
-			if (AChaosVDGeometryContainer* GeometryContainer = Cast<AChaosVDGeometryContainer>(GetOwner()))
-			{
-				if (TSharedPtr<FChaosVDScene> ScenePtr = GeometryContainer->GetScene())
-				{
-					ScenePtr->RequestUpdate();
-				}
-			}
+			SChaosVDPlaybackViewport::ExecuteExternalViewportInvalidateRequest();
 		}
 
 		if (CurrentInstanceHandlesByIndex.Num() == 0 && PendingOperationsByInstance.Num() == 0)
