[ChaosVD] Fix for an issue where scrubbing between solver stages using the next and prev buttons, not always invalidated the viewport, resulting in an out of date visualization until you interact with the viewport

Since the last refactor, now CVD queue most geometry operations and de-duplicates them /process them during a GT tick. The UpdateScene method for solver stages was still invalidation the viewport directly instead of queing it in the tick.

diff --git a/Engine/Plugins/ChaosVD/Source/ChaosVD/Private/ChaosVDScene.cpp b/Engine/Plugins/ChaosVD/Source/ChaosVD/Private/ChaosVDScene.cpp
--- a/Engine/Plugins/ChaosVD/Source/ChaosVD/Private/ChaosVDScene.cpp
+++ b/Engine/Plugins/ChaosVD/Source/ChaosVD/Private/ChaosVDScene.cpp
@@ -158,7 +158,7 @@ void FChaosVDScene::UpdateFromRecordedSolverStageData(const int32 SolverID, cons
 
 	SolverSceneData->UpdateFromNewSolverStageData(InFrameData, InRecordedStepData);
 
-	OnSceneUpdated().Broadcast();
+	bPendingUpdateRequest = true;
 }
 
 void FChaosVDScene::AddFromCVDWorldTagToActor(AActor* Actor)
