[ChaosVD] Adding support to load multiple CVD files in a single instance or connect to multiple live sessions.

This allow us to record a separate server and client instances from packaged builds and combine their visualization in a single CVD instance to debug networking issues. This features was only available in PIE sessions because both worlds run under the same instance in that case.

Main changes:

- Introduced a new option to the Open File and Connect to session modals, called Loading Mode. Single mode will unload any CVD file before loading a new one, MultiSource will load the CVD data on top of an already loaded file
- Added a way to Re-Map solver ids on the fly during trace analysis. These IDs are counters, therefore two recordings from two separate instances will have colliding ids, therefore we need to re-map them on load.
- Updated all data processors to remap the solver ID on load
- Removed 2 unused CVD Trace Events

Work Still to be done:
- Game Thread Data (currently Scene Queries & Generic Debug Draw) is discarded in any additional CVD file. This will be addressed in a follow up CL
- The Recorded Log Browser only shows the data for the first loaded CVD file. Same as before,  this will be addressed in a follow up CL.

diff --git a/Engine/Plugins/ChaosVD/Source/ChaosVD/Private/ChaosVDEngine.cpp b/Engine/Plugins/ChaosVD/Source/ChaosVD/Private/ChaosVDEngine.cpp
--- a/Engine/Plugins/ChaosVD/Source/ChaosVD/Private/ChaosVDEngine.cpp
+++ b/Engine/Plugins/ChaosVD/Source/ChaosVD/Private/ChaosVDEngine.cpp
@@ -18,7 +18,6 @@ void FChaosVDEngine::Initialize()
 	}
 
 	// Create an Empty Scene
-	//TODO: Handle multiple scenes. We will need it to represent multiple worlds
 	CurrentScene = MakeShared<FChaosVDScene>();
 	CurrentScene->Initialize();
 
@@ -31,13 +30,28 @@ void FChaosVDEngine::Initialize()
 	{
 		if (const TSharedPtr<FChaosVDEngine> CVDEngine = WeakThis.Pin())
 		{
-			CVDEngine->CurrentSessionDescriptor.bIsLiveSession = false;
+			if (CVDEngine->HasAnyLiveSessionActive())
+			{
+				CVDEngine->CloseActiveSessions();
+			}
 		}
 	}));
 
 	bIsInitialized = true;
 }
 
+void FChaosVDEngine::CloseActiveSessions()
+{
+	if (const TSharedPtr<FChaosVDTraceManager> CVDTraceManager = FChaosVDModule::Get().GetTraceManager())
+	{
+		for (FChaosVDTraceSessionDescriptor& SessionDescriptor : CurrentSessionDescriptors)
+		{
+			CVDTraceManager->CloseSession(SessionDescriptor.SessionName);
+			SessionDescriptor.bIsLiveSession = false;
+		}	
+	}
+}
+
 void FChaosVDEngine::DeInitialize()
 {
 	if (!bIsInitialized)
@@ -49,10 +63,9 @@ void FChaosVDEngine::DeInitialize()
 	CurrentScene.Reset();
 	PlaybackController.Reset();
 
-	if (const TSharedPtr<FChaosVDTraceManager> CVDTraceManager = FChaosVDModule::Get().GetTraceManager())
-	{
-		CVDTraceManager->CloseSession(CurrentSessionDescriptor.SessionName);
-	}
+	CloseActiveSessions();
+	
+	CurrentSessionDescriptors.Reset();
 
 	if (FChaosVDRuntimeModule::IsLoaded())
 	{
@@ -95,24 +108,81 @@ void FChaosVDEngine::UpdateRecentFilesList(const FString& InFilename)
 	MiscSettings->SaveConfig();
 }
 
-void FChaosVDEngine::LoadRecording(const FString& FilePath)
+bool FChaosVDEngine::HasAnyLiveSessionActive() const
+{
+	return PlaybackController && PlaybackController->IsPlayingLiveSession();
+}
+
+void FChaosVDEngine::LoadRecording(const FString& FilePath, EChaosVDLoadRecordedDataMode LoadingMode)
 {
 	FChaosVDTraceSessionDescriptor NewSessionFromFileDescriptor;
-	NewSessionFromFileDescriptor.SessionName = FChaosVDModule::Get().GetTraceManager()->LoadTraceFile(FilePath);
+
+	if (LoadingMode == EChaosVDLoadRecordedDataMode::MultiSource && !CurrentSessionDescriptors.IsEmpty())
+	{
+		TSharedPtr<FChaosVDRecording> ExistingRecordingInstance = nullptr;
+		if (const TSharedPtr<const TraceServices::IAnalysisSession> TraceSession = FChaosVDModule::Get().GetTraceManager()->GetSession(CurrentSessionDescriptors[0].SessionName))
+		{
+			if (const FChaosVDTraceProvider* ChaosVDProvider = TraceSession->ReadProvider<FChaosVDTraceProvider>(FChaosVDTraceProvider::ProviderName))
+			{
+				ExistingRecordingInstance = ChaosVDProvider->GetRecordingForSession();
+			}
+		}
+
+		NewSessionFromFileDescriptor.SessionName = FChaosVDModule::Get().GetTraceManager()->LoadTraceFile(FilePath, ExistingRecordingInstance);
+	}
+	else
+	{
+		NewSessionFromFileDescriptor.SessionName = FChaosVDModule::Get().GetTraceManager()->LoadTraceFile(FilePath);
+	}
+	
 	NewSessionFromFileDescriptor.bIsLiveSession = false;
 
 	if (NewSessionFromFileDescriptor.IsValid())
 	{
 		UpdateRecentFilesList(FilePath);
+		OpenSession(NewSessionFromFileDescriptor, LoadingMode);
+	}
+}
+
+bool FChaosVDEngine::ConnectToLiveSession(int32 SessionID, const FString& InSessionAddress, EChaosVDLoadRecordedDataMode LoadingMode)
+{
+	FChaosVDTraceSessionDescriptor NewSessionFromFileDescriptor;
+
+	NewSessionFromFileDescriptor.bIsLiveSession = true;
+	
+	TSharedPtr<FChaosVDRecording> ExistingRecordingInstance = nullptr;
+	if (LoadingMode == EChaosVDLoadRecordedDataMode::MultiSource && PlaybackController)
+	{
+		ExistingRecordingInstance = PlaybackController->GetCurrentRecording().Pin();
+	}
+
+	NewSessionFromFileDescriptor.SessionName = FChaosVDModule::Get().GetTraceManager()->ConnectToLiveSession(InSessionAddress, SessionID, ExistingRecordingInstance);
+
+	bool bSuccess = false;
+
+	if (NewSessionFromFileDescriptor.IsValid())
+	{
+		OpenSession(NewSessionFromFileDescriptor);
+		bSuccess = true;
 	}
 
-	SetCurrentSession(NewSessionFromFileDescriptor);
+	return bSuccess;
 }
 
-void FChaosVDEngine::SetCurrentSession(const FChaosVDTraceSessionDescriptor& SessionDescriptor)
+void FChaosVDEngine::OpenSession(const FChaosVDTraceSessionDescriptor& SessionDescriptor, EChaosVDLoadRecordedDataMode LoadingMode)
 {
-	CurrentSessionDescriptor = SessionDescriptor;
-	PlaybackController->LoadChaosVDRecordingFromTraceSession(CurrentSessionDescriptor);
+	if (LoadingMode == EChaosVDLoadRecordedDataMode::SingleSource)
+	{
+		CloseActiveSessions();
+		CurrentSessionDescriptors.Reset();
+	}
+	
+	CurrentSessionDescriptors.Emplace(SessionDescriptor);
+
+	if (CurrentSessionDescriptors.Num() <= 1)
+	{
+		PlaybackController->LoadChaosVDRecordingFromTraceSession(SessionDescriptor);
+	}
 }
 
 bool FChaosVDEngine::Tick(float DeltaTime)
diff --git a/Engine/Plugins/ChaosVD/Source/ChaosVD/Private/ChaosVDEngine.h b/Engine/Plugins/ChaosVD/Source/ChaosVD/Private/ChaosVDEngine.h
--- a/Engine/Plugins/ChaosVD/Source/ChaosVD/Private/ChaosVDEngine.h
+++ b/Engine/Plugins/ChaosVD/Source/ChaosVD/Private/ChaosVDEngine.h
@@ -7,11 +7,23 @@
 #include "Misc/Guid.h"
 #include "Trace/ChaosVDTraceManager.h"
 
+#include "ChaosVDEngine.generated.h"
+
 class FChaosVDTraceManager;
 class FChaosVDPlaybackController;
 class FChaosVDScene;
 class FChaosVisualDebuggerMainUI;
 
+/** Enumeration of the available modes controlling how data is loaded into CVD */
+UENUM()
+enum class EChaosVDLoadRecordedDataMode : uint8
+{
+	/** This mode will unload any CVD recording currently loaded before loading the selected file */
+	SingleSource,
+	/** CVD will load and merge the data of the selected recording into the currently loaded recording */
+	MultiSource
+};
+
 /** Core Implementation of the visual debugger - Owns the systems that are not UI */
 class FChaosVDEngine : public FTSTickerObjectBase, public TSharedFromThis<FChaosVDEngine>
 {
@@ -22,20 +34,38 @@ public:
 	}
 
 	void Initialize();
-	
+	void CloseActiveSessions();
+
 	void DeInitialize();
 
 	virtual bool Tick(float DeltaTime) override;
 
-	const FGuid& GetInstanceGuid() const { return InstanceGUID; };
-	TSharedPtr<FChaosVDScene>& GetCurrentScene() { return CurrentScene; };
-	TSharedPtr<FChaosVDPlaybackController>& GetPlaybackController() { return PlaybackController; };
+	const FGuid& GetInstanceGuid() const
+	{
+		return InstanceGUID;
+	}
+
+	TSharedPtr<FChaosVDScene>& GetCurrentScene()
+	{
+		return CurrentScene;
+	}
+
+	TSharedPtr<FChaosVDPlaybackController>& GetPlaybackController()
+	{
+		return PlaybackController;
+	}
+
+	TArrayView<FChaosVDTraceSessionDescriptor> GetCurrentSessionDescriptors()
+	{
+		return CurrentSessionDescriptors;
+	}
 
-	const FChaosVDTraceSessionDescriptor& GetCurrentSessionDescriptor() { return CurrentSessionDescriptor; };
+	bool HasAnyLiveSessionActive() const;
 
-	void LoadRecording(const FString& FilePath);
+	void LoadRecording(const FString& FilePath, EChaosVDLoadRecordedDataMode LoadingMode = EChaosVDLoadRecordedDataMode::MultiSource);
+	bool ConnectToLiveSession(int32 SessionID, const FString& InSessionAddress, EChaosVDLoadRecordedDataMode LoadingMode = EChaosVDLoadRecordedDataMode::MultiSource);
 
-	void SetCurrentSession(const FChaosVDTraceSessionDescriptor& SessionDescriptor);
+	void OpenSession(const FChaosVDTraceSessionDescriptor& SessionDescriptor, EChaosVDLoadRecordedDataMode LoadingMode = EChaosVDLoadRecordedDataMode::MultiSource);
 	
 private:
 
@@ -43,7 +73,7 @@ private:
 
 	FGuid InstanceGUID;
 
-	FChaosVDTraceSessionDescriptor CurrentSessionDescriptor;
+	TArray<FChaosVDTraceSessionDescriptor> CurrentSessionDescriptors;
 
 	TSharedPtr<FChaosVDScene> CurrentScene;
 	TSharedPtr<FChaosVDPlaybackController> PlaybackController;
diff --git a/Engine/Plugins/ChaosVD/Source/ChaosVD/Private/ChaosVDPlaybackController.cpp b/Engine/Plugins/ChaosVD/Source/ChaosVD/Private/ChaosVDPlaybackController.cpp
--- a/Engine/Plugins/ChaosVD/Source/ChaosVD/Private/ChaosVDPlaybackController.cpp
+++ b/Engine/Plugins/ChaosVD/Source/ChaosVD/Private/ChaosVDPlaybackController.cpp
@@ -370,10 +370,13 @@ void FChaosVDPlaybackController::GoToTrackFrameAndSync(FGuid InstigatorID, EChao
 void FChaosVDPlaybackController::GoToTrackFrame_AssumesLockedAndSync(FGuid InstigatorID, EChaosVDTrackType TrackType, int32 InTrackID, int32 FrameNumber, int32 StageNumber)
 {
 	GoToTrackFrame_AssumesLocked(InstigatorID, TrackType, InTrackID, FrameNumber, StageNumber);
-	
-	if (TSharedPtr<const FChaosVDTrackInfo> TrackToSyncWith = GetTrackInfo(TrackType, InTrackID))
+
+	if (CurrentSyncMode != EChaosVDSyncTimelinesMode::Manual)
 	{
-		SyncTracks_AssumesLocked(TrackToSyncWith.ToSharedRef(), CurrentSyncMode);
+		if (TSharedPtr<const FChaosVDTrackInfo> TrackToSyncWith = GetTrackInfo(TrackType, InTrackID))
+		{
+			SyncTracks_AssumesLocked(TrackToSyncWith.ToSharedRef(), CurrentSyncMode);
+		}
 	}
 }
 
@@ -629,7 +632,19 @@ void FChaosVDPlaybackController::UnlockTrackStage(EChaosVDTrackType TrackType, i
 	}
 }
 
-void FChaosVDPlaybackController::GetAvailableTracks(EChaosVDTrackType TrackType, TArray<TSharedPtr<FChaosVDTrackInfo>>& OutTrackInfo)
+void FChaosVDPlaybackController::GetAvailableTracks(EChaosVDTrackType TrackType, TArray<TSharedPtr<const FChaosVDTrackInfo>>& OutTrackInfo)
+{
+	OutTrackInfo.Reset();
+	
+	TMap<int32, TSharedPtr<FChaosVDTrackInfo>>& AvailableTracksByType = TrackInfoPerType.FindOrAdd(TrackType);
+	
+	for (const TPair<int32, TSharedPtr<FChaosVDTrackInfo>>& Track : AvailableTracksByType)
+	{
+		OutTrackInfo.Emplace(Track.Value);
+	}
+}
+
+void FChaosVDPlaybackController::GetMutableAvailableTracks(EChaosVDTrackType TrackType, TArray<TSharedPtr<FChaosVDTrackInfo>>& OutTrackInfo)
 {
 	OutTrackInfo.Reset();
 	TrackInfoPerType.FindOrAdd(TrackType).GenerateValueArray(OutTrackInfo);
@@ -1189,7 +1204,13 @@ void FChaosVDPlaybackController::SetTimelineSyncMode(EChaosVDSyncTimelinesMode S
 			const TSharedRef<const FChaosVDTrackInfo> TrackToReSyncFromRef = TrackToReSyncFrom.ToSharedRef();
 
 			TrySetActiveTrack(TrackToReSyncFromRef);
-			SyncTracks(TrackToReSyncFromRef, CurrentSyncMode);
+
+			if (CurrentSyncMode != EChaosVDSyncTimelinesMode::Manual)
+			{
+				SyncTracks(TrackToReSyncFromRef, CurrentSyncMode);
+			}
+
+			OnDataUpdated().Broadcast(AsWeak());
 		}
 		else
 		{
diff --git a/Engine/Plugins/ChaosVD/Source/ChaosVD/Private/Trace/ChaosVDTraceAnalyzer.cpp b/Engine/Plugins/ChaosVD/Source/ChaosVD/Private/Trace/ChaosVDTraceAnalyzer.cpp
--- a/Engine/Plugins/ChaosVD/Source/ChaosVD/Private/Trace/ChaosVDTraceAnalyzer.cpp
+++ b/Engine/Plugins/ChaosVD/Source/ChaosVD/Private/Trace/ChaosVDTraceAnalyzer.cpp
@@ -24,9 +24,7 @@ void FChaosVDTraceAnalyzer::OnAnalysisBegin(const FOnAnalysisContext& Context)
 	Builder.RouteEvent(RouteId_ChaosVDBinaryDataContent, "ChaosVDLogger", "ChaosVDBinaryDataContent");
 	Builder.RouteEvent(RouteId_ChaosVDBinaryDataEnd, "ChaosVDLogger", "ChaosVDBinaryDataEnd");
 	Builder.RouteEvent(RouteId_ChaosVDSolverSimulationSpace, "ChaosVDLogger", "ChaosVDSolverSimulationSpace");
-	
-	Builder.RouteEvent(RouteId_ChaosVDNonSolverLocation, "ChaosVDLogger", "ChaosVDNonSolverLocation");
-	Builder.RouteEvent(RouteId_ChaosVDNonSolverTransform, "ChaosVDLogger", "ChaosVDNonSolverTransform");
+
 	Builder.RouteEvent(RouteId_ChaosVDNetworkTickOffset, "ChaosVDLogger", "ChaosVDNetworkTickOffset");
 	Builder.RouteEvent(RouteId_ChaosVDRolledBackDataID, "ChaosVDLogger", "ChaosVDRolledBackDataID");
 	Builder.RouteEvent(RouteId_ChaosVDUsesAutoRTFM, "ChaosVDLogger", "ChaosVDUsesAutoRTFM");
@@ -94,7 +92,9 @@ bool FChaosVDTraceAnalyzer::OnEvent(uint16 RouteId, EStyle Style, const FOnEvent
 			NewFrameData.bIsResimulated = EventData.GetValue<bool>("IsReSimulated");
 			NewFrameData.StartTime = Context.EventTime.AsSeconds(NewFrameData.FrameCycle);
 
-			if (int32* TickOffsetPtr = ChaosVDTraceProvider->GetCurrentTickOffsetsBySolverID().Find(NewFrameData.SolverID))
+			int32 RemappedExistingSolverID = ChaosVDTraceProvider->GetRemappedSolverID(NewFrameData.SolverID);
+
+			if (int32* TickOffsetPtr = ChaosVDTraceProvider->GetCurrentTickOffsetsBySolverID().Find(RemappedExistingSolverID))
 			{
 				NewFrameData.NetworkTickOffset = *TickOffsetPtr;
 			}
@@ -104,13 +104,21 @@ bool FChaosVDTraceAnalyzer::OnEvent(uint16 RouteId, EStyle Style, const FOnEvent
 			NewFrameData.DebugFName = FName(DebugNameView);
 
 			// Currently not all solvers have an end frame event, so lets just set the end frame time of the previous frame, with the start of this new one.
+			if (RemappedExistingSolverID != INDEX_NONE)
 			{
-				if (FChaosVDSolverFrameData* PrevFrameData = ChaosVDTraceProvider->GetCurrentSolverFrame(NewFrameData.SolverID))
+				if (FChaosVDSolverFrameData* PrevFrameData = ChaosVDTraceProvider->GetCurrentSolverFrame(RemappedExistingSolverID))
 				{
 					PrevFrameData->EndTime = NewFrameData.StartTime;
 				}
 			}
 
+			if (RemappedExistingSolverID == INDEX_NONE)
+			{
+				RemappedExistingSolverID = ChaosVDTraceProvider->RemapSolverID(NewFrameData.SolverID);	
+			}
+
+			NewFrameData.SolverID = RemappedExistingSolverID;
+
 			// Add an empty frame. It will be filled out by the solver trace events
 			ChaosVDTraceProvider->StartSolverFrame(NewFrameData.SolverID, MoveTemp(NewFrameData));
 			break;
@@ -121,7 +129,7 @@ bool FChaosVDTraceAnalyzer::OnEvent(uint16 RouteId, EStyle Style, const FOnEvent
 		}
 	case RouteId_ChaosVDSolverStepStart:
 		{
-			const int32 SolverID = EventData.GetValue<int32>("SolverID");
+			const int32 SolverID = ChaosVDTraceProvider->GetRemappedSolverID(EventData.GetValue<int32>("SolverID"));
 
 			// This can be null if the recording started Mid-Frame. In this case we just discard the data for now
 			if (FChaosVDSolverFrameData* FrameData = ChaosVDTraceProvider->GetCurrentSolverFrame(SolverID))
@@ -151,7 +159,7 @@ bool FChaosVDTraceAnalyzer::OnEvent(uint16 RouteId, EStyle Style, const FOnEvent
 		}
 	case RouteId_ChaosVDSolverStepEnd:
 		{
-			const int32 SolverID = EventData.GetValue<int32>("SolverID");
+			const int32 SolverID = ChaosVDTraceProvider->GetRemappedSolverID(EventData.GetValue<int32>("SolverID"));
 			if (FChaosVDSolverFrameData* FrameData = ChaosVDTraceProvider->GetCurrentSolverFrame(SolverID))
 			{
 				if (FrameData->SolverSteps.Num() > 0)
@@ -163,7 +171,7 @@ bool FChaosVDTraceAnalyzer::OnEvent(uint16 RouteId, EStyle Style, const FOnEvent
 		}
 	case RouteId_ChaosVDParticleDestroyed:
 		{
-			const int32 SolverID = EventData.GetValue<int32>("SolverID");
+			const int32 SolverID = ChaosVDTraceProvider->GetRemappedSolverID(EventData.GetValue<int32>("SolverID"));
 
 			if (FChaosVDSolverFrameData* FrameData = ChaosVDTraceProvider->GetCurrentSolverFrame(SolverID))
 			{
@@ -224,7 +232,7 @@ bool FChaosVDTraceAnalyzer::OnEvent(uint16 RouteId, EStyle Style, const FOnEvent
 		}
 	case RouteId_ChaosVDSolverSimulationSpace:
 		{
-			const int32 SolverID = EventData.GetValue<int32>("SolverID");
+			const int32 SolverID = ChaosVDTraceProvider->GetRemappedSolverID(EventData.GetValue<int32>("SolverID"));
 
 			FVector Position;
 			CVD_READ_TRACE_VECTOR(Position, Position, float, EventData);
@@ -238,33 +246,6 @@ bool FChaosVDTraceAnalyzer::OnEvent(uint16 RouteId, EStyle Style, const FOnEvent
 				FrameData->SimulationTransform.SetLocation(Position);
 				FrameData->SimulationTransform.SetRotation(Rotation);
 			}
-			break;
-		}
-	case RouteId_ChaosVDNonSolverLocation:
-		{
-			FChaosVDTrackedLocation TrackedLocation;
-			CVD_READ_TRACE_VECTOR(TrackedLocation.Location, Position, float, EventData);
-
-			EventData.GetString("DebugName", TrackedLocation.DebugName);
-			
-			if (TSharedPtr<FChaosVDGameFrameData> CurrentFrameData = ChaosVDTraceProvider->GetCurrentGameFrame().Pin())
-			{
-				CurrentFrameData->RecordedNonSolverLocationsByID.Add(FName(TrackedLocation.DebugName), MoveTemp(TrackedLocation));
-			}
-
-			break;
-		}
-	case RouteId_ChaosVDNonSolverTransform:
-		{
-			FChaosVDTrackedTransform TrackedTransform;
-			CVD_READ_TRACE_TRANSFORM(TrackedTransform.Transform, float, EventData);
-
-			EventData.GetString("DebugName", TrackedTransform.DebugName);
-			if (TSharedPtr<FChaosVDGameFrameData> CurrentFrameData = ChaosVDTraceProvider->GetCurrentGameFrame().Pin())
-			{
-				CurrentFrameData->RecordedNonSolverTransformsByID.Add(FName(TrackedTransform.DebugName), MoveTemp(TrackedTransform));
-			}
-
 			break;
 		}
 	case RouteId_ChaosVDNetworkTickOffset:
@@ -272,9 +253,12 @@ bool FChaosVDTraceAnalyzer::OnEvent(uint16 RouteId, EStyle Style, const FOnEvent
 			FChaosVDTrackedTransform TrackedTransform;
 
 			const int32 TickOffset = EventData.GetValue<int32>("Offset");
-			const int32 SolverID = EventData.GetValue<int32>("SolverID");
+			const int32 SolverID = ChaosVDTraceProvider->GetRemappedSolverID(EventData.GetValue<int32>("SolverID"));
 
-			ChaosVDTraceProvider->GetCurrentTickOffsetsBySolverID().FindOrAdd(SolverID, TickOffset);
+			if (SolverID != INDEX_NONE)
+			{
+				ChaosVDTraceProvider->GetCurrentTickOffsetsBySolverID().FindOrAdd(SolverID, TickOffset);
+			}
 
 			break;
 		}
diff --git a/Engine/Plugins/ChaosVD/Source/ChaosVD/Private/Trace/ChaosVDTraceAnalyzer.h b/Engine/Plugins/ChaosVD/Source/ChaosVD/Private/Trace/ChaosVDTraceAnalyzer.h
--- a/Engine/Plugins/ChaosVD/Source/ChaosVD/Private/Trace/ChaosVDTraceAnalyzer.h
+++ b/Engine/Plugins/ChaosVD/Source/ChaosVD/Private/Trace/ChaosVDTraceAnalyzer.h
@@ -51,8 +51,6 @@ private:
 		RouteId_BeginFrame,
 		RouteId_EndFrame,
 		RouteId_ChaosVDParticleDestroyed,
-		RouteId_ChaosVDNonSolverLocation,
-		RouteId_ChaosVDNonSolverTransform,
 		RouteId_ChaosVDNetworkTickOffset,
 		RouteId_ChaosVDRolledBackDataID,
 		RouteId_ChaosVDUsesAutoRTFM,
diff --git a/Engine/Plugins/ChaosVD/Source/ChaosVD/Private/Trace/ChaosVDTraceManager.cpp b/Engine/Plugins/ChaosVD/Source/ChaosVD/Private/Trace/ChaosVDTraceManager.cpp
--- a/Engine/Plugins/ChaosVD/Source/ChaosVD/Private/Trace/ChaosVDTraceManager.cpp
+++ b/Engine/Plugins/ChaosVD/Source/ChaosVD/Private/Trace/ChaosVDTraceManager.cpp
@@ -2,6 +2,7 @@
 
 #include "Trace/ChaosVDTraceManager.h"
 
+#include "ChaosVDEngine.h"
 #include "ChaosVDModule.h"
 #include "Features/IModularFeatures.h"
 #include "Modules/ModuleManager.h"
@@ -26,7 +27,7 @@ FChaosVDTraceManager::~FChaosVDTraceManager()
 	IModularFeatures::Get().UnregisterModularFeature(TraceServices::ModuleFeatureName, ChaosVDTraceModule.Get());
 }
 
-FString FChaosVDTraceManager::LoadTraceFile(const FString& InTraceFilename)
+FString FChaosVDTraceManager::LoadTraceFile(const FString& InTraceFilename, const TSharedPtr<FChaosVDRecording>& ExistingRecordingPtr)
 {
 	CloseSession(InTraceFilename);
 
@@ -35,7 +36,8 @@ FString FChaosVDTraceManager::LoadTraceFile(const FString& InTraceFilename)
 	{
 		if (const TSharedPtr<const TraceServices::IAnalysisSession> NewSession = TraceAnalysisService->StartAnalysis(*InTraceFilename))
 		{
-			AnalysisSessionByName.Add(InTraceFilename, NewSession);
+			HandleTraceSessionStarted(InTraceFilename, NewSession.ToSharedRef(), ExistingRecordingPtr);
+
 			return NewSession->GetName();
 		}
 	}
@@ -43,7 +45,20 @@ FString FChaosVDTraceManager::LoadTraceFile(const FString& InTraceFilename)
 	return FString();
 }
 
-FString FChaosVDTraceManager::ConnectToLiveSession(FStringView InSessionHost, uint32 SessionID)
+void FChaosVDTraceManager::HandleTraceSessionStarted(const FString& InTraceName, const TSharedRef<const TraceServices::IAnalysisSession>& InNewSession, const TSharedPtr<FChaosVDRecording>& InExistingRecordingPtr)
+{
+	if (InExistingRecordingPtr)
+	{
+		if (FChaosVDTraceProvider* ChaosVDProvider = ConstCastSharedRef<TraceServices::IAnalysisSession>(InNewSession)->EditProvider<FChaosVDTraceProvider>(FChaosVDTraceProvider::ProviderName))
+		{
+			ChaosVDProvider->SetExternalRecordingInstanceForSession(InExistingRecordingPtr.ToSharedRef());
+		}
+	}
+			
+	AnalysisSessionByName.Add(InTraceName, InNewSession);
+}
+
+FString FChaosVDTraceManager::ConnectToLiveSession(FStringView InSessionHost, uint32 SessionID, const TSharedPtr<FChaosVDRecording>& ExistingRecordingPtr)
 {	
 	using namespace UE::Trace;
 	FStoreClient* StoreClient = FStoreClient::Connect(InSessionHost.GetData());
@@ -81,9 +96,11 @@ FString FChaosVDTraceManager::ConnectToLiveSession(FStringView InSessionHost, ui
 		// Close this session in case we were already analysing it
 		CloseSession(TraceName);
 	
-		const TSharedPtr<const TraceServices::IAnalysisSession> NewSession = TraceAnalysisService->StartAnalysis(SessionID, *TraceName, MoveTemp(TraceData));
-		AnalysisSessionByName.Add(TraceName, NewSession);
-		SessionName = NewSession->GetName();
+		if (const TSharedPtr<const TraceServices::IAnalysisSession> NewSession = TraceAnalysisService->StartAnalysis(SessionID, *TraceName, MoveTemp(TraceData)))
+		{
+			HandleTraceSessionStarted(TraceName, NewSession.ToSharedRef(), ExistingRecordingPtr);
+			SessionName = NewSession->GetName();
+		}
 	}
 
 	return SessionName;
diff --git a/Engine/Plugins/ChaosVD/Source/ChaosVD/Private/Trace/ChaosVDTraceManager.h b/Engine/Plugins/ChaosVD/Source/ChaosVD/Private/Trace/ChaosVDTraceManager.h
--- a/Engine/Plugins/ChaosVD/Source/ChaosVD/Private/Trace/ChaosVDTraceManager.h
+++ b/Engine/Plugins/ChaosVD/Source/ChaosVD/Private/Trace/ChaosVDTraceManager.h
@@ -7,6 +7,7 @@
 #include "Templates/SharedPointer.h"
 #include "Trace/StoreClient.h"
 
+enum class EChaosVDLoadRecordedDataMode : uint8;
 class FChaosVDEngine;
 class FChaosVDTraceModule;
 
@@ -33,7 +34,7 @@ public:
 	/** Load a trace file and starts analyzing it
 	 * @param InTraceFilename File Name including (Path Included) of the Trace file to load
 	 */
-	FString LoadTraceFile(const FString& InTraceFilename);
+	FString LoadTraceFile(const FString& InTraceFilename, const TSharedPtr<FChaosVDRecording>& ExistingRecordingPtr = nullptr);
 	
 	/**
 	 * Connects to a live Trace Session and starts analyzing it.
@@ -41,7 +42,7 @@ public:
 	 * @param SessionID Trace ID in the Trace Store provided as host
 	 * @return Session Name
 	 */
-	FString ConnectToLiveSession(FStringView InSessionHost, uint32 SessionID);
+	FString ConnectToLiveSession(FStringView InSessionHost, uint32 SessionID, const TSharedPtr<FChaosVDRecording>& ExistingRecordingPtr = nullptr);
 
 	/** Returns the path to the local trace store */
 	FString GetLocalTraceStoreDirPath();
@@ -56,6 +57,8 @@ public:
 	static void EnumerateActiveSessions(FStringView InSessionHost, TVisitor Callback);
 
 private:
+	
+	void HandleTraceSessionStarted(const FString& InTraceName, const TSharedRef<const TraceServices::IAnalysisSession>& InNewSession, const TSharedPtr<FChaosVDRecording>& InExistingRecordingPtr);
 
 	/** The trace analysis session. */
 	TMap<FString, TSharedPtr<const TraceServices::IAnalysisSession>> AnalysisSessionByName;
diff --git a/Engine/Plugins/ChaosVD/Source/ChaosVD/Private/Trace/ChaosVDTraceModule.cpp b/Engine/Plugins/ChaosVD/Source/ChaosVD/Private/Trace/ChaosVDTraceModule.cpp
--- a/Engine/Plugins/ChaosVD/Source/ChaosVD/Private/Trace/ChaosVDTraceModule.cpp
+++ b/Engine/Plugins/ChaosVD/Source/ChaosVD/Private/Trace/ChaosVDTraceModule.cpp
@@ -2,6 +2,8 @@
 
 #include "Trace/ChaosVDTraceModule.h"
 
+#include "ChaosVDModule.h"
+#include "ChaosVDTraceManager.h"
 #include "Trace/ChaosVDTraceAnalyzer.h"
 #include "Trace/ChaosVDTraceProvider.h"
 #include "TraceServices/Model/AnalysisSession.h"
@@ -18,7 +20,7 @@ void FChaosVDTraceModule::OnAnalysisBegin(TraceServices::IAnalysisSession& InSes
 {
 	const TSharedRef<FChaosVDTraceProvider> Provider = MakeShared<FChaosVDTraceProvider>(InSession);
 
-	InSession.AddProvider(FChaosVDTraceProvider::ProviderName, Provider);
+	InSession.AddProvider(FChaosVDTraceProvider::ProviderName, Provider, Provider);
 	InSession.AddAnalyzer(new FChaosVDTraceAnalyzer(InSession, Provider));
 }
 
diff --git a/Engine/Plugins/ChaosVD/Source/ChaosVD/Private/Trace/ChaosVDTraceProvider.cpp b/Engine/Plugins/ChaosVD/Source/ChaosVD/Private/Trace/ChaosVDTraceProvider.cpp
--- a/Engine/Plugins/ChaosVD/Source/ChaosVD/Private/Trace/ChaosVDTraceProvider.cpp
+++ b/Engine/Plugins/ChaosVD/Source/ChaosVD/Private/Trace/ChaosVDTraceProvider.cpp
@@ -42,12 +42,29 @@ FChaosVDTraceProvider::FChaosVDTraceProvider(TraceServices::IAnalysisSession& In
 
 void FChaosVDTraceProvider::CreateRecordingInstanceForSession(const FString& InSessionName)
 {
+	if (bHasRecordingOverride)
+	{
+		return;
+	}
+
 	DeleteRecordingInstanceForSession();
 
 	InternalRecording = MakeShared<FChaosVDRecording>();
 	InternalRecording->SessionName = InSessionName;
 }
 
+void FChaosVDTraceProvider::SetExternalRecordingInstanceForSession(const TSharedRef<FChaosVDRecording>& InExternalCVDRecording)
+{
+	bHasRecordingOverride = true;
+	if (InternalRecording)
+	{
+		ensure(InternalRecording->IsEmpty());
+	}
+	InExternalCVDRecording->AddAttributes(EChaosVDRecordingAttributes::Merged);
+
+	InternalRecording = InExternalCVDRecording;
+}
+
 void FChaosVDTraceProvider::DeleteRecordingInstanceForSession()
 {
 	InternalRecording.Reset();
@@ -60,6 +77,13 @@ void FChaosVDTraceProvider::StartSolverFrame(const int32 InSolverID, FChaosVDSol
 		return;
 	}
 
+	bool bIsInvalidSolverID = InSolverID == INDEX_NONE && !CurrentSolverFramesByID.IsEmpty();
+
+	if (!ensure(!bIsInvalidSolverID))
+	{
+		return;
+	}
+
 	if (FChaosVDSolverFrameData* SolveFrameData = CurrentSolverFramesByID.Find(InSolverID))
 	{
 		InternalRecording->AddFrameForSolver(InSolverID, MoveTemp(*SolveFrameData));
@@ -95,6 +119,45 @@ FString FChaosVDTraceProvider::GenerateFormattedStringListFromSet(const TSet<FSt
 	return MoveTemp(FormattedString);
 }
 
+int32 FChaosVDTraceProvider::RemapSolverID(int32 SolverID)
+{
+	int32 RemappedSolverID = SolverID;
+	if (InternalRecording->HasSolverID(SolverID))
+	{
+		while (InternalRecording->HasSolverID(RemappedSolverID))
+		{
+			RemappedSolverID = FMath::RandRange(RemappedSolverID, std::numeric_limits<int32>::max());
+		}
+
+		RemappedSolversID.Emplace(SolverID, RemappedSolverID);
+	}
+	else
+	{
+		RemappedSolversID.Emplace(SolverID, SolverID);
+	}
+
+	return RemappedSolverID;
+}
+
+int32 FChaosVDTraceProvider::GetRemappedSolverID(int32 SolverID)
+{
+	int32 RemappedSolverID = SolverID;
+
+	if (InternalRecording && EnumHasAnyFlags(InternalRecording->GetAttributes(), EChaosVDRecordingAttributes::Merged))
+	{
+		if (int32* SolverIDPtr = RemappedSolversID.Find(SolverID))
+        {
+        	RemappedSolverID = *SolverIDPtr;
+        }
+		else
+        {
+			RemappedSolverID = INDEX_NONE; 
+        }
+	}
+
+	return RemappedSolverID;
+}
+
 void FChaosVDTraceProvider::CommitProcessedGameFramesToRecording()
 {
 	TArray<int32, TInlineAllocator<16>> SolverIDs;
@@ -134,7 +197,7 @@ void FChaosVDTraceProvider::CommitProcessedGameFramesToRecording()
 			{
 				InternalRecording->AddGameFrameData(*ProcessedGameFrameData);	
 			}
-		}	
+		}
 	}
 
 	SolverIDs.Reset();
@@ -142,7 +205,8 @@ void FChaosVDTraceProvider::CommitProcessedGameFramesToRecording()
 
 void FChaosVDTraceProvider::StartGameFrame(const TSharedPtr<FChaosVDGameFrameData>& InFrameData)
 {
-	if (!InternalRecording.IsValid())
+	//TODO: Add support for game frame data when we are in multi recording mode (that is the only case where bHasRecordingOverride will be true)
+	if (!InternalRecording.IsValid() || bHasRecordingOverride)
 	{
 		return;
 	}
@@ -154,6 +218,13 @@ void FChaosVDTraceProvider::StartGameFrame(const TSharedPtr<FChaosVDGameFrameDat
 
 FChaosVDSolverFrameData* FChaosVDTraceProvider::GetCurrentSolverFrame(const int32 InSolverID)
 {
+	bool bIsInvalidSolverID = InSolverID == INDEX_NONE && !CurrentSolverFramesByID.IsEmpty();
+
+	if (!ensure(!bIsInvalidSolverID))
+	{
+		return nullptr;
+	}
+
 	if (FChaosVDSolverFrameData* SolveFrameData = CurrentSolverFramesByID.Find(InSolverID))
 	{
 		return SolveFrameData;
diff --git a/Engine/Plugins/ChaosVD/Source/ChaosVD/Private/Trace/DataProcessors/ChaosVDArchiveHeaderProcessor.cpp b/Engine/Plugins/ChaosVD/Source/ChaosVD/Private/Trace/DataProcessors/ChaosVDArchiveHeaderProcessor.cpp
--- a/Engine/Plugins/ChaosVD/Source/ChaosVD/Private/Trace/DataProcessors/ChaosVDArchiveHeaderProcessor.cpp
+++ b/Engine/Plugins/ChaosVD/Source/ChaosVD/Private/Trace/DataProcessors/ChaosVDArchiveHeaderProcessor.cpp
@@ -4,6 +4,36 @@
 
 #include "ChaosVDModule.h"
 
+namespace Chaos::VisualDebugger::Private
+{
+	void PopulateInnerVersionsSet(const FChaosVDArchiveHeader& InHeader, TSet<FGuid>& OutInnerVersionsKeys)
+	{
+		const FCustomVersionArray& CustomVersions = InHeader.CustomVersionContainer.GetAllVersions();
+		OutInnerVersionsKeys.Reserve(CustomVersions.Num());
+
+		Algo::Transform(CustomVersions, OutInnerVersionsKeys, [](const FCustomVersion& Version)-> FGuid
+		{
+			return Version.Key;
+		});
+	}
+
+	bool IsCompatibleHeader(const FChaosVDArchiveHeader& InHeaderA, const FChaosVDArchiveHeader& InHeaderB)
+	{
+		// This is a slow operation but we only expect this to be called once when multiple CVD files are loaded
+		if (InHeaderA.EngineVersion.ExactMatch(InHeaderB.EngineVersion))
+		{
+			TSet<FGuid> VersionsASet;
+			PopulateInnerVersionsSet(InHeaderA, VersionsASet);
+	
+			TSet<FGuid> VersionsBSet;
+			PopulateInnerVersionsSet(InHeaderB, VersionsBSet);
+		
+			return VersionsASet.Num() == VersionsBSet.Num() && VersionsASet.Includes(VersionsBSet);
+		}
+
+		return false;
+	} 
+}
 
 FChaosVDArchiveHeaderProcessor::FChaosVDArchiveHeaderProcessor() : FChaosVDDataProcessorBase(Chaos::VisualDebugger::FChaosVDArchiveHeader::WrapperTypeName)
 {
@@ -11,6 +41,8 @@ FChaosVDArchiveHeaderProcessor::FChaosVDArchiveHeaderProcessor() : FChaosVDDataP
 
 bool FChaosVDArchiveHeaderProcessor::ProcessRawData(const TArray<uint8>& InData)
 {
+	using namespace Chaos::VisualDebugger;
+
 	FChaosVDDataProcessorBase::ProcessRawData(InData);
 
 	TSharedPtr<FChaosVDTraceProvider> ProviderSharedPtr = TraceProvider.Pin();
@@ -19,26 +51,32 @@ bool FChaosVDArchiveHeaderProcessor::ProcessRawData(const TArray<uint8>& InData)
 		return false;
 	}
 
-	Chaos::VisualDebugger::FChaosVDArchiveHeader RecordedHeaderData;
+	FChaosVDArchiveHeader RecordedHeaderData;
 	const bool bSuccess = ReadDataFromBuffer(InData, RecordedHeaderData, ProviderSharedPtr.ToSharedRef());
 
 	if (const TSharedPtr<FChaosVDRecording> RecordingInstance = ProviderSharedPtr->GetRecordingForSession())
 	{
 		// This works under the assumption the header is the first thing written and therefore the first thing read
-		// If that didn't happen, we need to know to investigate further
+		// If that didn't happen, we need to know to investigate further.
+		// Also, this assumption breaks under the new multi file support feature
+
+		//TODO: Add support for multiple headers in case we want to compare a new recording to a very old one
 
 		// Note: This is not a fatal error, unless pretty drastic serialization changes were made and the loaded file is old.
 		// CVD can gracefully handle serialization errors (as long the types serializers themselves don't crash and properly error out as expected).
-		if (!ensure(RecordingInstance->IsEmpty()))
+		if (!RecordingInstance->IsEmpty())
 		{
-			UE_LOG(LogChaosVDEditor, Warning, TEXT("[%s] Processed an archive header while the recording already had data loaded. That initially loaded data used the default header and serialization errors might have occured | This should not happen..."), ANSI_TO_TCHAR(__FUNCTION__));
+			if (ensure(Private::IsCompatibleHeader(RecordingInstance->GetHeaderData(), RecordedHeaderData)))
+			{
+				UE_LOG(LogChaosVDEditor, Warning, TEXT("[%s] Processed an archive header while the recording already had data loaded. That initially loaded data used the default header and serialization errors might have occured | This should not happen..."), ANSI_TO_TCHAR(__FUNCTION__));
+			}	
 		}
 
 		RecordingInstance->SetHeaderData(RecordedHeaderData);
 	}
 	else
 	{
-		UE_LOG(LogChaosVDEditor, Error, TEXT("[%s] Processed an archive header without a valid recording instace. Any follwoing archive reads will use the default header and serialization errors might occur..."), ANSI_TO_TCHAR(__FUNCTION__));	
+		UE_LOG(LogChaosVDEditor, Error, TEXT("[%s] Processed an archive header without a valid recording instance. Any follwing archive reads will use the default header and serialization errors might occur..."), ANSI_TO_TCHAR(__FUNCTION__));	
 	}
 
 	return bSuccess;
diff --git a/Engine/Plugins/ChaosVD/Source/ChaosVD/Private/Trace/DataProcessors/ChaosVDCharacterGroundConstraintDataProcessor.cpp b/Engine/Plugins/ChaosVD/Source/ChaosVD/Private/Trace/DataProcessors/ChaosVDCharacterGroundConstraintDataProcessor.cpp
--- a/Engine/Plugins/ChaosVD/Source/ChaosVD/Private/Trace/DataProcessors/ChaosVDCharacterGroundConstraintDataProcessor.cpp
+++ b/Engine/Plugins/ChaosVD/Source/ChaosVD/Private/Trace/DataProcessors/ChaosVDCharacterGroundConstraintDataProcessor.cpp
@@ -26,6 +26,8 @@ bool FChaosVDCharacterGroundConstraintDataProcessor::ProcessRawData(const TArray
 
 	if (bSuccess)
 	{
+		Constraint->SolverID = ProviderSharedPtr->GetRemappedSolverID(Constraint->SolverID);
+
 		if (FChaosVDSolverFrameData* CurrentFrameData = ProviderSharedPtr->GetCurrentSolverFrame(Constraint->SolverID))
 		{
 			CurrentFrameData->RecordedCharacterGroundConstraints.Add(Constraint);
diff --git a/Engine/Plugins/ChaosVD/Source/ChaosVD/Private/Trace/DataProcessors/ChaosVDConstraintDataProcessor.cpp b/Engine/Plugins/ChaosVD/Source/ChaosVD/Private/Trace/DataProcessors/ChaosVDConstraintDataProcessor.cpp
--- a/Engine/Plugins/ChaosVD/Source/ChaosVD/Private/Trace/DataProcessors/ChaosVDConstraintDataProcessor.cpp
+++ b/Engine/Plugins/ChaosVD/Source/ChaosVD/Private/Trace/DataProcessors/ChaosVDConstraintDataProcessor.cpp
@@ -26,7 +26,10 @@ bool FChaosVDConstraintDataProcessor::ProcessRawData(const TArray<uint8>& InData
 
 	if (bSuccess)
 	{
+		RecordedConstraint.SolverID = ProviderSharedPtr->GetRemappedSolverID(RecordedConstraint.SolverID);
+
 		FChaosVDFrameStageData* CurrentSolverStage = ProviderSharedPtr->GetCurrentSolverStageDataForCurrentFrame(RecordedConstraint.SolverID, EChaosVDSolverStageAccessorFlags::None);
+
 		if (ensureMsgf(CurrentSolverStage, TEXT("A MidPhase was traced without a valid step scope")))
 		{
 			AddConstraintToParticleIDMap(RecordedConstraint, RecordedConstraint.Particle0Index, *CurrentSolverStage);
diff --git a/Engine/Plugins/ChaosVD/Source/ChaosVD/Private/Trace/DataProcessors/ChaosVDJointConstraintDataProcessor.cpp b/Engine/Plugins/ChaosVD/Source/ChaosVD/Private/Trace/DataProcessors/ChaosVDJointConstraintDataProcessor.cpp
--- a/Engine/Plugins/ChaosVD/Source/ChaosVD/Private/Trace/DataProcessors/ChaosVDJointConstraintDataProcessor.cpp
+++ b/Engine/Plugins/ChaosVD/Source/ChaosVD/Private/Trace/DataProcessors/ChaosVDJointConstraintDataProcessor.cpp
@@ -26,7 +26,10 @@ bool FChaosVDJointConstraintDataProcessor::ProcessRawData(const TArray<uint8>& I
 
 	if (bSuccess)
 	{
+		JointConstraint->SolverID = ProviderSharedPtr->GetRemappedSolverID(JointConstraint->SolverID);
+
 		FChaosVDFrameStageData* CurrentSolverStage = ProviderSharedPtr->GetCurrentSolverStageDataForCurrentFrame(JointConstraint->SolverID, EChaosVDSolverStageAccessorFlags::None);
+
 		if (ensureMsgf(CurrentSolverStage, TEXT("A Joint Constraint was traced without a valid step scope")))
 		{
 			CurrentSolverStage->RecordedJointConstraints.Add(JointConstraint);
diff --git a/Engine/Plugins/ChaosVD/Source/ChaosVD/Private/Trace/DataProcessors/ChaosVDMidPhaseDataProcessor.cpp b/Engine/Plugins/ChaosVD/Source/ChaosVD/Private/Trace/DataProcessors/ChaosVDMidPhaseDataProcessor.cpp
--- a/Engine/Plugins/ChaosVD/Source/ChaosVD/Private/Trace/DataProcessors/ChaosVDMidPhaseDataProcessor.cpp
+++ b/Engine/Plugins/ChaosVD/Source/ChaosVD/Private/Trace/DataProcessors/ChaosVDMidPhaseDataProcessor.cpp
@@ -26,6 +26,8 @@ bool FChaosVDMidPhaseDataProcessor::ProcessRawData(const TArray<uint8>& InData)
 
 	if (bSuccess)
 	{
+		MidPhase->SolverID = ProviderSharedPtr->GetRemappedSolverID(MidPhase->SolverID);
+
 		FChaosVDFrameStageData* CurrentSolverStage = ProviderSharedPtr->GetCurrentSolverStageDataForCurrentFrame(MidPhase->SolverID, EChaosVDSolverStageAccessorFlags::None);
 
 		if (ensureMsgf(CurrentSolverStage, TEXT("A MidPhase was traced without a valid step scope")))
diff --git a/Engine/Plugins/ChaosVD/Source/ChaosVD/Private/Trace/DataProcessors/ChaosVDSceneQueryDataProcessor.cpp b/Engine/Plugins/ChaosVD/Source/ChaosVD/Private/Trace/DataProcessors/ChaosVDSceneQueryDataProcessor.cpp
--- a/Engine/Plugins/ChaosVD/Source/ChaosVD/Private/Trace/DataProcessors/ChaosVDSceneQueryDataProcessor.cpp
+++ b/Engine/Plugins/ChaosVD/Source/ChaosVD/Private/Trace/DataProcessors/ChaosVDSceneQueryDataProcessor.cpp
@@ -30,6 +30,8 @@ bool FChaosVDSceneQueryDataProcessor::ProcessRawData(const TArray<uint8>& InData
 	{
 		if (const TSharedPtr<FChaosVDGameFrameData> CurrentFrameData = ProviderSharedPtr->GetCurrentGameFrame().Pin())
 		{
+			QueryData->WorldSolverID = ProviderSharedPtr->GetRemappedSolverID(QueryData->WorldSolverID);
+
 			// If ParentQueryID was set, this is a sub query, so find the parent add it to the sub-queries list so we can navigate through the query "hierarchy" later on
 			if (QueryData->ParentQueryID != INDEX_NONE)
 			{
diff --git a/Engine/Plugins/ChaosVD/Source/ChaosVD/Private/Trace/DataProcessors/ChaosVDTraceParticleDataProcessor.cpp b/Engine/Plugins/ChaosVD/Source/ChaosVD/Private/Trace/DataProcessors/ChaosVDTraceParticleDataProcessor.cpp
--- a/Engine/Plugins/ChaosVD/Source/ChaosVD/Private/Trace/DataProcessors/ChaosVDTraceParticleDataProcessor.cpp
+++ b/Engine/Plugins/ChaosVD/Source/ChaosVD/Private/Trace/DataProcessors/ChaosVDTraceParticleDataProcessor.cpp
@@ -28,7 +28,10 @@ bool FChaosVDTraceParticleDataProcessor::ProcessRawData(const TArray<uint8>& InD
 
 	if (bSuccess)
 	{
-		EChaosVDSolverStageAccessorFlags StageAccessorFlags = EChaosVDSolverStageAccessorFlags::CreateNewIfEmpty | EChaosVDSolverStageAccessorFlags::CreateNewIfClosed; 
+		EChaosVDSolverStageAccessorFlags StageAccessorFlags = EChaosVDSolverStageAccessorFlags::CreateNewIfEmpty | EChaosVDSolverStageAccessorFlags::CreateNewIfClosed;
+
+		ParticleData->SolverID = ProviderSharedPtr->GetRemappedSolverID(ParticleData->SolverID);
+
 		if (FChaosVDFrameStageData* CurrentSolverStage = ProviderSharedPtr->GetCurrentSolverStageDataForCurrentFrame(ParticleData->SolverID, StageAccessorFlags))
 		{
 			int32 NewParticleDataIndex = INDEX_NONE;
diff --git a/Engine/Plugins/ChaosVD/Source/ChaosVD/Private/Widgets/SChaosBrowseTraceFileSourceModal.cpp b/Engine/Plugins/ChaosVD/Source/ChaosVD/Private/Widgets/SChaosBrowseTraceFileSourceModal.cpp
--- a/Engine/Plugins/ChaosVD/Source/ChaosVD/Private/Widgets/SChaosBrowseTraceFileSourceModal.cpp
+++ b/Engine/Plugins/ChaosVD/Source/ChaosVD/Private/Widgets/SChaosBrowseTraceFileSourceModal.cpp
@@ -4,6 +4,7 @@
 
 #include "ChaosVDModule.h"
 #include "Editor.h"
+#include "SEnumCombo.h"
 #include "GenericPlatform/GenericPlatformMisc.h"
 #include "Widgets/Layout/SUniformGridPanel.h"
 #include "Widgets/SChaosVDNameListPicker.h"
@@ -12,113 +13,151 @@
 
 #define LOCTEXT_NAMESPACE "ChaosVisualDebugger"
 
-
-void SChaosBrowseTraceFileSourceModal::PopulateLocationNamesList()
-{
-	TArray<TSharedPtr<FName>> NewNameList;
-	NewNameList.Reserve(LocationNameToResponseID.Num());
-	Algo::Transform(LocationNameToResponseID, NewNameList, [](const TPair<FName, EChaosVDBrowseFileModalResponse>& TransformData){ return MakeShared<FName>(TransformData.Key);});
-	NamePickerWidget->UpdateNameList(MoveTemp(NewNameList));
-}
-
-void SChaosBrowseTraceFileSourceModal::RegisterAvailableLocationOptions()
-{
-	// TODO: This FText to FName conversion can be avoided if we make SChaosVDNameListPicker accept FText as well
-	// But doing so it is a larger change than the one we can do now. SChaosVDNameListPicker will need to take the type as a template and there are several places currently using it.
-	// FText.ToString should return the already localized name so we should be good in that area
-
-	CVDFilesLocationNames[0] = LOCTEXT("BrowseLastFolder", "Last Opened Folder");
-	LocationNameToResponseID.Add(FName(CVDFilesLocationNames[0].ToString()), EChaosVDBrowseFileModalResponse::OpenLastFolder);
-	
-	CVDFilesLocationNames[1] = LOCTEXT("BrowseProfilingFolder", "Profiling Folder");
-	LocationNameToResponseID.Add(FName(CVDFilesLocationNames[1].ToString()), EChaosVDBrowseFileModalResponse::OpenProfilingFolder);
-
-	CVDFilesLocationNames[2] = LOCTEXT("BrowseTraceStore", "Trace Store Folder");
-	LocationNameToResponseID.Add(FName(CVDFilesLocationNames[2].ToString()), EChaosVDBrowseFileModalResponse::OpenTraceStore);
-
-	CurrentSelectedLocationName = FName(CVDFilesLocationNames[0].ToString());
-}
-
 void SChaosBrowseTraceFileSourceModal::Construct(const FArguments& InArgs)
 {
 	SWindow::Construct(SWindow::FArguments()
-	.Title(LOCTEXT("SChaosVDBrowseFileModal_Title", "CVD Recording Source Selector"))
+	.Title(LOCTEXT("SChaosVDBrowseFileModal_Title", "Open CVD file"))
 	.SupportsMinimize(false)
 	.SupportsMaximize(false)
 	.UserResizeBorder(0)
-	.ClientSize(FVector2D(350, 80))
+	.ClientSize(FVector2D(350, 140))
 	[
 		SNew(SVerticalBox)
 		+SVerticalBox::Slot()
-		.Padding(15)
+		.Padding(0.0f, 0.0f, 0.0f, 5.0f)
+		.AutoHeight()
 		[
-			SNew(SVerticalBox)
-			+SVerticalBox::Slot()
-			.AutoHeight()
-			.HAlign(HAlign_Center)
+			SNew(SBorder)
+			.BorderImage(FAppStyle::GetBrush("ToolPanel.GroupBorder"))
 			[
-				SNew(STextBlock)
-				.Text(LOCTEXT("SelectSourceMessage", "Where is the recording file located?"))
-				.Font(FCoreStyle::GetDefaultFontStyle("Regular", 10))
-			]
-		]
-		+SVerticalBox::Slot()
-			.AutoHeight()
-			.HAlign(HAlign_Center)
-			.Padding(5)
-			[
-				SNew(SUniformGridPanel)
-				.SlotPadding(FAppStyle::GetMargin("StandardDialog.SlotPadding"))
-				.MinDesiredSlotWidth(FAppStyle::GetFloat("StandardDialog.MinDesiredSlotWidth"))
-				.MinDesiredSlotHeight(FAppStyle::GetFloat("StandardDialog.MinDesiredSlotHeight"))
-				+SUniformGridPanel::Slot(0, 0)
+				SNew(SVerticalBox)
+				+SVerticalBox::Slot()
+				.Padding(5.0f, 10.0f, 5.0f, 5.0f)
+				.AutoHeight()
+				[
+					SNew(SHorizontalBox)
+					+SHorizontalBox::Slot()
+					.HAlign(HAlign_Center)
+					.VAlign(VAlign_Center)
+					[
+						SNew(STextBlock)
+						.Text(LOCTEXT("OpenFileModalInnerTitle", "Select a folder and mode"))
+						.Font(FCoreStyle::GetDefaultFontStyle("Regular", 12))
+					]
+				]
+				+SVerticalBox::Slot()
+				.Padding(5.0f, 0.0f, 5.0f, 5.0f)
+				.AutoHeight()
 				[
-					SAssignNew(NamePickerWidget, SChaosVDNameListPicker)
-					.OnNameSleceted_Raw(this, &SChaosBrowseTraceFileSourceModal::HandleSessionNameSelected)
+					SNew(SHorizontalBox)
+					+SHorizontalBox::Slot()
+					.HAlign(HAlign_Left)
+					.VAlign(VAlign_Center)
+					[
+						SNew(STextBlock)
+						.Text(LOCTEXT("OpenFromFileMode", "Loading Mode"))
+						.Font(FCoreStyle::GetDefaultFontStyle("Regular", 10))
+					]
+					+SHorizontalBox::Slot()
+					.HAlign(HAlign_Right)
+					.VAlign(VAlign_Center)
+					[
+						GenerateConnectionModeWidget().ToSharedRef()
+					]	
 				]
-				+SUniformGridPanel::Slot(1, 0)
+				+SVerticalBox::Slot()
+				.Padding(5.0f, 0.0f, 5.0f, 5.0f)
+				.AutoHeight()
 				[
-					SNew(SButton)
-					.HAlign(HAlign_Center)
-					.ContentPadding(FAppStyle::GetMargin("StandardDialog.ContentPadding"))
-					.Text(LOCTEXT("OpenSelectedLocation", "Open"))
-					.OnClicked(this, &SChaosBrowseTraceFileSourceModal::OnButtonClick)
+					SNew(SHorizontalBox)
+					+SHorizontalBox::Slot()
+					.HAlign(HAlign_Left)
+					.VAlign(VAlign_Center)
+					[
+						SNew(STextBlock)
+						.Text(LOCTEXT("OpenFromFileMode", "Folder"))
+						.Font(FCoreStyle::GetDefaultFontStyle("Regular", 10))
+					]
+					+SHorizontalBox::Slot()
+					.HAlign(HAlign_Right)
+					.VAlign(VAlign_Center)
+					[
+						GenerateSourceFolderWidget().ToSharedRef()
+					]
 				]
 			]
+		]
+		+SVerticalBox::Slot()
+		.AutoHeight()
+		.HAlign(HAlign_Center)
+		.Padding(5)
+		[
+			SNew(SButton)
+			.HAlign(HAlign_Center)
+			.ContentPadding(FAppStyle::GetMargin("StandardDialog.ContentPadding"))
+			.Text(LOCTEXT("OpenSelectedLocation", "Open Folder"))
+			.OnClicked(this, &SChaosBrowseTraceFileSourceModal::OnOpenButtonClick)
+		]
 	]);
-	
-	RegisterAvailableLocationOptions();
+}
+
+TSharedPtr<SWidget> SChaosBrowseTraceFileSourceModal::GenerateConnectionModeWidget()
+{
+	TAttribute<int32> GetCurrentMode;
+	GetCurrentMode.BindSPLambda(this, [this]()
+	{
+		return static_cast<int32>(LoadingMode);
+	});
+
+	SEnumComboBox::FOnEnumSelectionChanged ValueChangedDelegate;
+	ValueChangedDelegate.BindSPLambda(this, [this](int32 NewValue, ESelectInfo::Type)
+	{
+		LoadingMode = static_cast<EChaosVDLoadRecordedDataMode>(NewValue);
+	});
+
+	return SNew(SEnumComboBox, StaticEnum<EChaosVDLoadRecordedDataMode>())
+					.CurrentValue(GetCurrentMode)
+					.OnEnumSelectionChanged(ValueChangedDelegate);
+}
+
+TSharedPtr<SWidget> SChaosBrowseTraceFileSourceModal::GenerateSourceFolderWidget()
+{
+	TAttribute<int32> GetCurrentMode;
+	GetCurrentMode.BindSPLambda(this, [this]()
+	{
+		return static_cast<int32>(UserSelectedResponse);
+	});
+
+	SEnumComboBox::FOnEnumSelectionChanged ValueChangedDelegate;
+	ValueChangedDelegate.BindSPLambda(this, [this](int32 NewValue, ESelectInfo::Type)
+	{
+		UserSelectedResponse = static_cast<EChaosVDBrowseFileModalResponse>(NewValue);
+	});
 
-	PopulateLocationNamesList();
+	return SNew(SEnumComboBox, StaticEnum<EChaosVDBrowseFileModalResponse>())
+					.CurrentValue(GetCurrentMode)
+					.OnEnumSelectionChanged(ValueChangedDelegate);
 }
 
 EChaosVDBrowseFileModalResponse SChaosBrowseTraceFileSourceModal::ShowModal()
 {
 	GEditor->EditorAddModalWindow(SharedThis(this));
-	return UserResponse;
+	
+	return bUserClickedOpen ? UserSelectedResponse : EChaosVDBrowseFileModalResponse::Cancel;
 }
 
-void SChaosBrowseTraceFileSourceModal::HandleSessionNameSelected(TSharedPtr<FName> SelectedName)
+FReply SChaosBrowseTraceFileSourceModal::OnOpenButtonClick()
 {
-	if (SelectedName)
-	{
-		CurrentSelectedLocationName = *SelectedName;
-	}
+	bUserClickedOpen = true;
+	RequestDestroyWindow();
+	return FReply::Handled();
 }
 
-FReply SChaosBrowseTraceFileSourceModal::OnButtonClick()
+FReply SChaosBrowseTraceFileSourceModal::OnCancelButtonClick()
 {
-	if (const EChaosVDBrowseFileModalResponse* SelectedLocationResponse = LocationNameToResponseID.Find(CurrentSelectedLocationName))
-	{
-		UserResponse = *SelectedLocationResponse;
-	}
-	else
-	{
-		UserResponse = EChaosVDBrowseFileModalResponse::OpenLastFolder;
-		UE_LOG(LogChaosVDEditor, Warning, TEXT("[%s] Failed to get selected response [%s] | Defaulting to opening the last opened folder..."), ANSI_TO_TCHAR(__FUNCTION__), *CurrentSelectedLocationName.ToString());
-	}
-	
+	UserSelectedResponse = EChaosVDBrowseFileModalResponse::Cancel;
 	RequestDestroyWindow();
+
 	return FReply::Handled();
 }
 
diff --git a/Engine/Plugins/ChaosVD/Source/ChaosVD/Private/Widgets/SChaosBrowseTraceFileSourceModal.h b/Engine/Plugins/ChaosVD/Source/ChaosVD/Private/Widgets/SChaosBrowseTraceFileSourceModal.h
--- a/Engine/Plugins/ChaosVD/Source/ChaosVD/Private/Widgets/SChaosBrowseTraceFileSourceModal.h
+++ b/Engine/Plugins/ChaosVD/Source/ChaosVD/Private/Widgets/SChaosBrowseTraceFileSourceModal.h
@@ -2,22 +2,32 @@
 
 #pragma once
 
+#include "ChaosVDEngine.h"
 #include "Widgets/SCompoundWidget.h"
 #include "Widgets/SWindow.h"
 
+#include "SChaosBrowseTraceFileSourceModal.generated.h"
+
 class SChaosVDNameListPicker;
 class FReply;
 
+/**
+ * Available default folder locations where CVD files might be located
+ */
+UENUM()
 enum class EChaosVDBrowseFileModalResponse
 {
-	OpenLastFolder,
-	OpenProfilingFolder,
-	OpenTraceStore,
-	Cancel
+	/** Opens the file picker at the last opened folder */
+	LastOpened,
+	/** Opens the file picker at this project profiling folder (Saved/Profiling) */
+	Profiling,
+	/** Opens the file picker at local trace store folder (This is where live recordings are located) */
+	TraceStore,
+	Cancel UMETA(Hidden)
 };
 
 /**
- * 
+ * Simple modal window that allows user pick a file to load and the mode to be used.
  */
 class SChaosBrowseTraceFileSourceModal : public SWindow
 {
@@ -33,16 +43,20 @@ public:
 
 	/** Displays the dialog in a blocking fashion */
 	EChaosVDBrowseFileModalResponse ShowModal();
+	
+	/** Returns the default connection mode to be used when loading the data for the selected session */
+	EChaosVDLoadRecordedDataMode GetSelectedLoadingMode() const
+	{
+		return LoadingMode;
+	}
 
 protected:
 
-	void PopulateLocationNamesList();
-
-	void RegisterAvailableLocationOptions();
-
-	void HandleSessionNameSelected(TSharedPtr<FName> SelectedName);
+	TSharedPtr<SWidget> GenerateConnectionModeWidget();
+	TSharedPtr<SWidget> GenerateSourceFolderWidget();
 	
-	FReply OnButtonClick();
+	FReply OnOpenButtonClick();
+	FReply OnCancelButtonClick();
 
 	TSharedPtr<SChaosVDNameListPicker> NamePickerWidget;
 	
@@ -50,7 +64,9 @@ protected:
 
 	FName CurrentSelectedLocationName;
 
-	EChaosVDBrowseFileModalResponse UserResponse = EChaosVDBrowseFileModalResponse::Cancel;
+	EChaosVDBrowseFileModalResponse UserSelectedResponse = EChaosVDBrowseFileModalResponse::LastOpened;
+
+	EChaosVDLoadRecordedDataMode LoadingMode = EChaosVDLoadRecordedDataMode::SingleSource;
 
-	FText CVDFilesLocationNames[3];
+	bool bUserClickedOpen = false;
 };
diff --git a/Engine/Plugins/ChaosVD/Source/ChaosVD/Private/Widgets/SChaosVDBrowseSessionsModal.cpp b/Engine/Plugins/ChaosVD/Source/ChaosVD/Private/Widgets/SChaosVDBrowseSessionsModal.cpp
--- a/Engine/Plugins/ChaosVD/Source/ChaosVD/Private/Widgets/SChaosVDBrowseSessionsModal.cpp
+++ b/Engine/Plugins/ChaosVD/Source/ChaosVD/Private/Widgets/SChaosVDBrowseSessionsModal.cpp
@@ -4,9 +4,11 @@
 
 #include "ChaosVDModule.h"
 #include "Editor.h"
+#include "SEnumCombo.h"
 #include "Framework/Application/SlateApplication.h"
 #include "GenericPlatform/GenericPlatformMisc.h"
 #include "Trace/StoreClient.h"
+#include "Utils/ChaosVDUserInterfaceUtils.h"
 #include "Widgets/SChaosVDNameListPicker.h"
 #include "Widgets/Layout/SUniformGridPanel.h"
 #include "Widgets/Text/STextBlock.h"
@@ -90,13 +92,15 @@ SChaosVDBrowseSessionsModal::~SChaosVDBrowseSessionsModal()
 void SChaosVDBrowseSessionsModal::Construct(const FArguments& InArgs)
 {
 	CurrentTraceStoreAddress = TEXT("127.0.0.1");
+
+	ConnectionMode = EChaosVDLoadRecordedDataMode::SingleSource;
 	
 	SWindow::Construct(SWindow::FArguments()
 	.Title(LOCTEXT("SChaosVDBrowseSessionsModal_Title", "Live Session Browser"))
 	.SupportsMinimize(false)
 	.SupportsMaximize(false)
 	.UserResizeBorder(0)
-	.ClientSize(FVector2D(450, 120))
+	.ClientSize(FVector2D(450, 145))
 	[
 		SNew(SVerticalBox)
 		+SVerticalBox::Slot()
@@ -130,6 +134,7 @@ void SChaosVDBrowseSessionsModal::Construct(const FArguments& InArgs)
 					]
 				]
 				+SVerticalBox::Slot()
+				.Padding(0.0f, 0.0f, 0.0f, 5.0f)
 				.AutoHeight()
 				[
 					SNew(SHorizontalBox)
@@ -149,6 +154,25 @@ void SChaosVDBrowseSessionsModal::Construct(const FArguments& InArgs)
 						.OnNameSleceted_Raw(this, &SChaosVDBrowseSessionsModal::HandleSessionNameSelected)
 					]	
 				]
+				+SVerticalBox::Slot()
+				.AutoHeight()
+				[
+					SNew(SHorizontalBox)
+					+SHorizontalBox::Slot()
+					.HAlign(HAlign_Left)
+					.VAlign(VAlign_Center)
+					[
+						SNew(STextBlock)
+						.Text(LOCTEXT("LiveConnectionMode", "Connection Mode"))
+						.Font(FCoreStyle::GetDefaultFontStyle("Regular", 12))
+					]
+					+SHorizontalBox::Slot()
+					.HAlign(HAlign_Right)
+					.VAlign(VAlign_Center)
+					[
+						GenerateConnectionModeWidget().ToSharedRef()
+					]	
+				]
 			]
 		]
 		+SVerticalBox::Slot()
@@ -188,6 +212,25 @@ void SChaosVDBrowseSessionsModal::Construct(const FArguments& InArgs)
 	}
 }
 
+TSharedPtr<SWidget> SChaosVDBrowseSessionsModal::GenerateConnectionModeWidget()
+{
+	TAttribute<int32> GetCurrentMode;
+	GetCurrentMode.BindSPLambda(this, [this]()
+	{
+		return static_cast<int32>(ConnectionMode);
+	});
+
+	SEnumComboBox::FOnEnumSelectionChanged ValueChangedDelegate;
+	ValueChangedDelegate.BindSPLambda(this, [this](int32 NewValue, ESelectInfo::Type)
+	{
+		ConnectionMode = static_cast<EChaosVDLoadRecordedDataMode>(NewValue);
+	});
+
+	return SNew(SEnumComboBox, StaticEnum<EChaosVDLoadRecordedDataMode>())
+					.CurrentValue(GetCurrentMode)
+					.OnEnumSelectionChanged(ValueChangedDelegate);
+}
+
 EAppReturnType::Type SChaosVDBrowseSessionsModal::ShowModal()
 {
 	GEditor->EditorAddModalWindow(SharedThis(this));
@@ -254,8 +297,8 @@ void SChaosVDBrowseSessionsModal::UpdateCurrentSessionInfoMap()
 			//TODO: We need a way to know if it is an IPV6
 			FCVDIPv4Address IP(SessionInfo->GetIpAddress());
 			
-			FStringFormatOrderedArguments Args {IP.ToString(), FString::FromInt(SessionInfo->GetId()) };
-			FString SessionDisplayNameString = FString::Format(TEXT("IP: {0} + {1}"), Args);
+			FStringFormatOrderedArguments Args {IP.ToString(), FString::FromInt(SessionInfo->GetTraceId()), FString::FromInt(SessionIndex)};
+			FString SessionDisplayNameString = FString::Format(TEXT("IP: {0} + {1} + {2}"), Args);
 
 			FChaosVDTraceSessionInfo CVDSessionInfo { SessionInfo->GetTraceId(), SessionInfo->GetIpAddress(), SessionInfo->GetControlPort(), true };
 
diff --git a/Engine/Plugins/ChaosVD/Source/ChaosVD/Private/Widgets/SChaosVDBrowseSessionsModal.h b/Engine/Plugins/ChaosVD/Source/ChaosVD/Private/Widgets/SChaosVDBrowseSessionsModal.h
--- a/Engine/Plugins/ChaosVD/Source/ChaosVD/Private/Widgets/SChaosVDBrowseSessionsModal.h
+++ b/Engine/Plugins/ChaosVD/Source/ChaosVD/Private/Widgets/SChaosVDBrowseSessionsModal.h
@@ -2,6 +2,7 @@
 
 #pragma once
 
+#include "ChaosVDEngine.h"
 #include "Containers/Map.h"
 #include "Templates/SharedPointer.h"
 #include "Widgets/SWindow.h"
@@ -43,18 +44,32 @@ public:
 	EAppReturnType::Type ShowModal();
 
 	/** Response that triggered the closing of this modal */
-	EAppReturnType::Type GetResponse() const { return UserResponse; }
+	EAppReturnType::Type GetResponse() const
+	{
+		return UserResponse;
+	}
 
 	/** Returns information about the active trace session selected */
 	FChaosVDTraceSessionInfo GetSelectedTraceInfo();
 
 	/** Returns the address of the trace store selected while looking for active Trace Sessions */
-	FString GetSelectedTraceStoreAddress() const { return CurrentTraceStoreAddress; };
+	FString GetSelectedTraceStoreAddress() const
+	{
+		return CurrentTraceStoreAddress;
+	}
+
+	/** Returns the default connection mode to be used when loading the data for the selected session */
+	EChaosVDLoadRecordedDataMode GetSelectedConnectionMode() const
+	{
+		return ConnectionMode;
+	}
 
 	void ModalTick(float InDeltaTime);
 
 protected:
 
+	TSharedPtr<SWidget> GenerateConnectionModeWidget();
+
 	bool CanOpenSession() const;
 	
 	void UpdateCurrentSessionInfoMap();
@@ -77,6 +92,8 @@ protected:
 
 	FName CurrentTraceSessionSelected;
 
+	EChaosVDLoadRecordedDataMode ConnectionMode = EChaosVDLoadRecordedDataMode::SingleSource;
+
 	FDelegateHandle ModalTickHandle;
 	float AccumulatedTimeBetweenTicks = 0.0f;
 };
diff --git a/Engine/Plugins/ChaosVD/Source/ChaosVD/Private/Widgets/SChaosVDMainTab.cpp b/Engine/Plugins/ChaosVD/Source/ChaosVD/Private/Widgets/SChaosVDMainTab.cpp
--- a/Engine/Plugins/ChaosVD/Source/ChaosVD/Private/Widgets/SChaosVDMainTab.cpp
+++ b/Engine/Plugins/ChaosVD/Source/ChaosVD/Private/Widgets/SChaosVDMainTab.cpp
@@ -414,7 +414,7 @@ void SChaosVDMainTab::GenerateRecentFilesMenu(FMenuBuilder& MenuBuilder)
 		{
 			const FText DisplayName = FText::FromString(FPaths::GetBaseFilename(*RecentFile.FileName));
 			const FText Tooltip = FText::FromString(IFileManager::Get().ConvertToAbsolutePathForExternalAppForRead(*RecentFile.FileName) );
-			MenuBuilder.AddMenuEntry(DisplayName, Tooltip, FSlateIcon(FChaosVDStyle::Get().GetStyleSetName(), "OpenFileIcon"), FUIAction(FExecuteAction::CreateSPLambda(this, [this, FileNameCopy = RecentFile.FileName](){ LoadCVDFile(FileNameCopy); })),
+			MenuBuilder.AddMenuEntry(DisplayName, Tooltip, FSlateIcon(FChaosVDStyle::Get().GetStyleSetName(), "OpenFileIcon"), FUIAction(FExecuteAction::CreateSPLambda(this, [this, FileNameCopy = RecentFile.FileName](){ LoadCVDFile(FileNameCopy, EChaosVDLoadRecordedDataMode::SingleSource); })),
 					NAME_None, EUserInterfaceActionType::Button);
 		}
 	}
@@ -425,19 +425,20 @@ FReply SChaosVDMainTab::BrowseAndOpenChaosVDRecording()
 	const TSharedRef<SChaosBrowseTraceFileSourceModal> SessionBrowserModal = SNew(SChaosBrowseTraceFileSourceModal);
 
 	EChaosVDBrowseFileModalResponse Response = SessionBrowserModal->ShowModal();
+	EChaosVDLoadRecordedDataMode LoadingMode = SessionBrowserModal->GetSelectedLoadingMode();
 	switch(Response)
 	{
-		case EChaosVDBrowseFileModalResponse::OpenLastFolder:
+		case EChaosVDBrowseFileModalResponse::LastOpened:
 			{
-				BrowseChaosVDRecordingFromFolder();
+				BrowseChaosVDRecordingFromFolder(TEXT(""), LoadingMode);
 				break;
 			}
-		case EChaosVDBrowseFileModalResponse::OpenProfilingFolder:
+		case EChaosVDBrowseFileModalResponse::Profiling:
 			{
-				BrowseChaosVDRecordingFromFolder(*FPaths::ProfilingDir());
+				BrowseChaosVDRecordingFromFolder(*FPaths::ProfilingDir(), LoadingMode);
 				break;
 			}
-		case EChaosVDBrowseFileModalResponse::OpenTraceStore:
+		case EChaosVDBrowseFileModalResponse::TraceStore:
 			{
 				//TODO: Support remote Trace Stores
 				const FString TraceStorePath = FChaosVDModule::Get().GetTraceManager()->GetLocalTraceStoreDirPath();
@@ -447,7 +448,7 @@ FReply SChaosVDMainTab::BrowseAndOpenChaosVDRecording()
 					FMessageDialog::Open(EAppMsgType::Ok, LOCTEXT("OpenTraceStoreFailedMessage", "Failed to access the Trace Store, The default profiling folder will be open. \n Please see the logs for mor details... "));
 				}
 
-				BrowseChaosVDRecordingFromFolder(TraceStorePath);
+				BrowseChaosVDRecordingFromFolder(TraceStorePath, LoadingMode);
 				break;
 			}
 		case EChaosVDBrowseFileModalResponse::Cancel:
@@ -523,11 +524,11 @@ TSharedRef<SWidget> SChaosVDMainTab::GenerateSettingsMenuWidget()
 	return UToolMenus::Get()->GenerateWidget(SettingsMenuName, MenuContext);
 }
 
-void SChaosVDMainTab::LoadCVDFile(const FString& InFilename)
+void SChaosVDMainTab::LoadCVDFile(const FString& InFilename, EChaosVDLoadRecordedDataMode LoadingMode)
 {
 	if (ensure(InFilename.EndsWith(TEXT("utrace"))))
 	{
-		GetChaosVDEngineInstance()->LoadRecording(InFilename);
+		GetChaosVDEngineInstance()->LoadRecording(InFilename, LoadingMode);
 	}
 	else
 	{
@@ -651,7 +652,7 @@ void SChaosVDMainTab::SetCustomPropertyLayouts(IDetailsView* DetailsView)
 	});
 }
 
-void SChaosVDMainTab::BrowseChaosVDRecordingFromFolder(FStringView FolderPath)
+void SChaosVDMainTab::BrowseChaosVDRecordingFromFolder(FStringView FolderPath, EChaosVDLoadRecordedDataMode LoadingMode)
 {
 	TArray<FString> OutOpenFilenames;
 	if (IDesktopPlatform* DesktopPlatform = FDesktopPlatformModule::Get())
@@ -674,30 +675,13 @@ void SChaosVDMainTab::BrowseChaosVDRecordingFromFolder(FStringView FolderPath)
 
 	if (OutOpenFilenames.Num() > 0)
 	{
-		LoadCVDFile(OutOpenFilenames[0]);
+		LoadCVDFile(OutOpenFilenames[0], LoadingMode);
 	}
 }
 
-bool SChaosVDMainTab::ConnectToLiveSession(int32 SessionID, const FString SessionAddress) const
+bool SChaosVDMainTab::ConnectToLiveSession(int32 SessionID, const FString& InSessionAddress, EChaosVDLoadRecordedDataMode LoadingMode) const
 {
-	FChaosVDTraceSessionDescriptor NewSessionFromFileDescriptor;
-	NewSessionFromFileDescriptor.SessionName = FChaosVDModule::Get().GetTraceManager()->ConnectToLiveSession(SessionAddress, SessionID);
-	NewSessionFromFileDescriptor.bIsLiveSession = true;
-
-	bool bSuccess = false;
-
-	if (!NewSessionFromFileDescriptor.IsValid())
-	{
-		// If it failed we want to clean the current session name, so it is ok calling it either way
-		GetChaosVDEngineInstance()->SetCurrentSession(FChaosVDTraceSessionDescriptor());
-	}
-	else
-	{
-		GetChaosVDEngineInstance()->SetCurrentSession(NewSessionFromFileDescriptor);
-		bSuccess = true;
-	}
-
-	return bSuccess;
+	return GetChaosVDEngineInstance()->ConnectToLiveSession(SessionID, InSessionAddress, LoadingMode);
 }
 
 void SChaosVDMainTab::SetUpDisableCPUThrottlingDelegate()
@@ -774,41 +758,89 @@ void SChaosVDMainTab::RegisterMainTabMenu()
 		}
 
 		TFunction<FText()> GetTextDelegate = [WeakTab = MainTabPtr.ToWeakPtr()]()
+		{
+			return LOCTEXT("ConnectToSession", "Connect to Session");
+		};
+
+		TFunction<FText()> GetTooltipTextDelegate = [WeakTab = MainTabPtr.ToWeakPtr()]()
+		{
+			return LOCTEXT("ConnectToSessionTooltip", "Opens a panel where you can browse active live sessions and connect to one");
+		};
+		
+		FOnClicked OnClickedDelegate = FOnClicked::CreateLambda([WeakTab = StaticCastWeakPtr<SChaosVDMainTab>(MainTabPtr->AsWeak())]()
 		{
 			if (TSharedPtr<SChaosVDMainTab> TabPtr = WeakTab.Pin())
 			{
-				return TabPtr->GetConnectButtonText();
+				return TabPtr->HandleSessionConnectionClicked();
 			}
 
-			return FText();
-		};
+			return FReply::Handled();
+		});
 
-		TFunction<FText()> GetTooltipTextDelegate = [WeakTab = MainTabPtr.ToWeakPtr()]()
+		TSharedRef<SButton> ConnectToSessionButton = MainTabPtr->CreateSimpleButton(MoveTemp(GetTextDelegate), MoveTemp(GetTooltipTextDelegate), FChaosVDStyle::Get().GetBrush("OpenSessionIcon"), Context, MoveTemp(OnClickedDelegate));
+
+		InSection.AddEntry(
+			FToolMenuEntry::InitWidget(
+				"ConnectToSession",
+				ConnectToSessionButton,
+				FText::GetEmpty(),
+				true,
+				false
+			));
+	}));
+
+	Section.AddDynamicEntry("DisconnectFromSession", FNewToolMenuSectionDelegate::CreateLambda([](FToolMenuSection& InSection)
+	{
+		const UChaosVDMainToolbarMenuContext* Context = InSection.FindContext<UChaosVDMainToolbarMenuContext>();
+		TSharedPtr<SChaosVDMainTab> MainTabPtr = Context->MainTab.Pin();
+		if (!MainTabPtr)
+		{
+			return;
+		}
+
+		TFunction<FText()> GetTextDelegate = [WeakTab = MainTabPtr.ToWeakPtr()]()
 		{
 			if (TSharedPtr<SChaosVDMainTab> TabPtr = WeakTab.Pin())
 			{
-				return TabPtr->GetConnectButtonTooltipText();
+				return TabPtr->GetDisconnectButtonText();
 			}
 
 			return FText();
 		};
-		
+
+		TFunction<FText()> GetTooltipTextDelegate = [WeakTab = MainTabPtr.ToWeakPtr()]()
+		{
+			return LOCTEXT("DisconnectFromSessionTooltip", "Disconnects from the current live session but it does not stop it");
+		};
+			
 		FOnClicked OnClickedDelegate = FOnClicked::CreateLambda([WeakTab = StaticCastWeakPtr<SChaosVDMainTab>(MainTabPtr->AsWeak())]()
 		{
 			if (TSharedPtr<SChaosVDMainTab> TabPtr = WeakTab.Pin())
 			{
-				return TabPtr->HandleSessionConnectionClicked();
+				return TabPtr->HandleDisconnectSessionClicked();
 			}
 
 			return FReply::Handled();
 		});
 
-		TSharedRef<SButton> ConnectToSessionButton = MainTabPtr->CreateSimpleButton(MoveTemp(GetTextDelegate), MoveTemp(GetTooltipTextDelegate), FChaosVDStyle::Get().GetBrush("OpenSessionIcon"), Context, MoveTemp(OnClickedDelegate));
+		TAttribute<EVisibility> ButtonVisibilityAttribute;
+		ButtonVisibilityAttribute.BindLambda([WeakTab = MainTabPtr.ToWeakPtr()]
+		{
+			if (TSharedPtr<SChaosVDMainTab> TabPtr = WeakTab.Pin())
+			{
+				return TabPtr->GetChaosVDEngineInstance()->HasAnyLiveSessionActive() ? EVisibility::Visible : EVisibility::Collapsed;
+			}
+			return EVisibility::Collapsed;
+		});
+
+		TSharedRef<SButton> DisconnectFromSessionButton = MainTabPtr->CreateSimpleButton(MoveTemp(GetTextDelegate), MoveTemp(GetTooltipTextDelegate), FChaosVDStyle::Get().GetBrush("OpenSessionIcon"), Context, MoveTemp(OnClickedDelegate));
+
+		DisconnectFromSessionButton->SetVisibility(ButtonVisibilityAttribute);
 
 		InSection.AddEntry(
 			FToolMenuEntry::InitWidget(
-				"ConnectToSession",
-				ConnectToSessionButton,
+				"DisconnectFromSessionButton",
+				DisconnectFromSessionButton,
 				FText::GetEmpty(),
 				true,
 				false
@@ -893,7 +925,9 @@ void SChaosVDMainTab::BrowseLiveSessionsFromTraceStore() const
 		if (SessionInfo.bIsValid)
 		{
 			const FString SessionAddress = SessionBrowserModal->GetSelectedTraceStoreAddress();
-			bSuccess = ConnectToLiveSession(SessionInfo.TraceID, SessionAddress);
+			const EChaosVDLoadRecordedDataMode ConnectionMode = SessionBrowserModal->GetSelectedConnectionMode();
+
+			bSuccess = ConnectToLiveSession(SessionInfo.TraceID, SessionAddress, ConnectionMode);
 		}
 
 		if (!bSuccess)
@@ -906,10 +940,16 @@ void SChaosVDMainTab::BrowseLiveSessionsFromTraceStore() const
 bool SChaosVDMainTab::ShouldDisableCPUThrottling() const
 {
 	// If we are playing a live session, it is likely the editor will be in the background, so we need to disable CPU Throttling
-	return ChaosVDEngine && ChaosVDEngine->GetCurrentSessionDescriptor().bIsLiveSession;
+	return ChaosVDEngine && ChaosVDEngine->HasAnyLiveSessionActive();
 }
 
 FReply SChaosVDMainTab::HandleSessionConnectionClicked()
+{
+	BrowseLiveSessionsFromTraceStore();
+	return FReply::Handled();
+}
+
+FReply SChaosVDMainTab::HandleDisconnectSessionClicked()
 {
 	if (TSharedPtr<FChaosVDPlaybackController> PlaybackControllerPtr =  GetChaosVDEngineInstance()->GetPlaybackController())
 	{
@@ -917,38 +957,31 @@ FReply SChaosVDMainTab::HandleSessionConnectionClicked()
 	
 		if (bIsAlreadyInLiveSession)
 		{
-			FChaosVDModule::Get().GetTraceManager()->CloseSession(GetChaosVDEngineInstance()->GetCurrentSessionDescriptor().SessionName);
+			TArrayView<FChaosVDTraceSessionDescriptor> ActiveSessions = GetChaosVDEngineInstance()->GetCurrentSessionDescriptors();
+			for (FChaosVDTraceSessionDescriptor& ActiveSession : ActiveSessions)
+			{
+				FChaosVDModule::Get().GetTraceManager()->CloseSession(ActiveSession.SessionName);
+				ActiveSession.bIsLiveSession = false;
+			}
+
 			PlaybackControllerPtr->HandleDisconnectedFromSession();
 		}
-		else
-		{
-			BrowseLiveSessionsFromTraceStore();
-		}
 	}
 
 	return FReply::Handled();
 }
 
-FText SChaosVDMainTab::GetConnectButtonText() const
+FText SChaosVDMainTab::GetDisconnectButtonText() const
 {
-	TSharedPtr<FChaosVDPlaybackController> PlaybackControllerPtr = GetChaosVDEngineInstance()->GetPlaybackController();
-	if (PlaybackControllerPtr.IsValid())
+	if (GetChaosVDEngineInstance()->HasAnyLiveSessionActive())
 	{
-		const bool bIsAlreadyInLiveSession = PlaybackControllerPtr->IsPlayingLiveSession();
-		return bIsAlreadyInLiveSession ? LOCTEXT("DisconnectFromSession", "Disconnect from Session") : LOCTEXT("ConnectToSession", "Connect to Session");
+		if (GetChaosVDEngineInstance()->GetCurrentSessionDescriptors().Num() > 1)
+		{
+			return LOCTEXT("DisconnectFromMultipleSessions", "Disconnect from all Sessions");
+		}
 	}
-	return FText();
-}
 
-FText SChaosVDMainTab::GetConnectButtonTooltipText() const
-{
-	TSharedPtr<FChaosVDPlaybackController> PlaybackControllerPtr = GetChaosVDEngineInstance()->GetPlaybackController();
-	if (PlaybackControllerPtr.IsValid())
-	{
-		const bool bIsAlreadyInLiveSession = PlaybackControllerPtr->IsPlayingLiveSession();
-		return bIsAlreadyInLiveSession ? LOCTEXT("DisconnectFromSessionTooltip", "Disconnects from the current live session but it does not stop it") : LOCTEXT("ConnectToSessionTooltip", "Opens a panel where you can browse active live sessions and connect to one");
-	}
-	return FText();
+	return LOCTEXT("DisconnectFromSession", "Disconnect from Session");
 }
 
 #undef LOCTEXT_NAMESPACE
diff --git a/Engine/Plugins/ChaosVD/Source/ChaosVD/Private/Widgets/SChaosVDRecordedLogBrowser.cpp b/Engine/Plugins/ChaosVD/Source/ChaosVD/Private/Widgets/SChaosVDRecordedLogBrowser.cpp
--- a/Engine/Plugins/ChaosVD/Source/ChaosVD/Private/Widgets/SChaosVDRecordedLogBrowser.cpp
+++ b/Engine/Plugins/ChaosVD/Source/ChaosVD/Private/Widgets/SChaosVDRecordedLogBrowser.cpp
@@ -145,7 +145,17 @@ void SChaosVDRecordedLogBrowser::UpdateBrowserContents()
 		return;	
 	}
 
-	TSharedPtr<const TraceServices::IAnalysisSession> Session = FChaosVDModule::Get().GetTraceManager()->GetSession(CVDEngineInstance->GetCurrentSessionDescriptor().SessionName);
+	TSharedPtr<const TraceServices::IAnalysisSession> Session = nullptr;
+
+	// TODO: Add support for multiple sessions instead of picking the first one.
+	// The easiest way to do it is have a log window per trace session open, and have the log browser receive the session name when it is opened
+
+	TArrayView<FChaosVDTraceSessionDescriptor> ActiveSessions = CVDEngineInstance->GetCurrentSessionDescriptors();
+	if (!ActiveSessions.IsEmpty())
+	{
+		Session = FChaosVDModule::Get().GetTraceManager()->GetSession(ActiveSessions[0].SessionName);
+	}
+	
 	if (!Session)
 	{
 		return;
@@ -435,8 +445,8 @@ void SChaosVDRecordedLogBrowser::HandleItemFocused(const TSharedPtr<FChaosVDLogV
 	}
 
 	TSharedPtr<FChaosVDEngine> CVDEngineInstance = EngineInstanceWeakPtr.Pin();
-	TSharedPtr<const TraceServices::IAnalysisSession> Session = CVDEngineInstance ?
-																FChaosVDModule::Get().GetTraceManager()->GetSession(CVDEngineInstance->GetCurrentSessionDescriptor().SessionName) :
+	TSharedPtr<const TraceServices::IAnalysisSession> Session = /*CVDEngineInstance ?
+																FChaosVDModule::Get().GetTraceManager()->GetSession(CVDEngineInstance->GetCurrentSessionDescriptor().SessionName) :*/
 																nullptr;
 
 	if (!Session)
@@ -588,8 +598,8 @@ void SChaosVDRecordedLogBrowser::ProcessPendingTrackUpdates()
 
 	TSharedPtr<FChaosVDEngine> CVDEngineInstance = EngineInstanceWeakPtr.Pin();
 	TSharedPtr<FChaosVDTraceManager> TraceManager = FChaosVDModule::Get().GetTraceManager();
-	TSharedPtr<const TraceServices::IAnalysisSession> Session = CVDEngineInstance && TraceManager ?
-																TraceManager->GetSession(CVDEngineInstance->GetCurrentSessionDescriptor().SessionName) :
+	TSharedPtr<const TraceServices::IAnalysisSession> Session = /*CVDEngineInstance && TraceManager ?
+																TraceManager->GetSession(CVDEngineInstance->GetCurrentSessionDescriptor().SessionName) :*/
 																nullptr;
 	if (!Session)
 	{
diff --git a/Engine/Plugins/ChaosVD/Source/ChaosVD/Private/Widgets/SChaosVDRecordingControls.cpp b/Engine/Plugins/ChaosVD/Source/ChaosVD/Private/Widgets/SChaosVDRecordingControls.cpp
--- a/Engine/Plugins/ChaosVD/Source/ChaosVD/Private/Widgets/SChaosVDRecordingControls.cpp
+++ b/Engine/Plugins/ChaosVD/Source/ChaosVD/Private/Widgets/SChaosVDRecordingControls.cpp
@@ -249,7 +249,7 @@ void SChaosVDRecordingControls::HandleRecordingStop()
 		return;
 	}
 
-	const bool bIsLiveSession = MainTabSharedPtr->GetChaosVDEngineInstance()->GetCurrentSessionDescriptor().bIsLiveSession;
+	const bool bIsLiveSession = MainTabSharedPtr->GetChaosVDEngineInstance()->HasAnyLiveSessionActive();
 
 	if (UStatusBarSubsystem* StatusBarSubsystem = GEditor ? GEditor->GetEditorSubsystem<UStatusBarSubsystem>() : nullptr)
 	{
@@ -333,7 +333,7 @@ void SChaosVDRecordingControls::AttemptToConnectToLiveSession()
 				// CVD needs the trace session name to be able to load a live session. Although the session exist, the session name might not be written right away
 				// Trace files don't really have metadata, it is all part of the same stream, so we need to wait until it is written which might take a few ticks.
 				// Therefore if it is not ready, try again a few times.
-				if (!MainTabSharedPtr->ConnectToLiveSession(SessionID, SessionAddress))
+				if (!MainTabSharedPtr->ConnectToLiveSession(SessionID, SessionAddress, EChaosVDLoadRecordedDataMode::SingleSource))
 				{
 					if (RecordingControlsPtr->CurrentConnectionAttempts <= RecordingControlsPtr->MaxAutoplayConnectionAttempts)
 					{
@@ -408,7 +408,8 @@ bool SChaosVDRecordingControls::IsRecordingToggleButtonEnabled(EChaosVDRecording
 		return false;
 
 	}
-	const bool bIsLiveSession = MainTabSharedPtr->GetChaosVDEngineInstance()->GetCurrentSessionDescriptor().bIsLiveSession;
+
+	const bool bIsLiveSession = MainTabSharedPtr->GetChaosVDEngineInstance()->HasAnyLiveSessionActive();
 
 	const bool bIsRecording = IsRecording();
 
diff --git a/Engine/Plugins/ChaosVD/Source/ChaosVD/Private/Widgets/SChaosVDSolverTracks.cpp b/Engine/Plugins/ChaosVD/Source/ChaosVD/Private/Widgets/SChaosVDSolverTracks.cpp
--- a/Engine/Plugins/ChaosVD/Source/ChaosVD/Private/Widgets/SChaosVDSolverTracks.cpp
+++ b/Engine/Plugins/ChaosVD/Source/ChaosVD/Private/Widgets/SChaosVDSolverTracks.cpp
@@ -108,7 +108,15 @@ void SChaosVDSolverTracks::UpdatedCachedTrackInfoData(TWeakPtr<FChaosVDPlaybackC
 	if (const TSharedPtr<FChaosVDPlaybackController> CurrentPlaybackControllerPtr = InPlaybackController.Pin())
 	{
 		TArray<TSharedPtr<const FChaosVDTrackInfo>> TrackInfoArray;
-		CurrentPlaybackControllerPtr->GetAvailableTrackInfosAtTrackFrame(EChaosVDTrackType::Solver, UpdatedTrackInfo, TrackInfoArray);
+
+		if (CurrentPlaybackControllerPtr->GetTimelineSyncMode() == EChaosVDSyncTimelinesMode::Manual)
+		{
+			CurrentPlaybackControllerPtr->GetAvailableTracks(EChaosVDTrackType::Solver,TrackInfoArray);
+		}
+		else
+		{
+			CurrentPlaybackControllerPtr->GetAvailableTrackInfosAtTrackFrame(EChaosVDTrackType::Solver, UpdatedTrackInfo, TrackInfoArray);
+		}
 
 		if (TrackInfoArray != CachedTrackInfoArray)
 		{
diff --git a/Engine/Plugins/ChaosVD/Source/ChaosVD/Public/ChaosVDPlaybackController.h b/Engine/Plugins/ChaosVD/Source/ChaosVD/Public/ChaosVDPlaybackController.h
--- a/Engine/Plugins/ChaosVD/Source/ChaosVD/Public/ChaosVDPlaybackController.h
+++ b/Engine/Plugins/ChaosVD/Source/ChaosVD/Public/ChaosVDPlaybackController.h
@@ -82,12 +82,19 @@ enum class EChaosVDUnloadRecordingFlags : uint8
 };
 ENUM_CLASS_FLAGS(EChaosVDUnloadRecordingFlags)
 
+/** Available sync modes that determine how tracks will sync between each other during playback */
 UENUM()
 enum class EChaosVDSyncTimelinesMode : uint8
 {
 	None UMETA(Hidden),
+	/** Syncs all tracks using the recorded timestamp */
 	RecordedTimestamp,
-	NetworkTick
+	/** Syncs all tracks using the recorded network ticks offset of the Predictive networked physics system*/
+	NetworkTick,
+	/** No-Auto sync is performed and all available solver tracks inside the recording are visible at once and can be scrubbed
+	 * independently
+	 */
+	Manual
 };
 
 typedef TMap<int32, TSharedPtr<FChaosVDTrackInfo>> TrackInfoByIDMap;
@@ -215,7 +222,8 @@ public:
 	 * @param OutTrackInfo Array where any found track info data will be added
 	 */
 
-	void GetAvailableTracks(EChaosVDTrackType TrackType, TArray<TSharedPtr<FChaosVDTrackInfo>>& OutTrackInfo);
+	void GetAvailableTracks(EChaosVDTrackType TrackType, TArray<TSharedPtr<const FChaosVDTrackInfo>>& OutTrackInfo);
+	void GetMutableAvailableTracks(EChaosVDTrackType TrackType, TArray<TSharedPtr<FChaosVDTrackInfo>>& OutTrackInfo);
 	
 	/**
 	 * Gets all the ids of the tracks, of the specified type, that are available available on the loaded recording, at a specified frame
diff --git a/Engine/Plugins/ChaosVD/Source/ChaosVD/Public/Trace/ChaosVDTraceProvider.h b/Engine/Plugins/ChaosVD/Source/ChaosVD/Public/Trace/ChaosVDTraceProvider.h
--- a/Engine/Plugins/ChaosVD/Source/ChaosVD/Public/Trace/ChaosVDTraceProvider.h
+++ b/Engine/Plugins/ChaosVD/Source/ChaosVD/Public/Trace/ChaosVDTraceProvider.h
@@ -35,7 +35,7 @@ ENUM_CLASS_FLAGS(EChaosVDSolverStageAccessorFlags);
  * It stores and handles rebuilt recorded frame data from Trace events
  * dispatched by the Chaos VD Trace analyzer
  */
-class FChaosVDTraceProvider : public TraceServices::IProvider, public TSharedFromThis<FChaosVDTraceProvider>
+class FChaosVDTraceProvider : public TraceServices::IProvider, public TraceServices::IEditableProvider, public TSharedFromThis<FChaosVDTraceProvider>
 {
 public:
 	
@@ -48,6 +48,8 @@ public:
 	 */
 	void CreateRecordingInstanceForSession(const FString& InSessionName);
 
+	void SetExternalRecordingInstanceForSession(const TSharedRef<FChaosVDRecording>& InExternalCVDRecording);
+
 	/** Opens a solver frame entry into the active CVD recording structure
 	 * @param InSolverGUID ID of the solver
 	 * @param FrameData Solver frame data
@@ -123,7 +125,17 @@ private:
 	void GetAvailablePendingSolverIDsAtGameFrame(const TSharedRef<FChaosVDGameFrameData>& InProcessedGameFrameData, TArray<int32, TInlineAllocator<16>>& OutSolverIDs);
 
 	FString GenerateFormattedStringListFromSet(const TSet<FString>& StringsSet);
+
+	int32 RemapSolverID(int32 SolverID);
+
+public:
+	virtual void BeginEdit() const override{};
+	virtual void EndEdit() const override{};
+	virtual void EditAccessCheck() const override{};
 	
+	int32 CHAOSVD_API GetRemappedSolverID(int32 SolverID);
+
+private:
 	TraceServices::IAnalysisSession& Session;
 
 	TSharedPtr<FChaosVDRecording> InternalRecording;
@@ -146,10 +158,14 @@ private:
 
 	double StartLastCommitedFrameTimeSeconds = 0.0;
 
+	bool bHasRecordingOverride = false;
+
 	TMap<int32, int32> CurrentNetworkTickOffsets;
+	TSortedMap<int32, int32> RemappedSolversID;
 
 	TSet<FString> MissingDataProcessors;
 	TSet<FString> TypesFailedToSerialize;
 
 	friend class FChaosVDTraceAnalyzer;
 };
+
diff --git a/Engine/Plugins/ChaosVD/Source/ChaosVD/Public/Widgets/SChaosVDMainTab.h b/Engine/Plugins/ChaosVD/Source/ChaosVD/Public/Widgets/SChaosVDMainTab.h
--- a/Engine/Plugins/ChaosVD/Source/ChaosVD/Public/Widgets/SChaosVDMainTab.h
+++ b/Engine/Plugins/ChaosVD/Source/ChaosVD/Public/Widgets/SChaosVDMainTab.h
@@ -33,6 +33,8 @@ class SDockTab;
 struct FDetailsViewArgs;
 struct FStructureDetailsViewArgs;
 
+enum class EChaosVDLoadRecordedDataMode : uint8;
+
 UCLASS()
 class CHAOSVD_API UChaosVDMainToolbarMenuContext : public UObject
 {
@@ -99,9 +101,9 @@ public:
 		return ComponentVisualizers;
 	}
 	
-	bool ConnectToLiveSession(int32 SessionID, FString SessionAddress) const;
+	bool ConnectToLiveSession(int32 SessionID, const FString& InSessionAddress, EChaosVDLoadRecordedDataMode LoadingMode) const;
 
-	void LoadCVDFile(const FString& InFilename);
+	void LoadCVDFile(const FString& InFilename, EChaosVDLoadRecordedDataMode LoadingMode);
 
 	TSharedRef<IDetailsView> CreateDetailsView(const FDetailsViewArgs& InDetailsViewArgs);
 
@@ -157,7 +159,7 @@ private:
 	TSharedRef<SWidget> GenerateMainToolbarWidget();
 	TSharedRef<SWidget> GenerateSettingsMenuWidget();
 
-	void BrowseChaosVDRecordingFromFolder(FStringView FolderPath = TEXT(""));
+	void BrowseChaosVDRecordingFromFolder(FStringView FolderPath, EChaosVDLoadRecordedDataMode LoadingMode);
 
 	void BrowseLiveSessionsFromTraceStore() const;
 
@@ -185,8 +187,8 @@ private:
 	FDelegateHandle DisableCPUThrottleHandle;
 
 	FReply HandleSessionConnectionClicked();
-	FText GetConnectButtonText() const;
-	FText GetConnectButtonTooltipText() const;
+	FReply HandleDisconnectSessionClicked();
+	FText GetDisconnectButtonText() const;
 
 	bool bShowTrackSelectorKeyShortcut = false;
 
diff --git a/Engine/Plugins/ChaosVD/Source/ChaosVDBuiltInExtensions/Private/AccelerationStructures/Trace/ChaosVDAABBTreeDataProcessor.cpp b/Engine/Plugins/ChaosVD/Source/ChaosVDBuiltInExtensions/Private/AccelerationStructures/Trace/ChaosVDAABBTreeDataProcessor.cpp
--- a/Engine/Plugins/ChaosVD/Source/ChaosVDBuiltInExtensions/Private/AccelerationStructures/Trace/ChaosVDAABBTreeDataProcessor.cpp
+++ b/Engine/Plugins/ChaosVD/Source/ChaosVDBuiltInExtensions/Private/AccelerationStructures/Trace/ChaosVDAABBTreeDataProcessor.cpp
@@ -53,6 +53,8 @@ bool FChaosVDAABBTreeDataProcessor::ProcessRawData(const TArray<uint8>& InData)
 	{
 		if (TSharedPtr<FChaosVDGameFrameData> CurrentFrameData = ProviderSharedPtr->GetCurrentGameFrame().Pin())
 		{
+			AABBTreeData->SolverId = ProviderSharedPtr->GetRemappedSolverID(AABBTreeData->SolverId);
+
 			AABBTreeData->TreeDepth = Chaos::VisualDebugger::Utils::CalculateAABBTreeDepth(AABBTreeData.ToSharedRef(), AABBTreeData->GetCorrectedRootNodeIndex());
 
 			if (TSharedPtr<FChaosVDAccelerationStructureContainer> AABBTreeDataContainer = CurrentFrameData->GetCustomDataHandler().GetOrAddDefaultData<FChaosVDAccelerationStructureContainer>())
diff --git a/Engine/Plugins/ChaosVD/Source/ChaosVDBuiltInExtensions/Private/GenericDebugDraw/Trace/DataProcessors/ChaosVDDebugDrawBoxDataProcessor.cpp b/Engine/Plugins/ChaosVD/Source/ChaosVDBuiltInExtensions/Private/GenericDebugDraw/Trace/DataProcessors/ChaosVDDebugDrawBoxDataProcessor.cpp
--- a/Engine/Plugins/ChaosVD/Source/ChaosVDBuiltInExtensions/Private/GenericDebugDraw/Trace/DataProcessors/ChaosVDDebugDrawBoxDataProcessor.cpp
+++ b/Engine/Plugins/ChaosVD/Source/ChaosVDBuiltInExtensions/Private/GenericDebugDraw/Trace/DataProcessors/ChaosVDDebugDrawBoxDataProcessor.cpp
@@ -26,6 +26,8 @@ bool FChaosVDDebugDrawBoxDataProcessor::ProcessRawData(const TArray<uint8>& InDa
 	{
 		if (TSharedPtr<FChaosVDGameFrameData> CurrentFrameData = ProviderSharedPtr->GetCurrentGameFrame().Pin())
 		{
+			DebugDrawData->SolverID = ProviderSharedPtr->GetRemappedSolverID(DebugDrawData->SolverID);
+
 			if (TSharedPtr<FChaosVDDebugShapeDataContainer> DebugDrawShapeData = CurrentFrameData->GetCustomDataHandler().GetOrAddDefaultData<FChaosVDDebugShapeDataContainer>())
 			{
 				DebugDrawShapeData->RecordedDebugDrawBoxesBySolverID.FindOrAdd(DebugDrawData->SolverID).Add(DebugDrawData);
diff --git a/Engine/Plugins/ChaosVD/Source/ChaosVDBuiltInExtensions/Private/GenericDebugDraw/Trace/DataProcessors/ChaosVDDebugDrawImplicitObjectDataProcessor.cpp b/Engine/Plugins/ChaosVD/Source/ChaosVDBuiltInExtensions/Private/GenericDebugDraw/Trace/DataProcessors/ChaosVDDebugDrawImplicitObjectDataProcessor.cpp
--- a/Engine/Plugins/ChaosVD/Source/ChaosVDBuiltInExtensions/Private/GenericDebugDraw/Trace/DataProcessors/ChaosVDDebugDrawImplicitObjectDataProcessor.cpp
+++ b/Engine/Plugins/ChaosVD/Source/ChaosVDBuiltInExtensions/Private/GenericDebugDraw/Trace/DataProcessors/ChaosVDDebugDrawImplicitObjectDataProcessor.cpp
@@ -26,6 +26,8 @@ bool FChaosVDDebugDrawImplicitObjectDataProcessor::ProcessRawData(const TArray<u
 	{
 		if (TSharedPtr<FChaosVDGameFrameData> CurrentFrameData = ProviderSharedPtr->GetCurrentGameFrame().Pin())
 		{
+			DebugDrawData->SolverID = ProviderSharedPtr->GetRemappedSolverID(DebugDrawData->SolverID);
+
 			if (TSharedPtr<FChaosVDDebugShapeDataContainer> DebugDrawShapeData = CurrentFrameData->GetCustomDataHandler().GetOrAddDefaultData<FChaosVDDebugShapeDataContainer>())
 			{	
 				DebugDrawShapeData->RecordedDebugDrawImplicitObjectsBySolverID.FindOrAdd(DebugDrawData->SolverID).Add(DebugDrawData);
diff --git a/Engine/Plugins/ChaosVD/Source/ChaosVDBuiltInExtensions/Private/GenericDebugDraw/Trace/DataProcessors/ChaosVDDebugDrawLineDataProcessor.cpp b/Engine/Plugins/ChaosVD/Source/ChaosVDBuiltInExtensions/Private/GenericDebugDraw/Trace/DataProcessors/ChaosVDDebugDrawLineDataProcessor.cpp
--- a/Engine/Plugins/ChaosVD/Source/ChaosVDBuiltInExtensions/Private/GenericDebugDraw/Trace/DataProcessors/ChaosVDDebugDrawLineDataProcessor.cpp
+++ b/Engine/Plugins/ChaosVD/Source/ChaosVDBuiltInExtensions/Private/GenericDebugDraw/Trace/DataProcessors/ChaosVDDebugDrawLineDataProcessor.cpp
@@ -26,6 +26,8 @@ bool FChaosVDDebugDrawLineDataProcessor::ProcessRawData(const TArray<uint8>& InD
 	{
 		if (TSharedPtr<FChaosVDGameFrameData> CurrentFrameData = ProviderSharedPtr->GetCurrentGameFrame().Pin())
 		{
+			DebugDrawData->SolverID = ProviderSharedPtr->GetRemappedSolverID(DebugDrawData->SolverID);
+
 			if (TSharedPtr<FChaosVDDebugShapeDataContainer> DebugDrawShapeData = CurrentFrameData->GetCustomDataHandler().GetOrAddDefaultData<FChaosVDDebugShapeDataContainer>())
 			{
 				DebugDrawShapeData->RecordedDebugDrawLinesBySolverID.FindOrAdd(DebugDrawData->SolverID).Add(DebugDrawData);
diff --git a/Engine/Plugins/ChaosVD/Source/ChaosVDBuiltInExtensions/Private/GenericDebugDraw/Trace/DataProcessors/ChaosVDDebugDrawSphereDataProcessor.cpp b/Engine/Plugins/ChaosVD/Source/ChaosVDBuiltInExtensions/Private/GenericDebugDraw/Trace/DataProcessors/ChaosVDDebugDrawSphereDataProcessor.cpp
--- a/Engine/Plugins/ChaosVD/Source/ChaosVDBuiltInExtensions/Private/GenericDebugDraw/Trace/DataProcessors/ChaosVDDebugDrawSphereDataProcessor.cpp
+++ b/Engine/Plugins/ChaosVD/Source/ChaosVDBuiltInExtensions/Private/GenericDebugDraw/Trace/DataProcessors/ChaosVDDebugDrawSphereDataProcessor.cpp
@@ -26,6 +26,8 @@ bool FChaosVDDebugDrawSphereDataProcessor::ProcessRawData(const TArray<uint8>& I
 	{
 		if (TSharedPtr<FChaosVDGameFrameData> CurrentFrameData = ProviderSharedPtr->GetCurrentGameFrame().Pin())
 		{
+			DebugDrawData->SolverID = ProviderSharedPtr->GetRemappedSolverID(DebugDrawData->SolverID);
+	
 			if (TSharedPtr<FChaosVDDebugShapeDataContainer> DebugDrawShapeData = CurrentFrameData->GetCustomDataHandler().GetOrAddDefaultData<FChaosVDDebugShapeDataContainer>())
 			{
 				DebugDrawShapeData->RecordedDebugDrawSpheresBySolverID.FindOrAdd(DebugDrawData->SolverID).Add(DebugDrawData);
diff --git a/Engine/Source/Runtime/Experimental/ChaosVDData/Private/ChaosVDRecording.cpp b/Engine/Source/Runtime/Experimental/ChaosVDData/Private/ChaosVDRecording.cpp
--- a/Engine/Source/Runtime/Experimental/ChaosVDData/Private/ChaosVDRecording.cpp
+++ b/Engine/Source/Runtime/Experimental/ChaosVDData/Private/ChaosVDRecording.cpp
@@ -466,6 +466,11 @@ void FChaosVDRecording::SetCollisionChannelsInfoContainer(const TSharedPtr<FChao
 	CollisionChannelsInfoContainer = InCollisionChannelsInfo;
 }
 
+bool FChaosVDRecording::HasSolverID(int32 SolverID)
+{
+	return RecordedFramesDataPerSolver.Contains(SolverID);
+}
+
 void FChaosVDRecording::AddImplicitObject(const uint32 ID, const Chaos::FImplicitObject* InImplicitObject)
 {
 	FWriteScopeLock WriteLock(RecordingDataLock);
diff --git a/Engine/Source/Runtime/Experimental/ChaosVDData/Public/ChaosVDRecording.h b/Engine/Source/Runtime/Experimental/ChaosVDData/Public/ChaosVDRecording.h
--- a/Engine/Source/Runtime/Experimental/ChaosVDData/Public/ChaosVDRecording.h
+++ b/Engine/Source/Runtime/Experimental/ChaosVDData/Public/ChaosVDRecording.h
@@ -300,8 +300,11 @@ PRAGMA_ENABLE_DEPRECATION_WARNINGS
 		bIsDirty = true;
 	}
 
+	UE_DEPRECATED(5.6, "RecordedNonSolverLocationsByID is deprecated and will be removed in a future release.")
 	TMap<FName, FChaosVDTrackedLocation> RecordedNonSolverLocationsByID;
+	UE_DEPRECATED(5.6, "RecordedNonSolverTransformsByID is deprecated and will be removed in a future release.")
 	TMap<FName, FChaosVDTrackedTransform> RecordedNonSolverTransformsByID;
+
 	TMap<int32, TMap<int32, TSharedPtr<FChaosVDQueryDataWrapper>>> RecordedSceneQueriesBySolverID;
 
 	UE_DEPRECATED(5.5, "RecordedSceneQueries is deprecated and will be removed in a future release, use RecordedSceneQueriesByQueryID instead.")
@@ -309,7 +312,7 @@ PRAGMA_ENABLE_DEPRECATION_WARNINGS
 
 	TMap<int32, TSharedPtr<FChaosVDQueryDataWrapper>> RecordedSceneQueriesByQueryID;
 
-	UE_DEPRECATED(5.6, "RecordedDebugDrawBoxesBySolverID is deprecated and will be removed in a future release, use GetCustomDataContainer instead.")
+	UE_DEPRECATED(5.6, "RecordedAABBTreesBySolverID is deprecated and will be removed in a future release, use GetCustomDataContainer instead.")
 	TMap<int32, TArray<TSharedPtr<FChaosVDAABBTreeDataWrapper>>> RecordedAABBTreesBySolverID;
 
 	UE_DEPRECATED(5.6, "RecordedDebugDrawBoxesBySolverID is deprecated and will be removed in a future release, use GetCustomDataContainer instead.")
@@ -336,6 +339,22 @@ private:
 	bool bIsDirty = false;
 };
 
+struct FChaosVDSolverDataOffset
+{
+	double StartTimeOffset;
+	uint64 CycleNumberOffset;
+};
+
+enum class EChaosVDRecordingAttributes : uint8
+{
+	None = 0,
+	/** Set if this recording is being populated from a live session */
+	Live = 1 << 0,
+	/** Set if this recording contains data from multiple recordings */
+	Merged = 1 << 1,
+};
+ENUM_CLASS_FLAGS(EChaosVDRecordingAttributes)
+
 /**
  * Struct that represents a recorded Physics simulation.
  * It is currently populated while analyzing a Trace session
@@ -550,13 +569,35 @@ struct CHAOSVDDATA_API FChaosVDRecording
 	/** Returns true if this recording is being populated from a live session */
 	bool IsLive() const
 	{
-		return bIsLive;
+		return EnumHasAnyFlags(RecordingAttributes, EChaosVDRecordingAttributes::Live);
 	}
 
 	/** Sets if this recording is being populated from a live session */
 	void SetIsLive(bool bNewIsLive)
 	{
-		bIsLive = bNewIsLive;
+		if (bNewIsLive)
+		{
+			EnumAddFlags(RecordingAttributes, EChaosVDRecordingAttributes::Live);
+		}
+		else
+		{
+			EnumRemoveFlags(RecordingAttributes, EChaosVDRecordingAttributes::Live);
+		}	
+	}
+
+	void AddAttributes(EChaosVDRecordingAttributes Attributes)
+	{
+		EnumAddFlags(RecordingAttributes, Attributes);
+	}
+
+	void RemoveAttributes(EChaosVDRecordingAttributes Attributes)
+	{
+		EnumAddFlags(RecordingAttributes, Attributes);
+	}
+
+	EChaosVDRecordingAttributes GetAttributes() const
+	{
+		return RecordingAttributes;
 	}
 
 	/** Returns the name table instances used to de-duplicate strings serialization */
@@ -593,6 +634,8 @@ struct CHAOSVDDATA_API FChaosVDRecording
 	
 	void SetCollisionChannelsInfoContainer(const TSharedPtr<FChaosVDCollisionChannelsInfoContainer>& InCollisionChannelsInfo);
 
+	bool HasSolverID(int32 SolverID);
+
 protected:
 
 	/** Adds an Implicit Object to the recording and takes ownership of it */
@@ -621,8 +664,7 @@ protected:
 
 	mutable FRWLock RecordingDataLock;
 
-	/** True if this recording is being populated from a live session */
-	bool bIsLive = false;
+	EChaosVDRecordingAttributes RecordingAttributes = EChaosVDRecordingAttributes::None; 
 
 	/** Last Platform Cycle on which this recording was updated */
 	std::atomic<uint64> LastUpdatedTimeAsCycle;
@@ -714,6 +756,6 @@ void FChaosVDRecording::GetAvailableSolverIDsAtGameFrame_AssumesLocked(const FCh
 			{
 				OutSolversID.Add(SolverFramesWithIDPair.Key);
 			}
-		}	
+		}
 	}
 }
diff --git a/Engine/Source/Runtime/Experimental/ChaosVisualDebugger/Private/ChaosVDRuntimeModule.cpp b/Engine/Source/Runtime/Experimental/ChaosVisualDebugger/Private/ChaosVDRuntimeModule.cpp
--- a/Engine/Source/Runtime/Experimental/ChaosVisualDebugger/Private/ChaosVDRuntimeModule.cpp
+++ b/Engine/Source/Runtime/Experimental/ChaosVisualDebugger/Private/ChaosVDRuntimeModule.cpp
@@ -15,6 +15,7 @@
 #include "Misc/ScopedSlowTask.h"
 #endif
 
+#include "Misc/App.h"
 #include "Modules/ModuleManager.h"
 #include "ProfilingDebugging/TraceAuxiliary.h"
 
@@ -134,7 +135,9 @@ void FChaosVDRuntimeModule::StopTrace()
 
 void FChaosVDRuntimeModule::GenerateRecordingFileName(FString& OutFileName)
 {
-	OutFileName = TEXT("ChaosVD-") + FDateTime::Now().ToString(TEXT("%Y%m%d_%H%M%S.utrace"));
+	FStringFormatOrderedArguments NameArgs { FString(FApp::GetProjectName()), LexToString(FApp::GetBuildTargetType()), FDateTime::Now().ToString(TEXT("%Y%m%d_%H%M%S")) };
+	
+	OutFileName = FString::Format(TEXT("ChaosVD-{0}-{1}-{2}.utrace"), NameArgs);
 }
 
 bool FChaosVDRuntimeModule::RequestFullCapture(float DeltaTime)
