[ChaosVD] Fix an issue where the recorded log browser does not update the log line selection during playback or scrubbing

This issue is fallout of the change to support multiple CVD files.

diff --git a/Engine/Plugins/ChaosVD/Source/ChaosVD/Private/Widgets/SChaosVDRecordedLogBrowser.cpp b/Engine/Plugins/ChaosVD/Source/ChaosVD/Private/Widgets/SChaosVDRecordedLogBrowser.cpp
--- a/Engine/Plugins/ChaosVD/Source/ChaosVD/Private/Widgets/SChaosVDRecordedLogBrowser.cpp
+++ b/Engine/Plugins/ChaosVD/Source/ChaosVD/Private/Widgets/SChaosVDRecordedLogBrowser.cpp
@@ -445,9 +445,13 @@ void SChaosVDRecordedLogBrowser::HandleItemFocused(const TSharedPtr<FChaosVDLogV
 	}
 
 	TSharedPtr<FChaosVDEngine> CVDEngineInstance = EngineInstanceWeakPtr.Pin();
-	TSharedPtr<const TraceServices::IAnalysisSession> Session = /*CVDEngineInstance ?
-																FChaosVDModule::Get().GetTraceManager()->GetSession(CVDEngineInstance->GetCurrentSessionDescriptor().SessionName) :*/
-																nullptr;
+	TSharedPtr<const TraceServices::IAnalysisSession> Session = nullptr;
+
+	TArrayView<FChaosVDTraceSessionDescriptor> ActiveSessions = CVDEngineInstance ? CVDEngineInstance->GetCurrentSessionDescriptors() : TArrayView<FChaosVDTraceSessionDescriptor>();
+	if (!ActiveSessions.IsEmpty())
+	{
+		Session = FChaosVDModule::Get().GetTraceManager()->GetSession(ActiveSessions[0].SessionName);
+	}
 
 	if (!Session)
 	{
@@ -597,10 +601,15 @@ void SChaosVDRecordedLogBrowser::ProcessPendingTrackUpdates()
 	}
 
 	TSharedPtr<FChaosVDEngine> CVDEngineInstance = EngineInstanceWeakPtr.Pin();
-	TSharedPtr<FChaosVDTraceManager> TraceManager = FChaosVDModule::Get().GetTraceManager();
-	TSharedPtr<const TraceServices::IAnalysisSession> Session = /*CVDEngineInstance && TraceManager ?
-																TraceManager->GetSession(CVDEngineInstance->GetCurrentSessionDescriptor().SessionName) :*/
-																nullptr;
+
+	TSharedPtr<const TraceServices::IAnalysisSession> Session = nullptr;
+	
+	TArrayView<FChaosVDTraceSessionDescriptor> ActiveSessions = CVDEngineInstance ? CVDEngineInstance->GetCurrentSessionDescriptors() : TArrayView<FChaosVDTraceSessionDescriptor>();
+	if (!ActiveSessions.IsEmpty())
+	{
+		Session = FChaosVDModule::Get().GetTraceManager()->GetSession(ActiveSessions[0].SessionName);
+	}
+
 	if (!Session)
 	{
 		return;
