[ChaosVD] Fix localization text warning in SChaosBrowseTraceFileSourceModal

diff --git a/Engine/Plugins/ChaosVD/Source/ChaosVD/Private/Widgets/SChaosBrowseTraceFileSourceModal.cpp b/Engine/Plugins/ChaosVD/Source/ChaosVD/Private/Widgets/SChaosBrowseTraceFileSourceModal.cpp
--- a/Engine/Plugins/ChaosVD/Source/ChaosVD/Private/Widgets/SChaosBrowseTraceFileSourceModal.cpp
+++ b/Engine/Plugins/ChaosVD/Source/ChaosVD/Private/Widgets/SChaosBrowseTraceFileSourceModal.cpp
@@ -75,7 +75,7 @@ void SChaosBrowseTraceFileSourceModal::Construct(const FArguments& InArgs)
 					.VAlign(VAlign_Center)
 					[
 						SNew(STextBlock)
-						.Text(LOCTEXT("OpenFromFileMode", "Folder"))
+						.Text(LOCTEXT("OpenFromFileFolder", "Folder"))
 						.Font(FCoreStyle::GetDefaultFontStyle("Regular", 10))
 					]
 					+SHorizontalBox::Slot()
