[ChaosVD] Added support to visualize data recorded from the game thread, when multiple CVD files are loaded in the same CVD instance.

This CL is partially a workaround. CVD does not support the concept of multiple game thread data tracks, therefore this CL adds the data to a simulated solver track, and then a new CVD Data Component takes care or re-route the recorded GT data to the corresponding callbacks during playback.

The long term fix is remove the concept of Game Thread or Solver tracks, and promote the solver track type to just a "Data Track" .

This CL also updated the recorded log browser to make it handle multiple log browser tabs, one per CVD recording loaded.

diff --git a/Engine/Plugins/ChaosVD/Source/ChaosVD/Private/Actors/ChaosVDSolverInfoActor.cpp b/Engine/Plugins/ChaosVD/Source/ChaosVD/Private/Actors/ChaosVDSolverInfoActor.cpp
--- a/Engine/Plugins/ChaosVD/Source/ChaosVD/Private/Actors/ChaosVDSolverInfoActor.cpp
+++ b/Engine/Plugins/ChaosVD/Source/ChaosVD/Private/Actors/ChaosVDSolverInfoActor.cpp
@@ -7,6 +7,7 @@
 #include "ChaosVDScene.h"
 
 #include "EditorActorFolders.h"
+#include "Components/ChaosVDAdditionalGTDataRouterComponent.h"
 #include "Components/ChaosVDParticleDataComponent.h"
 #include "Components/ChaosVDSceneQueryDataComponent.h"
 #include "Components/ChaosVDSolverCharacterGroundConstraintDataComponent.h"
@@ -23,6 +24,7 @@ AChaosVDSolverInfoActor::AChaosVDSolverInfoActor()
 	JointsDataComponent = CreateDefaultSubobject<UChaosVDSolverJointConstraintDataComponent>(TEXT("JointDataComponent"));
 	CharacterGroundConstraintDataComponent = CreateDefaultSubobject<UChaosVDSolverCharacterGroundConstraintDataComponent>(TEXT("CharacterGroundConstraintDataComponent"));
 	SceneQueryDataComponent = CreateDefaultSubobject<UChaosVDSceneQueryDataComponent>(TEXT("ChaosVDSceneQueryDataComponent"));
+	GTDataReRouteComponent = CreateDefaultSubobject<UChaosVDAdditionalGTDataRouterComponent>(TEXT("ChaosVDAdditionalGTDataRouterComponent"));
 
 	bIsServer = false;
 }
diff --git a/Engine/Plugins/ChaosVD/Source/ChaosVD/Private/Actors/ChaosVDSolverInfoActor.h b/Engine/Plugins/ChaosVD/Source/ChaosVD/Private/Actors/ChaosVDSolverInfoActor.h
--- a/Engine/Plugins/ChaosVD/Source/ChaosVD/Private/Actors/ChaosVDSolverInfoActor.h
+++ b/Engine/Plugins/ChaosVD/Source/ChaosVD/Private/Actors/ChaosVDSolverInfoActor.h
@@ -8,6 +8,7 @@
 #include "GameFramework/Actor.h"
 #include "ChaosVDSolverInfoActor.generated.h"
 
+class UChaosVDAdditionalGTDataRouterComponent;
 class UChaosVDGTAccelerationStructuresDataComponent;
 class UChaosVDSolverCharacterGroundConstraintDataComponent;
 class UChaosVDSolverJointConstraintDataComponent;
@@ -112,4 +113,7 @@ protected:
 
 	UPROPERTY()
 	TObjectPtr<UChaosVDSceneQueryDataComponent> SceneQueryDataComponent;
+
+	UPROPERTY()
+	TObjectPtr<UChaosVDAdditionalGTDataRouterComponent> GTDataReRouteComponent;
 };
diff --git a/Engine/Plugins/ChaosVD/Source/ChaosVD/Private/ChaosVDEngine.cpp b/Engine/Plugins/ChaosVD/Source/ChaosVD/Private/ChaosVDEngine.cpp
--- a/Engine/Plugins/ChaosVD/Source/ChaosVD/Private/ChaosVDEngine.cpp
+++ b/Engine/Plugins/ChaosVD/Source/ChaosVD/Private/ChaosVDEngine.cpp
@@ -46,9 +46,11 @@ void FChaosVDEngine::CloseActiveSessions()
 	{
 		for (FChaosVDTraceSessionDescriptor& SessionDescriptor : CurrentSessionDescriptors)
 		{
+			OnSessionClosed().Broadcast(SessionDescriptor);
+
 			CVDTraceManager->CloseSession(SessionDescriptor.SessionName);
 			SessionDescriptor.bIsLiveSession = false;
-		}	
+		}
 	}
 }
 
@@ -183,6 +185,8 @@ void FChaosVDEngine::OpenSession(const FChaosVDTraceSessionDescriptor& SessionDe
 	{
 		PlaybackController->LoadChaosVDRecordingFromTraceSession(SessionDescriptor);
 	}
+
+	OnSessionOpened().Broadcast(SessionDescriptor);
 }
 
 bool FChaosVDEngine::Tick(float DeltaTime)
diff --git a/Engine/Plugins/ChaosVD/Source/ChaosVD/Private/ChaosVDEngine.h b/Engine/Plugins/ChaosVD/Source/ChaosVD/Private/ChaosVDEngine.h
--- a/Engine/Plugins/ChaosVD/Source/ChaosVD/Private/ChaosVDEngine.h
+++ b/Engine/Plugins/ChaosVD/Source/ChaosVD/Private/ChaosVDEngine.h
@@ -24,6 +24,8 @@ enum class EChaosVDLoadRecordedDataMode : uint8
 	MultiSource
 };
 
+DECLARE_MULTICAST_DELEGATE_OneParam(FSessionStateChangedDelegate, const FChaosVDTraceSessionDescriptor& InSessionDescriptor)
+
 /** Core Implementation of the visual debugger - Owns the systems that are not UI */
 class FChaosVDEngine : public FTSTickerObjectBase, public TSharedFromThis<FChaosVDEngine>
 {
@@ -67,6 +69,16 @@ public:
 
 	void OpenSession(const FChaosVDTraceSessionDescriptor& SessionDescriptor, EChaosVDLoadRecordedDataMode LoadingMode = EChaosVDLoadRecordedDataMode::MultiSource);
 	
+	FSessionStateChangedDelegate& OnSessionOpened()
+	{
+		return OnSessionOpenedDelegate;
+	}
+
+	FSessionStateChangedDelegate& OnSessionClosed()
+	{
+		return OnSessionClosedDelegate;
+	}
+
 private:
 
 	void UpdateRecentFilesList(const FString& InFilename);
@@ -81,4 +93,7 @@ private:
 	bool bIsInitialized = false;
 
 	FDelegateHandle LiveSessionStoppedDelegateHandle;
+	
+	FSessionStateChangedDelegate OnSessionOpenedDelegate;
+	FSessionStateChangedDelegate OnSessionClosedDelegate;
 };
diff --git a/Engine/Plugins/ChaosVD/Source/ChaosVD/Private/ChaosVDPlaybackController.cpp b/Engine/Plugins/ChaosVD/Source/ChaosVD/Private/ChaosVDPlaybackController.cpp
--- a/Engine/Plugins/ChaosVD/Source/ChaosVD/Private/ChaosVDPlaybackController.cpp
+++ b/Engine/Plugins/ChaosVD/Source/ChaosVD/Private/ChaosVDPlaybackController.cpp
@@ -514,16 +514,27 @@ int32 FChaosVDPlaybackController::ConvertCurrentFrameToOtherTrackFrame_AssumesLo
 								const FChaosVDSolverFrameData* ToSolverFrameData = LoadedRecording->GetSolverFrameData_AssumesLocked(InToTrack->TrackID, InToTrack->CurrentFrame);
 								if (FromSolverFrameData && ToSolverFrameData)
 								{
-									if (InFromTrack->bIsServer) // Server --> Client track
+									const bool bTargetTrackIsActualSolverTrack = !EnumHasAnyFlags(ToSolverFrameData->GetAttributes(), EChaosVDSolverFrameAttributes::HasGTDataToReRoute);
+									if (bTargetTrackIsActualSolverTrack)
 									{
-										ToFrame = FromSolverFrameData->InternalFrameNumber - ToSolverFrameData->GetClampedNetworkTickOffset();
+										if (InFromTrack->bIsServer) // Server --> Client track
+										{
+											ToFrame = FromSolverFrameData->InternalFrameNumber - ToSolverFrameData->GetClampedNetworkTickOffset();
+										}
+										else // Client --> Client Track or Client --> Server Track
+										{
+											// This works for Client --> Server conversion because in that case we want to add the frame offset.
+											// As the tick offset in server tracks is 0,  the following calculation will return a negative offset, which it will result in the intended addition in the last caluclation
+											int32 FrameOffset = ToSolverFrameData->GetClampedNetworkTickOffset() - FromSolverFrameData->GetClampedNetworkTickOffset();
+											ToFrame = FromSolverFrameData->InternalFrameNumber - FrameOffset;
+										}
 									}
-									else // Client --> Client Track or Client --> Server Track
+									else
 									{
-										// This works for Client --> Server conversion because in that case we want to add the frame offset.
-										// As the tick offset in server tracks is 0,  the following calculation will return a negative offset, which it will result in the intended addition in the last caluclation
-										int32 FrameOffset = ToSolverFrameData->GetClampedNetworkTickOffset() - FromSolverFrameData->GetClampedNetworkTickOffset();
-										ToFrame = FromSolverFrameData->InternalFrameNumber - FrameOffset;
+										// This fallback is required because at the moment of implementing this, multi-file / multi-session mode, process additional GT data as
+										// a fake solver track, because CVD does not support multiple gt data tracks.
+										// This also means it needs to fallback to time-stamp based sync mode for now.
+										return LoadedRecording->GetLowestSolverFrameNumberAtCycle_AssumesLocked(InToTrack->TrackID, FromSolverFrameData->FrameCycle);
 									}
 								}
 
@@ -536,7 +547,7 @@ int32 FChaosVDPlaybackController::ConvertCurrentFrameToOtherTrackFrame_AssumesLo
 				}
 				
 				// From Solver to Game track, we can only convert a frame based on timestapms
-				// TODO: Techincally we are ignoring the requested sync mode, but the current implementation of CVD relies in this fallback as we always want to sync from solver to game tracks using the original timestamo mode
+				// TODO: Technically we are ignoring the requested sync mode, but the current implementation of CVD relies in this fallback as we always want to sync from solver to game tracks using the original timestamo mode
 				// We should investigate if it is worth coming up with a better API so this default fallback does not catch anyone using it off guard.
 				return LoadedRecording->GetLowestGameFrameAtSolverFrameNumber_AssumesLocked(InFromTrack->TrackID, InFromTrack->CurrentFrame);
 			}
@@ -1271,21 +1282,25 @@ void FChaosVDPlaybackController::UpdateSolverTracksData()
 	{
 		TSharedPtr<FChaosVDTrackInfo>& SolverTrackInfo = TrackInfoPerType[EChaosVDTrackType::Solver].FindOrAdd(SolverIDPair.Key);
 
-		if (!SolverTrackInfo.IsValid())
+		bool bIsNewTrack = !SolverTrackInfo.IsValid();
+		if (bIsNewTrack)
 		{
 			SolverTrackInfo = CreateTrackInfo();
+			// Solver frame tracks only holding GT Data, are not able to change visibility states
+			SolverTrackInfo->bSupportsVisibilityChange = !SolverIDPair.Value.IsEmpty() ? !EnumHasAnyFlags(SolverIDPair.Value[0].GetAttributes(), EChaosVDSolverFrameAttributes::HasGTDataToReRoute) : true;
 			SolverTrackInfo->CurrentFrame = 0;
 			SolverTrackInfo->CurrentStage = 0;
-		};
+		}
 		
 		SolverTrackInfo->TrackID = SolverIDPair.Key;
 		SolverTrackInfo->MaxFrames = GetTrackFramesNumber(EChaosVDTrackType::Solver, SolverIDPair.Key);
 		SolverTrackInfo->TrackName = LoadedRecording->GetSolverFName(SolverIDPair.Key);
 		SolverTrackInfo->TrackType = EChaosVDTrackType::Solver;
 		SolverTrackInfo->bIsServer = LoadedRecording->IsServerSolver_AssumesLocked(SolverIDPair.Key);
-		if (SolverTrackInfo->bIsServer)
+
+		if (bIsNewTrack)
 		{
-			CachedServerTrack = SolverTrackInfo;
+			SyncTracks(ActiveTrack, CurrentSyncMode);
 		}
 	}
 }
diff --git a/Engine/Plugins/ChaosVD/Source/ChaosVD/Private/ChaosVDRecordedLogTab.cpp b/Engine/Plugins/ChaosVD/Source/ChaosVD/Private/ChaosVDRecordedLogTab.cpp
--- a/Engine/Plugins/ChaosVD/Source/ChaosVD/Private/ChaosVDRecordedLogTab.cpp
+++ b/Engine/Plugins/ChaosVD/Source/ChaosVD/Private/ChaosVDRecordedLogTab.cpp
@@ -8,7 +8,7 @@
 #include "Framework/Docking/TabManager.h"
 #include "Widgets/Docking/SDockTab.h"
 #include "Widgets/SChaosVDMainTab.h"
-#include "Widgets/SChaosVDRecordedLogBrowser.h"
+#include "Widgets/SChaosVDRecordedLogInstances.h"
 
 #define LOCTEXT_NAMESPACE "ChaosVisualDebugger"
 
@@ -23,7 +23,7 @@ TSharedRef<SDockTab> FChaosVDRecordedLogTab::HandleTabSpawnRequest(const FSpawnT
 	{
 		RecordedLogTab->SetContent
 		(
-			SNew(SChaosVDRecordedLogBrowser, MainTabPtr->GetChaosVDEngineInstance())
+			SNew(SChaosVDRecordedLogInstances, RecordedLogTab, MainTabPtr->GetChaosVDEngineInstance())
 		);
 	}
 	else
diff --git a/Engine/Plugins/ChaosVD/Source/ChaosVD/Private/Components/ChaosVDAdditionalGTDataRouterComponent.cpp b/Engine/Plugins/ChaosVD/Source/ChaosVD/Private/Components/ChaosVDAdditionalGTDataRouterComponent.cpp
new file mode 100644
--- /dev/null
+++ b/Engine/Plugins/ChaosVD/Source/ChaosVD/Private/Components/ChaosVDAdditionalGTDataRouterComponent.cpp
@@ -0,0 +1,33 @@
+// Copyright Epic Games, Inc. All Rights Reserved.
+
+#include "ChaosVDAdditionalGTDataRouterComponent.h"
+
+#include "ChaosVDRecording.h"
+#include "ChaosVDScene.h"
+#include "Actors/ChaosVDSolverInfoActor.h"
+
+void UChaosVDAdditionalGTDataRouterComponent::UpdateFromSolverFrameData(const FChaosVDSolverFrameData& InSolverFrameData)
+{
+	if (!EnumHasAnyFlags(InSolverFrameData.GetAttributes(), EChaosVDSolverFrameAttributes::HasGTDataToReRoute))
+	{
+		return;
+	}
+
+	if (AChaosVDSolverInfoActor* Owner = Cast<AChaosVDSolverInfoActor>(GetOwner()))
+	{
+		if (TSharedPtr<FChaosVDScene> CVDScene = Owner->GetScene().Pin())
+		{
+			if (TSharedPtr<FChaosVDGameFrameDataWrapper> GTFrameDataWrapper = InSolverFrameData.GetCustomData().GetData<FChaosVDGameFrameDataWrapper>())
+			{
+				for (const TPair<int32, AChaosVDSolverInfoActor*>& SolverInfoWithID : CVDScene->GetSolverInfoActorsMap())
+				{
+					SolverInfoWithID.Value->UpdateFromNewGameFrameData(*GTFrameDataWrapper->FrameData);
+				}
+			}			
+		}
+	}
+}
+
+void UChaosVDAdditionalGTDataRouterComponent::ClearData()
+{
+}
diff --git a/Engine/Plugins/ChaosVD/Source/ChaosVD/Private/Components/ChaosVDAdditionalGTDataRouterComponent.h b/Engine/Plugins/ChaosVD/Source/ChaosVD/Private/Components/ChaosVDAdditionalGTDataRouterComponent.h
new file mode 100644
--- /dev/null
+++ b/Engine/Plugins/ChaosVD/Source/ChaosVD/Private/Components/ChaosVDAdditionalGTDataRouterComponent.h
@@ -0,0 +1,20 @@
+// Copyright Epic Games, Inc. All Rights Reserved.
+
+#pragma once
+
+#include "Components/ChaosVDSolverDataComponent.h"
+#include "ChaosVDAdditionalGTDataRouterComponent.generated.h"
+
+/** Helper component used to re-route GT data loaded as a Solver Frame.
+ * This is part of a compatibility layer to support GT data from multiple files
+ * Until we refactor CVD to either support multiple GT Tracks, or remove the difference between GT track and solver track
+ */
+UCLASS()
+class UChaosVDAdditionalGTDataRouterComponent : public UChaosVDSolverDataComponent
+{
+	GENERATED_BODY()
+
+	virtual void UpdateFromSolverFrameData(const FChaosVDSolverFrameData& InSolverFrameData) override;
+
+	virtual void ClearData() override;
+};
diff --git a/Engine/Plugins/ChaosVD/Source/ChaosVD/Private/Components/ChaosVDSceneQueryDataComponent.cpp b/Engine/Plugins/ChaosVD/Source/ChaosVD/Private/Components/ChaosVDSceneQueryDataComponent.cpp
--- a/Engine/Plugins/ChaosVD/Source/ChaosVD/Private/Components/ChaosVDSceneQueryDataComponent.cpp
+++ b/Engine/Plugins/ChaosVD/Source/ChaosVD/Private/Components/ChaosVDSceneQueryDataComponent.cpp
@@ -29,11 +29,8 @@ void UChaosVDSceneQueryDataComponent::ClearSceneQuerySelection()
 	}
 }
 
-void UChaosVDSceneQueryDataComponent::UpdateFromNewGameFrameData(const FChaosVDGameFrameData& InGameFrameData)
+void UChaosVDSceneQueryDataComponent::ProccessSQData(const TMap<int32, TSharedPtr<FChaosVDQueryDataWrapper>>* RecordedQueriesByQueryID)
 {
-
-	const TMap<int32, TSharedPtr<FChaosVDQueryDataWrapper>>* RecordedQueriesByQueryID = InGameFrameData.RecordedSceneQueriesBySolverID.Find(SolverID);
-
 	if (!RecordedQueriesByQueryID)
 	{
 		return;
@@ -56,10 +53,22 @@ void UChaosVDSceneQueryDataComponent::UpdateFromNewGameFrameData(const FChaosVDG
 			RecordedQueries.Add(QueryData);
 		}
 	}
-	
+
 	ClearSceneQuerySelection();
 }
 
+void UChaosVDSceneQueryDataComponent::UpdateFromNewGameFrameData(const FChaosVDGameFrameData& InGameFrameData)
+{
+	const TMap<int32, TSharedPtr<FChaosVDQueryDataWrapper>>* RecordedQueriesByQueryID = nullptr;
+
+	if (TSharedPtr<FChaosVDSceneQueriesDataContainer> SQDataContainer = InGameFrameData.GetCustomDataHandler().GetData<FChaosVDSceneQueriesDataContainer>())
+	{
+		RecordedQueriesByQueryID = SQDataContainer->RecordedSceneQueriesBySolverID.Find(SolverID);
+	}
+
+	ProccessSQData(RecordedQueriesByQueryID);
+}
+
 TConstArrayView<TSharedPtr<FChaosVDQueryDataWrapper>> UChaosVDSceneQueryDataComponent::GetQueriesByType(EChaosVDSceneQueryType Type) const
 {
 	if (const TArray<TSharedPtr<FChaosVDQueryDataWrapper>>* FoundQueries = RecordedQueriesByType.Find(Type))
diff --git a/Engine/Plugins/ChaosVD/Source/ChaosVD/Private/Components/ChaosVDSceneQueryDataComponent.h b/Engine/Plugins/ChaosVD/Source/ChaosVD/Private/Components/ChaosVDSceneQueryDataComponent.h
--- a/Engine/Plugins/ChaosVD/Source/ChaosVD/Private/Components/ChaosVDSceneQueryDataComponent.h
+++ b/Engine/Plugins/ChaosVD/Source/ChaosVD/Private/Components/ChaosVDSceneQueryDataComponent.h
@@ -29,6 +29,8 @@ public:
 	virtual void ClearData() override;
 
 protected:
+
+	void ProccessSQData(const TMap<int32, TSharedPtr<FChaosVDQueryDataWrapper>>* RecordedQueriesByQueryID);
 	
 	TMap<EChaosVDSceneQueryType, TArray<TSharedPtr<FChaosVDQueryDataWrapper>>> RecordedQueriesByType;
 	TMap<int32, TSharedPtr<FChaosVDQueryDataWrapper>> RecordedQueriesByID;
diff --git a/Engine/Plugins/ChaosVD/Source/ChaosVD/Private/Settings/ChaosVDSceneQueryVisualizationSettings.h b/Engine/Plugins/ChaosVD/Source/ChaosVD/Private/Settings/ChaosVDSceneQueryVisualizationSettings.h
--- a/Engine/Plugins/ChaosVD/Source/ChaosVD/Private/Settings/ChaosVDSceneQueryVisualizationSettings.h
+++ b/Engine/Plugins/ChaosVD/Source/ChaosVD/Private/Settings/ChaosVDSceneQueryVisualizationSettings.h
@@ -29,8 +29,8 @@ enum class EChaosVDSQFrameVisualizationMode : uint8
 	/** All the recorded scene queries for the current frame, that passes the visualization flag filter
 	 * will be shown */
 	AllEnabledQueries,
-	/** Scene queries will be shown one at the time, in the order they were recorded, controlled by the Scene Query browser timeline */
-	RecordingOrder
+	/** Scene queries will be shown one at the time, in the order they were recorded (per solver), controlled by the Scene Query browser timeline */
+	PerSolverRecordingOrder
 };
 
 UCLASS(config=ChaosVD, PerObjectConfig)
diff --git a/Engine/Plugins/ChaosVD/Source/ChaosVD/Private/Trace/ChaosVDTraceAnalyzer.cpp b/Engine/Plugins/ChaosVD/Source/ChaosVD/Private/Trace/ChaosVDTraceAnalyzer.cpp
--- a/Engine/Plugins/ChaosVD/Source/ChaosVD/Private/Trace/ChaosVDTraceAnalyzer.cpp
+++ b/Engine/Plugins/ChaosVD/Source/ChaosVD/Private/Trace/ChaosVDTraceAnalyzer.cpp
@@ -43,6 +43,42 @@ void FChaosVDTraceAnalyzer::OnAnalysisEnd()
 	OnAnalysisComplete().Broadcast();
 }
 
+void FChaosVDTraceAnalyzer::PushSimulatedSolverTrackForGTData(const UE::Trace::IAnalyzer::FOnEventContext& Context, const UE::Trace::IAnalyzer::FEventData& EventData)
+{
+	FChaosVDSolverFrameData NewFrameData;
+
+	int32 CurrentGTTrackID = ChaosVDTraceProvider->GetRemappedSolverID(ChaosVDTraceProvider->GetCurrentGameThreadTrackID());
+	if (CurrentGTTrackID == INDEX_NONE)
+	{
+		CurrentGTTrackID = ChaosVDTraceProvider->RemapSolverID(0);
+		ChaosVDTraceProvider->SetCurrentGameThreadTrackID(CurrentGTTrackID);
+	}
+
+	NewFrameData.SolverID = CurrentGTTrackID;
+	NewFrameData.FrameCycle = EventData.GetValue<uint64>("Cycle");
+	NewFrameData.bIsKeyFrame = true;
+	NewFrameData.bIsResimulated = false;
+	NewFrameData.StartTime = Context.EventTime.AsSeconds(NewFrameData.FrameCycle);
+
+	static FString GeneratedStageName = TEXT("Stage 0");
+	NewFrameData.SolverSteps.AddDefaulted_GetRef().StepName = GeneratedStageName;
+
+	NewFrameData.GetCustomData().GetOrAddDefaultData<FChaosVDGameFrameDataWrapper>()->FrameData = MakeShared<FChaosVDGameFrameData>();
+
+	NewFrameData.DebugFName = FName(TEXT("Additional Game Frame Data Track"));
+	
+	NewFrameData.AddAttributes(EChaosVDSolverFrameAttributes::HasGTDataToReRoute);
+
+	{
+		if (FChaosVDSolverFrameData* PrevFrameData = ChaosVDTraceProvider->GetCurrentSolverFrame(CurrentGTTrackID))
+		{
+			PrevFrameData->EndTime = NewFrameData.StartTime;
+		}
+	}
+					
+	ChaosVDTraceProvider->StartSolverFrame(NewFrameData.SolverID, MoveTemp(NewFrameData));
+}
+
 bool FChaosVDTraceAnalyzer::OnEvent(uint16 RouteId, EStyle Style, const FOnEventContext& Context)
 {
 	LLM_SCOPE_BYNAME(TEXT("Insights/FChaosVDTraceAnalyzer"));
@@ -58,11 +94,23 @@ bool FChaosVDTraceAnalyzer::OnEvent(uint16 RouteId, EStyle Style, const FOnEvent
 			uint8 FrameType = EventData.GetValue<uint8>("FrameType");
 			if (static_cast<ETraceFrameType>(FrameType) == TraceFrameType_Game)
 			{
-				TSharedPtr<FChaosVDGameFrameData> FrameData = MakeShared<FChaosVDGameFrameData>();
-				FrameData->FirstCycle = EventData.GetValue<uint64>("Cycle");
-				FrameData->StartTime = Context.EventTime.AsSeconds(FrameData->FirstCycle);
+				// TODO: Currently, CVD does not support multiple GT tracks -
+				// The proper long term solution is remove the concept of GT tracks altogether, and just have 
+				// a dedicated solver data track, that happens to represent GT data.
+				// Until that is done, we generate fake solver track to hold the data in multi recording/session mode
 
-				ChaosVDTraceProvider->StartGameFrame(FrameData);
+				if (ChaosVDTraceProvider->DoesOwnRecordingInstance())
+				{
+					TSharedPtr<FChaosVDGameFrameData> FrameData = MakeShared<FChaosVDGameFrameData>();
+					FrameData->FirstCycle = EventData.GetValue<uint64>("Cycle");
+					FrameData->StartTime = Context.EventTime.AsSeconds(FrameData->FirstCycle);
+
+					ChaosVDTraceProvider->StartGameFrame(FrameData);
+				}
+				else
+				{
+					PushSimulatedSolverTrackForGTData(Context, EventData);
+				}
 			}
 
 			break;
diff --git a/Engine/Plugins/ChaosVD/Source/ChaosVD/Private/Trace/ChaosVDTraceAnalyzer.h b/Engine/Plugins/ChaosVD/Source/ChaosVD/Private/Trace/ChaosVDTraceAnalyzer.h
--- a/Engine/Plugins/ChaosVD/Source/ChaosVD/Private/Trace/ChaosVDTraceAnalyzer.h
+++ b/Engine/Plugins/ChaosVD/Source/ChaosVD/Private/Trace/ChaosVDTraceAnalyzer.h
@@ -29,6 +29,7 @@ public:
 
 	virtual void OnAnalysisBegin(const FOnAnalysisContext& Context) override;
 	virtual void OnAnalysisEnd() override;
+	void PushSimulatedSolverTrackForGTData(const UE::Trace::IAnalyzer::FOnEventContext& Context, const UE::Trace::IAnalyzer::FEventData& EventData);
 	virtual bool OnEvent(uint16 RouteId, EStyle Style, const FOnEventContext& Context) override;
 
 	FChaosVDTraceAnalysisComplete& OnAnalysisComplete()
diff --git a/Engine/Plugins/ChaosVD/Source/ChaosVD/Private/Trace/ChaosVDTraceProvider.cpp b/Engine/Plugins/ChaosVD/Source/ChaosVD/Private/Trace/ChaosVDTraceProvider.cpp
--- a/Engine/Plugins/ChaosVD/Source/ChaosVD/Private/Trace/ChaosVDTraceProvider.cpp
+++ b/Engine/Plugins/ChaosVD/Source/ChaosVD/Private/Trace/ChaosVDTraceProvider.cpp
@@ -124,10 +124,11 @@ int32 FChaosVDTraceProvider::RemapSolverID(int32 SolverID)
 	int32 RemappedSolverID = SolverID;
 	if (InternalRecording->HasSolverID(SolverID))
 	{
-		while (InternalRecording->HasSolverID(RemappedSolverID))
+		do
 		{
 			RemappedSolverID = FMath::RandRange(RemappedSolverID, std::numeric_limits<int32>::max());
 		}
+		while (InternalRecording->HasSolverID(RemappedSolverID));
 
 		RemappedSolversID.Emplace(SolverID, RemappedSolverID);
 	}
@@ -205,7 +206,6 @@ void FChaosVDTraceProvider::CommitProcessedGameFramesToRecording()
 
 void FChaosVDTraceProvider::StartGameFrame(const TSharedPtr<FChaosVDGameFrameData>& InFrameData)
 {
-	//TODO: Add support for game frame data when we are in multi recording mode (that is the only case where bHasRecordingOverride will be true)
 	if (!InternalRecording.IsValid() || bHasRecordingOverride)
 	{
 		return;
@@ -218,7 +218,9 @@ void FChaosVDTraceProvider::StartGameFrame(const TSharedPtr<FChaosVDGameFrameDat
 
 FChaosVDSolverFrameData* FChaosVDTraceProvider::GetCurrentSolverFrame(const int32 InSolverID)
 {
-	bool bIsInvalidSolverID = InSolverID == INDEX_NONE && !CurrentSolverFramesByID.IsEmpty();
+	// If we didn't remap any ID yet, InSolverID might be INDEX_NONE. This is expected as we can have data that was started being recorded in the
+	// middle of a frame and therefore the solver hasn't been open in CVD yet.
+	bool bIsInvalidSolverID = InSolverID == INDEX_NONE && (!CurrentSolverFramesByID.IsEmpty() && RemappedGameThreadTrackID != INDEX_NONE);
 
 	if (!ensure(!bIsInvalidSolverID))
 	{
@@ -235,7 +237,23 @@ FChaosVDSolverFrameData* FChaosVDTraceProvider::GetCurrentSolverFrame(const int3
 
 TWeakPtr<FChaosVDGameFrameData> FChaosVDTraceProvider::GetCurrentGameFrame()
 {
-	return CurrentGameFrame;
+	if (bHasRecordingOverride)
+	{
+		if (FChaosVDSolverFrameData* FrameData = GetCurrentSolverFrame(GetCurrentGameThreadTrackID()))
+		{
+			TSharedPtr<FChaosVDGameFrameDataWrapper> GTFrameDataWrapper = FrameData->GetCustomData().GetData<FChaosVDGameFrameDataWrapper>();
+			if (ensure(GTFrameDataWrapper))
+			{
+				return GTFrameDataWrapper->FrameData;
+			}
+		}
+	}
+	else if (TSharedPtr<FChaosVDGameFrameData> GTFrameData = CurrentGameFrame.Pin())
+	{
+		return GTFrameData;
+	}
+
+	return nullptr;
 }
 
 FChaosVDTraceProvider::FBinaryDataContainer& FChaosVDTraceProvider::FindOrAddUnprocessedData(const int32 DataID)
diff --git a/Engine/Plugins/ChaosVD/Source/ChaosVD/Private/Trace/DataProcessors/ChaosVDSceneQueryDataProcessor.cpp b/Engine/Plugins/ChaosVD/Source/ChaosVD/Private/Trace/DataProcessors/ChaosVDSceneQueryDataProcessor.cpp
--- a/Engine/Plugins/ChaosVD/Source/ChaosVD/Private/Trace/DataProcessors/ChaosVDSceneQueryDataProcessor.cpp
+++ b/Engine/Plugins/ChaosVD/Source/ChaosVD/Private/Trace/DataProcessors/ChaosVDSceneQueryDataProcessor.cpp
@@ -28,26 +28,29 @@ bool FChaosVDSceneQueryDataProcessor::ProcessRawData(const TArray<uint8>& InData
 
 	if (bSuccess)
 	{
-		if (const TSharedPtr<FChaosVDGameFrameData> CurrentFrameData = ProviderSharedPtr->GetCurrentGameFrame().Pin())
+		if (TSharedPtr<FChaosVDGameFrameData> CurrentFrameData = ProviderSharedPtr->GetCurrentGameFrame().Pin())
 		{
-			QueryData->WorldSolverID = ProviderSharedPtr->GetRemappedSolverID(QueryData->WorldSolverID);
-
-			// If ParentQueryID was set, this is a sub query, so find the parent add it to the sub-queries list so we can navigate through the query "hierarchy" later on
-			if (QueryData->ParentQueryID != INDEX_NONE)
+			if (TSharedPtr<FChaosVDSceneQueriesDataContainer> SQDataContainer = CurrentFrameData->GetCustomDataHandler().GetData<FChaosVDSceneQueriesDataContainer>())
 			{
-				if (TMap<int32, TSharedPtr<FChaosVDQueryDataWrapper>>* ParentQueryDataByQueryIDPtr = CurrentFrameData->RecordedSceneQueriesBySolverID.Find(QueryData->WorldSolverID))
+				QueryData->WorldSolverID = ProviderSharedPtr->GetRemappedSolverID(QueryData->WorldSolverID);
+
+				// If ParentQueryID was set, this is a sub query, so find the parent add it to the sub-queries list so we can navigate through the query "hierarchy" later on
+				if (QueryData->ParentQueryID != INDEX_NONE)
 				{
-					if (const TSharedPtr<FChaosVDQueryDataWrapper>* ParentQueryData = ParentQueryDataByQueryIDPtr->Find(QueryData->ParentQueryID))
+					if (TMap<int32, TSharedPtr<FChaosVDQueryDataWrapper>>* ParentQueryDataByQueryIDPtr = SQDataContainer->RecordedSceneQueriesBySolverID.Find(QueryData->WorldSolverID))
 					{
-						(*ParentQueryData)->SubQueriesIDs.Add(QueryData->ID);
-					}
-				}	
-			}
+						if (const TSharedPtr<FChaosVDQueryDataWrapper>* ParentQueryData = ParentQueryDataByQueryIDPtr->Find(QueryData->ParentQueryID))
+						{
+							(*ParentQueryData)->SubQueriesIDs.Add(QueryData->ID);
+						}
+					}	
+				}
 
-			CurrentFrameData->RecordedSceneQueriesByQueryID.Add(QueryData->ID, QueryData);
-			CurrentFrameData->RecordedSceneQueriesBySolverID.FindOrAdd(QueryData->WorldSolverID).Add(QueryData->ID, QueryData);
+				SQDataContainer->RecordedSceneQueriesByQueryID.Add(QueryData->ID, QueryData);
+				SQDataContainer->RecordedSceneQueriesBySolverID.FindOrAdd(QueryData->WorldSolverID).Add(QueryData->ID, QueryData);
 
-			CurrentFrameData->MarkDirty();
+				CurrentFrameData->MarkDirty();
+			}
 		}
 	}
 
diff --git a/Engine/Plugins/ChaosVD/Source/ChaosVD/Private/Trace/DataProcessors/ChaosVDSceneQueryVisitDataProcessor.cpp b/Engine/Plugins/ChaosVD/Source/ChaosVD/Private/Trace/DataProcessors/ChaosVDSceneQueryVisitDataProcessor.cpp
--- a/Engine/Plugins/ChaosVD/Source/ChaosVD/Private/Trace/DataProcessors/ChaosVDSceneQueryVisitDataProcessor.cpp
+++ b/Engine/Plugins/ChaosVD/Source/ChaosVD/Private/Trace/DataProcessors/ChaosVDSceneQueryVisitDataProcessor.cpp
@@ -27,28 +27,31 @@ bool FChaosVDSceneQueryVisitDataProcessor::ProcessRawData(const TArray<uint8>& I
 
 	if (bSuccess)
 	{
-		if (const TSharedPtr<FChaosVDGameFrameData> CurrentFrameData = ProviderSharedPtr->GetCurrentGameFrame().Pin())
-		{	
-			if (TSharedPtr<FChaosVDQueryDataWrapper>* QueryDataPtrPtr = CurrentFrameData->RecordedSceneQueriesByQueryID.Find(VisitStepData.OwningQueryID))
+		if (TSharedPtr<FChaosVDGameFrameData> CurrentFrameData = ProviderSharedPtr->GetCurrentGameFrame().Pin())
+		{
+			if (TSharedPtr<FChaosVDSceneQueriesDataContainer> SQDataContainer = CurrentFrameData->GetCustomDataHandler().GetOrAddDefaultData<FChaosVDSceneQueriesDataContainer>())
 			{
-				TSharedPtr<FChaosVDQueryDataWrapper> QueryDataPtr = *QueryDataPtrPtr;
-				if (QueryDataPtrPtr->IsValid())
+				if (TSharedPtr<FChaosVDQueryDataWrapper>* QueryDataPtrPtr = SQDataContainer->RecordedSceneQueriesByQueryID.Find(VisitStepData.OwningQueryID))
 				{
-					//TODO: There is an existing issue where CVD is recording garbage data for the Hit Face Normalif the query is a line trace.
-					// That value is not used in line traces, so for now just clear it out until we can implement a better solution
-					// where we include any post-processing done on the hit data during the HitConversion step.
-					if (QueryDataPtr->Type == EChaosVDSceneQueryType::RayCast)
+					TSharedPtr<FChaosVDQueryDataWrapper> QueryDataPtr = *QueryDataPtrPtr;
+					if (QueryDataPtrPtr->IsValid())
 					{
-						VisitStepData.HitData.FaceNormal = FVector::ZeroVector;
+						//TODO: There is an existing issue where CVD is recording garbage data for the Hit Face Normalif the query is a line trace.
+						// That value is not used in line traces, so for now just clear it out until we can implement a better solution
+						// where we include any post-processing done on the hit data during the HitConversion step.
+						if (QueryDataPtr->Type == EChaosVDSceneQueryType::RayCast)
+						{
+							VisitStepData.HitData.FaceNormal = FVector::ZeroVector;
+						}
+
+						if (VisitStepData.HitData.HasValidData())
+						{
+							// Quick and dirty way of show the hits in the details panel. If copying this data around becomes a bottle neck we can write a customization layout for it
+							QueryDataPtr->Hits.Add(VisitStepData);
+						}
+
+						QueryDataPtr->SQVisitData.Emplace(MoveTemp(VisitStepData));
 					}
-
-					if (VisitStepData.HitData.HasValidData())
-					{
-						// Quick and dirty way of show the hits in the details panel. If copying this data around becomes a bottle neck we can write a customization layout for it
-						QueryDataPtr->Hits.Add(VisitStepData);
-					}
-
-					QueryDataPtr->SQVisitData.Emplace(MoveTemp(VisitStepData));
 				}
 			}
 		}
diff --git a/Engine/Plugins/ChaosVD/Source/ChaosVD/Private/Visualizers/ChaosVDSceneQueryDataComponentVisualizer.cpp b/Engine/Plugins/ChaosVD/Source/ChaosVD/Private/Visualizers/ChaosVDSceneQueryDataComponentVisualizer.cpp
--- a/Engine/Plugins/ChaosVD/Source/ChaosVD/Private/Visualizers/ChaosVDSceneQueryDataComponentVisualizer.cpp
+++ b/Engine/Plugins/ChaosVD/Source/ChaosVD/Private/Visualizers/ChaosVDSceneQueryDataComponentVisualizer.cpp
@@ -113,7 +113,7 @@ void FChaosVDSceneQueryDataComponentVisualizer::DrawVisualization(const UActorCo
 
 		TSharedPtr<FChaosVDSolverDataSelectionHandle> SelectionHandle = SolverDataSelectionObject->GetCurrentSelectionHandle();
 		const bool bHasSelectedQuery = SelectionHandle && SelectionHandle->IsA<FChaosVDQueryDataWrapper>();
-		const bool bOnlyDrawSelected = Settings->CurrentVisualizationMode == EChaosVDSQFrameVisualizationMode::RecordingOrder ||
+		const bool bOnlyDrawSelected = Settings->CurrentVisualizationMode == EChaosVDSQFrameVisualizationMode::PerSolverRecordingOrder ||
 										EnumHasAnyFlags(EChaosVDSceneQueryVisualizationFlags::OnlyDrawSelectedQuery, static_cast<EChaosVDSceneQueryVisualizationFlags>(VisualizationContext.VisualizationFlags));
 
 		if (bHasSelectedQuery && bOnlyDrawSelected)
diff --git a/Engine/Plugins/ChaosVD/Source/ChaosVD/Private/Widgets/SChaosVDRecordedLogBrowser.cpp b/Engine/Plugins/ChaosVD/Source/ChaosVD/Private/Widgets/SChaosVDRecordedLogBrowser.cpp
--- a/Engine/Plugins/ChaosVD/Source/ChaosVD/Private/Widgets/SChaosVDRecordedLogBrowser.cpp
+++ b/Engine/Plugins/ChaosVD/Source/ChaosVD/Private/Widgets/SChaosVDRecordedLogBrowser.cpp
@@ -13,7 +13,6 @@
 #include "Misc/TextFilterExpressionEvaluator.h"
 #include "TraceServices/Model/AnalysisSession.h"
 #include "TraceServices/Model/Log.h"
-#include "Widgets/Input/SSearchBox.h"
 
 #define LOCTEXT_NAMESPACE "ChaosVisualDebugger"
 
@@ -137,6 +136,12 @@ void SChaosVDRecordedLogBrowser::UpdateLogLineSelectionFromGameTrack()
 	}
 }
 
+void SChaosVDRecordedLogBrowser::SetSessionName(const FString& NewSessionName)
+{
+	CurrentSessionName = NewSessionName;
+	UpdateBrowserContents();
+}
+
 void SChaosVDRecordedLogBrowser::UpdateBrowserContents()
 {
 	TSharedPtr<FChaosVDEngine> CVDEngineInstance = EngineInstanceWeakPtr.Pin();
@@ -147,13 +152,9 @@ void SChaosVDRecordedLogBrowser::UpdateBrowserContents()
 
 	TSharedPtr<const TraceServices::IAnalysisSession> Session = nullptr;
 
-	// TODO: Add support for multiple sessions instead of picking the first one.
-	// The easiest way to do it is have a log window per trace session open, and have the log browser receive the session name when it is opened
-
-	TArrayView<FChaosVDTraceSessionDescriptor> ActiveSessions = CVDEngineInstance->GetCurrentSessionDescriptors();
-	if (!ActiveSessions.IsEmpty())
+	if (!CurrentSessionName.IsEmpty())
 	{
-		Session = FChaosVDModule::Get().GetTraceManager()->GetSession(ActiveSessions[0].SessionName);
+		Session = FChaosVDModule::Get().GetTraceManager()->GetSession(CurrentSessionName);
 	}
 	
 	if (!Session)
@@ -445,13 +446,9 @@ void SChaosVDRecordedLogBrowser::HandleItemFocused(const TSharedPtr<FChaosVDLogV
 	}
 
 	TSharedPtr<FChaosVDEngine> CVDEngineInstance = EngineInstanceWeakPtr.Pin();
-	TSharedPtr<const TraceServices::IAnalysisSession> Session = nullptr;
-
-	TArrayView<FChaosVDTraceSessionDescriptor> ActiveSessions = CVDEngineInstance ? CVDEngineInstance->GetCurrentSessionDescriptors() : TArrayView<FChaosVDTraceSessionDescriptor>();
-	if (!ActiveSessions.IsEmpty())
-	{
-		Session = FChaosVDModule::Get().GetTraceManager()->GetSession(ActiveSessions[0].SessionName);
-	}
+	TSharedPtr<const TraceServices::IAnalysisSession> Session = CVDEngineInstance && !CurrentSessionName.IsEmpty() ?
+																FChaosVDModule::Get().GetTraceManager()->GetSession(CurrentSessionName) :
+																nullptr;
 
 	if (!Session)
 	{
@@ -601,15 +598,10 @@ void SChaosVDRecordedLogBrowser::ProcessPendingTrackUpdates()
 	}
 
 	TSharedPtr<FChaosVDEngine> CVDEngineInstance = EngineInstanceWeakPtr.Pin();
-
-	TSharedPtr<const TraceServices::IAnalysisSession> Session = nullptr;
-	
-	TArrayView<FChaosVDTraceSessionDescriptor> ActiveSessions = CVDEngineInstance ? CVDEngineInstance->GetCurrentSessionDescriptors() : TArrayView<FChaosVDTraceSessionDescriptor>();
-	if (!ActiveSessions.IsEmpty())
-	{
-		Session = FChaosVDModule::Get().GetTraceManager()->GetSession(ActiveSessions[0].SessionName);
-	}
-
+	TSharedPtr<FChaosVDTraceManager> TraceManager = FChaosVDModule::Get().GetTraceManager();
+	TSharedPtr<const TraceServices::IAnalysisSession> Session = CVDEngineInstance && TraceManager && !CurrentSessionName.IsEmpty() ?
+																TraceManager->GetSession(CurrentSessionName) :
+																nullptr;
 	if (!Session)
 	{
 		return;
diff --git a/Engine/Plugins/ChaosVD/Source/ChaosVD/Private/Widgets/SChaosVDRecordedLogBrowser.h b/Engine/Plugins/ChaosVD/Source/ChaosVD/Private/Widgets/SChaosVDRecordedLogBrowser.h
--- a/Engine/Plugins/ChaosVD/Source/ChaosVD/Private/Widgets/SChaosVDRecordedLogBrowser.h
+++ b/Engine/Plugins/ChaosVD/Source/ChaosVD/Private/Widgets/SChaosVDRecordedLogBrowser.h
@@ -7,12 +7,15 @@
 #include "ChaosVDPlaybackControllerObserver.h"
 #include "ContentBrowserDataMenuContexts.h"
 #include "Async/AsyncWork.h"
+#include "Framework/Docking/TabManager.h"
 #include "Misc/TextFilterExpressionEvaluator.h"
 #include "TraceServices/Model/Log.h"
 #include "Widgets/SCompoundWidget.h"
 
 #include "SChaosVDRecordedLogBrowser.generated.h"
 
+struct FChaosVDTraceSessionDescriptor;
+class FTabManager;
 class FChaosVDEngine;
 class SChaosVDLogBrowserToolbar;
 class SChaosVDRecordedLogBrowser;
@@ -156,6 +159,7 @@ public:
 	void HandleSearchTextChanged(const FText& NewText);
 	void UpdateLogLineSelectionFromGameTrack();
 
+	void SetSessionName(const FString& NewSessionName);
 
 	TSharedPtr<SChaosVDLogBrowserToolbar> GetToolBar()
 	{
@@ -272,6 +276,8 @@ private:
 	mutable FRWLock DirtyFlagsLock;
 	mutable FRWLock VerbosityFlagsLock;
 
+	FString CurrentSessionName;
+
 	friend FChaosVDLogFilteringAsyncTask;
 };
 
diff --git a/Engine/Plugins/ChaosVD/Source/ChaosVD/Private/Widgets/SChaosVDRecordedLogInstances.cpp b/Engine/Plugins/ChaosVD/Source/ChaosVD/Private/Widgets/SChaosVDRecordedLogInstances.cpp
new file mode 100644
--- /dev/null
+++ b/Engine/Plugins/ChaosVD/Source/ChaosVD/Private/Widgets/SChaosVDRecordedLogInstances.cpp
@@ -0,0 +1,174 @@
+// Copyright Epic Games, Inc. All Rights Reserved.
+
+#include "SChaosVDRecordedLogInstances.h"
+
+#include "ChaosVDEngine.h"
+#include "ChaosVDModule.h"
+#include "ChaosVDStyle.h"
+#include "SChaosVDRecordedLogBrowser.h"
+#include "Application/SlateApplicationBase.h"
+#include "Trace/ChaosVDTraceManager.h"
+#include "Widgets/SToolTip.h"
+#include "Widgets/Docking/SDockTab.h"
+#include "Widgets/Text/STextBlock.h"
+
+#define LOCTEXT_NAMESPACE "ChaosVisualDebugger"
+
+SChaosVDRecordedLogInstances::~SChaosVDRecordedLogInstances()
+{
+	InstancesTabManager->CloseAllAreas();
+
+	if (TSharedPtr<FChaosVDEngine> EngineInstance = EngineInstanceWeakPtr.Pin())
+	{
+		EngineInstance->OnSessionOpened().RemoveAll(this);
+		EngineInstance->OnSessionClosed().RemoveAll(this);
+	}
+}
+
+void SChaosVDRecordedLogInstances::Construct(const FArguments& InArgs, const TSharedRef<SDockTab>& InOwnerTab, const TSharedRef<FChaosVDEngine>& InEngineInstance)
+{
+	EngineInstanceWeakPtr = InEngineInstance;
+
+	InstancesTabManager = FGlobalTabmanager::Get()->NewTabManager(InOwnerTab).ToSharedPtr();
+	
+	TabsStack = FTabManager::NewStack()
+				->SetHideTabWell(true);
+
+	ChildSlot
+	[
+		SNew(SHorizontalBox)
+		+SHorizontalBox::Slot()
+		[
+			InstancesTabManager->RestoreFrom(GenerateMainLayout(), InstancesTabManager->GetOwnerTab()->GetParentWindow()).ToSharedRef()
+		]
+	];
+
+	InEngineInstance->OnSessionClosed().AddSP(this, &SChaosVDRecordedLogInstances::HandleSessionClosed);
+	InEngineInstance->OnSessionOpened().AddSP(this, &SChaosVDRecordedLogInstances::HandleSessionOpened);
+
+	TArrayView<FChaosVDTraceSessionDescriptor> SessionDescriptors = InEngineInstance->GetCurrentSessionDescriptors();
+	if (!SessionDescriptors.IsEmpty())
+	{
+		for (const FChaosVDTraceSessionDescriptor& SessionDescriptor : SessionDescriptors)
+		{
+			HandleSessionOpened(SessionDescriptor);
+		}
+	}
+}
+
+void SChaosVDRecordedLogInstances::HandleSessionOpened(const FChaosVDTraceSessionDescriptor& InSessionDescriptor)
+{
+	FName TabID = GetAvailableTabID();
+
+	ActiveTabsBySessionNameByTabID.Add(TabID, InSessionDescriptor.SessionName);
+
+	InstancesTabManager->TryInvokeTab(TabID);
+}
+
+void SChaosVDRecordedLogInstances::HandleSessionClosed(const FChaosVDTraceSessionDescriptor& InSessionDescriptor)
+{
+	if (TWeakPtr<SDockTab>* ActiveTabForSessionPtrPtr = ActiveTabsBySessionName.Find(InSessionDescriptor.SessionName))
+	{
+		if (TSharedPtr<SDockTab> ActiveTabPtr = ActiveTabForSessionPtrPtr->Pin())
+		{
+			ActiveTabPtr->RequestCloseTab();
+		}
+	}
+}
+
+TSharedRef<SDockTab> SChaosVDRecordedLogInstances::HandleTabSpawnRequest(const FSpawnTabArgs& Args, FName TabSessionID)
+{
+	TSharedRef<SDockTab> RecordedLogTab = SNew(SDockTab)
+											.TabRole(ETabRole::PanelTab)
+											.Label(LOCTEXT("InnerRecordedOutputLogTabLabel", "Output Log"));
+	if (ensure(!ActiveTabsByID.Contains(TabSessionID)))
+	{
+		if (const TSharedPtr<FChaosVDEngine> EngineInstance = EngineInstanceWeakPtr.Pin())
+		{
+			TSharedRef<SChaosVDRecordedLogBrowser> LogBrowser = SNew(SChaosVDRecordedLogBrowser, EngineInstance.ToSharedRef());
+
+			RecordedLogTab->SetContent(LogBrowser);
+
+			LogBrowser->SetSessionName(ActiveTabsBySessionNameByTabID.FindChecked(TabSessionID));
+		}
+
+		const FString& CurrentSessionNameForTab = ActiveTabsBySessionNameByTabID.FindChecked(TabSessionID);
+		TSharedRef<IToolTip> TabToolTip = FSlateApplicationBase::Get().MakeToolTip(FText::AsCultureInvariant(CurrentSessionNameForTab));
+
+		RecordedLogTab->SetTabToolTipWidget(StaticCastSharedRef<SToolTip>(TabToolTip));
+		RecordedLogTab->SetTabIcon(FChaosVDStyle::Get().GetBrush("TabIconOutputLog"));
+		RecordedLogTab->SetOnTabClosed(SDockTab::FOnTabClosedCallback::CreateSP(this, &SChaosVDRecordedLogInstances::HandleTabClosed, TabSessionID));
+
+		
+		ActiveTabsByID.Add(TabSessionID, RecordedLogTab);
+		ActiveTabsBySessionName.Add(ActiveTabsBySessionNameByTabID.FindChecked(TabSessionID), RecordedLogTab);
+	}
+	else
+	{
+		RecordedLogTab->SetContent(
+			SNew(SVerticalBox)
+			+SVerticalBox::Slot()
+			[
+				SNew(STextBlock)
+				.Text(LOCTEXT("ChaosVDEditorLogTabSpawnerError", "Failed to create log browser tab."))
+			]
+		);
+	}
+
+	return RecordedLogTab;
+}
+
+void SChaosVDRecordedLogInstances::HandleTabClosed(TSharedRef<SDockTab> InTabClosed, FName TabSessionID)
+{
+	AvailableTabIDs.Add(TabSessionID);
+	ActiveTabsByID.Remove(TabSessionID);
+	ActiveTabsBySessionName.Remove(ActiveTabsBySessionNameByTabID.FindChecked(TabSessionID));
+	ActiveTabsBySessionNameByTabID.Remove(TabSessionID);
+}
+
+FName SChaosVDRecordedLogInstances::GenerateTabID()
+{
+	static FName PerInstanceLogBrowserTab = FName("RecordedLogInstanceTab");
+	FName NewID = FName(PerInstanceLogBrowserTab, LastInstanceNumberUsed);
+	LastInstanceNumberUsed++;
+	InstancesTabManager->RegisterTabSpawner(NewID, FOnSpawnTab::CreateSP(this, &SChaosVDRecordedLogInstances::HandleTabSpawnRequest, NewID));
+	TabsStack->AddTab(NewID, ETabState::ClosedTab);
+
+	return NewID;
+}
+
+FName SChaosVDRecordedLogInstances::GetAvailableTabID()
+{
+	if (AvailableTabIDs.IsEmpty())
+	{
+		return GenerateTabID();
+	}
+
+	return AvailableTabIDs.Pop();
+}
+
+TSharedRef<FTabManager::FLayout> SChaosVDRecordedLogInstances::GenerateMainLayout()
+{
+	constexpr int32 MaxDefaultTabs = 16;
+	for (int32 DefaultTabsCounter = 0; DefaultTabsCounter < MaxDefaultTabs; ++DefaultTabsCounter)
+	{
+		FName GeneratedTabID = GenerateTabID();
+		AvailableTabIDs.Emplace(GeneratedTabID);
+
+		UE_LOG(LogChaosVDEditor, Verbose, TEXT("[%s] Generated default log browser tab with ID [%s]"), ANSI_TO_TCHAR(__FUNCTION__), *GeneratedTabID.ToString());
+	}
+
+	return FTabManager::NewLayout("RecordedLogBrowser_Layout")
+		->AddArea
+		(
+			FTabManager::NewPrimaryArea()
+			->SetOrientation(Orient_Horizontal)
+			->SetExtensionId("TopLevelArea")
+			->Split
+			(
+				TabsStack.ToSharedRef()
+			)
+		);
+}
+
+#undef LOCTEXT_NAMESPACE
diff --git a/Engine/Plugins/ChaosVD/Source/ChaosVD/Private/Widgets/SChaosVDRecordedLogInstances.h b/Engine/Plugins/ChaosVD/Source/ChaosVD/Private/Widgets/SChaosVDRecordedLogInstances.h
new file mode 100644
--- /dev/null
+++ b/Engine/Plugins/ChaosVD/Source/ChaosVD/Private/Widgets/SChaosVDRecordedLogInstances.h
@@ -0,0 +1,55 @@
+// Copyright Epic Games, Inc. All Rights Reserved.
+
+#pragma once
+
+#include "Framework/Docking/TabManager.h"
+#include "Widgets/SCompoundWidget.h"
+
+class FSpawnTabArgs;
+struct FChaosVDTraceSessionDescriptor;
+class FChaosVDEngine;
+class SDockTab;
+
+/**
+ * Widget that owns and manages multiple tabs of the recorded log browser widget
+ */
+class SChaosVDRecordedLogInstances : public SCompoundWidget
+{
+public:
+	SLATE_BEGIN_ARGS(SChaosVDRecordedLogInstances)
+	{
+	}
+	SLATE_END_ARGS()
+
+	virtual ~SChaosVDRecordedLogInstances() override;
+
+	void Construct(const FArguments& InArgs, const TSharedRef<SDockTab>& InOwnerTab, const TSharedRef<FChaosVDEngine>& InEngineInstance);
+
+private:
+
+	void HandleSessionOpened(const FChaosVDTraceSessionDescriptor& InSessionDescriptor);
+	void HandleSessionClosed(const FChaosVDTraceSessionDescriptor& InSessionDescriptor);
+
+	virtual TSharedRef<SDockTab> HandleTabSpawnRequest(const FSpawnTabArgs& Args, FName TabSessionID);
+
+	void HandleTabClosed(TSharedRef<SDockTab> InTabClosed, FName TabSessionID);
+	FName GenerateTabID();
+
+	FName GetAvailableTabID();
+	
+	TSharedRef<FTabManager::FLayout> GenerateMainLayout();
+
+	TSharedPtr<FTabManager> InstancesTabManager;
+
+	TSharedPtr<FTabManager::FStack> TabsStack;
+
+	TWeakPtr<FChaosVDEngine> EngineInstanceWeakPtr;
+
+	TMap<FName, TWeakPtr<SDockTab>> ActiveTabsByID;
+	TMap<FString, TWeakPtr<SDockTab>> ActiveTabsBySessionName;
+	TMap<FName, FString> ActiveTabsBySessionNameByTabID;
+
+	int32 LastInstanceNumberUsed = 0;
+
+	TArray<FName> AvailableTabIDs;
+};
diff --git a/Engine/Plugins/ChaosVD/Source/ChaosVD/Private/Widgets/SChaosVDSceneQueryBrowser.cpp b/Engine/Plugins/ChaosVD/Source/ChaosVD/Private/Widgets/SChaosVDSceneQueryBrowser.cpp
--- a/Engine/Plugins/ChaosVD/Source/ChaosVD/Private/Widgets/SChaosVDSceneQueryBrowser.cpp
+++ b/Engine/Plugins/ChaosVD/Source/ChaosVD/Private/Widgets/SChaosVDSceneQueryBrowser.cpp
@@ -153,7 +153,7 @@ void SChaosVDSceneQueryBrowser::Tick(const FGeometry& AllottedGeometry, const do
 
 void SChaosVDSceneQueryBrowser::SelectSceneQuery(const TSharedPtr<const FChaosVDQueryDataWrapper>& InQuery, ESelectInfo::Type Type)
 {
-	if (TSharedPtr<FChaosVDSceneQueryTreeItem>* ItemFound = CachedTreeItemsByID.Find(InQuery->ID))
+	if (TSharedPtr<FChaosVDSceneQueryTreeItem>* ItemFound = CachedTreeItemsByID.Find(HashCombineFast(InQuery->ID, InQuery->WorldSolverID)))
 	{
 		SelectSceneQuery(*ItemFound, Type);
 	}
@@ -342,7 +342,7 @@ void SChaosVDSceneQueryBrowser::ApplySettingsChange()
 				CurrentPlaybackIndex = 0;
 				break;
 			}
-		case EChaosVDSQFrameVisualizationMode::RecordingOrder:
+		case EChaosVDSQFrameVisualizationMode::PerSolverRecordingOrder:
 			{
 
 				if (FilteredCachedTreeItems->IsValidIndex(CurrentPlaybackIndex))
@@ -386,7 +386,7 @@ void SChaosVDSceneQueryBrowser::UpdateBrowserContents()
 
 	const TMap<int32, AChaosVDSolverInfoActor*>& SolverInfoDataMap = ScenePtr->GetSolverInfoActorsMap();
 
-	for (const TPair<int32, AChaosVDSolverInfoActor*> SolverDataWithID : SolverInfoDataMap)
+	for (const TPair<int32, AChaosVDSolverInfoActor*>& SolverDataWithID : SolverInfoDataMap)
 	{
 		if (SolverDataWithID.Key)
 		{
@@ -410,15 +410,39 @@ void SChaosVDSceneQueryBrowser::UpdateBrowserContents()
 
 	ApplyFilterToData(UnfilteredCachedTreeItems, FilteredCachedTreeItems.ToSharedRef());
 
-	if (CurrentVisualizationMode == EChaosVDSQFrameVisualizationMode::RecordingOrder)
+	if (CurrentVisualizationMode == EChaosVDSQFrameVisualizationMode::PerSolverRecordingOrder)
 	{
-		// This is used to visualize queries on their recorded order when that mode is enabled.
-		// This works because the id is only a counter, and the queries are recorded in the same thread.
-		//TODO: This will no longer be the case when we add support to merge 2 CVD files. We will have to find another solution when the time comes
-		Algo::SortBy(*FilteredCachedTreeItems, [](const TSharedPtr<FChaosVDSceneQueryTreeItem>& Item)
+		// Note: Queries need to be sorted by query id but within their specific solver because query ids are created from global counters.
+		// As CVD now supports multi-file mode we can't just rely on the IDs directly as they might collide between loaded recordings (and even
+		// if they don't, the numeric value does not indicate the order of each query across recordings.
+
+		// TODO: Now that this sorting is more complex, we might need to make this post-filtering async - Jira for tracking UE-241976
+
+		FMemMark StackMarker(FMemStack::Get());
+		typedef TArray<TSharedPtr<FChaosVDSceneQueryTreeItem>, TMemStackAllocator<>> FTempQueryItemsArray;
+
+		TMap<int32, TSharedPtr<FTempQueryItemsArray>> FilteredItemsBySolverID;
+
+		for (TSharedPtr<FChaosVDSceneQueryTreeItem>& TreeItem : (*FilteredCachedTreeItems))
+		{
+			TSharedPtr<FTempQueryItemsArray>& SolverQueries = FilteredItemsBySolverID.FindOrAdd(TreeItem->OwnerSolverID);
+			if (!SolverQueries)
+			{
+				SolverQueries = MakeShared<FTempQueryItemsArray>();
+			}
+
+			SolverQueries->HeapPush(TreeItem, [](const TSharedPtr<FChaosVDSceneQueryTreeItem>& ItemA, const TSharedPtr<FChaosVDSceneQueryTreeItem>& ItemB)
+			{
+				return ItemA && ItemB && ItemA->QueryID < ItemB->QueryID;
+			});
+		}
+
+		FilteredCachedTreeItems->Reset();
+
+		for (const TPair<int32, TSharedPtr<FTempQueryItemsArray>>& SortedQueriesPerSolverWithID : FilteredItemsBySolverID)
 		{
-			return Item->QueryID;
-		});
+			FilteredCachedTreeItems->Append((*SortedQueriesPerSolverWithID.Value));
+		}
 	}
 
 	SceneQueryTreeWidgetPtr->SetExternalSourceData(FilteredCachedTreeItems);
@@ -440,7 +464,7 @@ bool SChaosVDSceneQueryBrowser::CanUpdate() const
 	TSharedPtr<FChaosVDPlaybackController> PlaybackController = ToolkitHost ? ToolkitHost->GetChaosVDEngineInstance()->GetPlaybackController() : nullptr;
 
 	if (PlaybackController)
-	{	
+	{
 		return !PlaybackController->IsPlaying();
 	}
 
@@ -518,7 +542,7 @@ TSharedRef<SWidget> SChaosVDSceneQueryBrowser::GenerateQueriesPlaybackControls()
 
 bool SChaosVDSceneQueryBrowser::GetPlaybackControlsEnabled() const
 {
-	return CurrentVisualizationMode == EChaosVDSQFrameVisualizationMode::RecordingOrder && GetCurrentMaxPlaybackQueryIndex() > 0;
+	return CurrentVisualizationMode == EChaosVDSQFrameVisualizationMode::PerSolverRecordingOrder && GetCurrentMaxPlaybackQueryIndex() > 0;
 }
 
 void SChaosVDSceneQueryBrowser::HandlePlaybackControlInput(EChaosVDPlaybackButtonsID InputID)
@@ -590,7 +614,7 @@ TSharedPtr<FChaosVDSceneQueryTreeItem> SChaosVDSceneQueryBrowser::MakeSceneQuery
 		return nullptr;
 	}
 
-	if (TSharedPtr<FChaosVDSceneQueryTreeItem>* FoundItemPtr = CachedTreeItemsByID.Find(InQueryData->ID))
+	if (TSharedPtr<FChaosVDSceneQueryTreeItem>* FoundItemPtr = CachedTreeItemsByID.Find(HashCombineFast(InQueryData->ID, InQueryData->WorldSolverID)))
 	{
 		return *FoundItemPtr;
 	}
@@ -614,7 +638,7 @@ TSharedPtr<FChaosVDSceneQueryTreeItem> SChaosVDSceneQueryBrowser::MakeSceneQuery
 		NewTreeItem->OwnerSolverName = SolverInfo->GetSolverName();
 	}
 	
-	CachedTreeItemsByID.Add(InQueryData->ID, NewTreeItem);
+	CachedTreeItemsByID.Add(HashCombineFast(InQueryData->ID, InQueryData->WorldSolverID), NewTreeItem);
 
 	return NewTreeItem;
 }
diff --git a/Engine/Plugins/ChaosVD/Source/ChaosVD/Private/Widgets/SChaosVDSceneQueryBrowser.h b/Engine/Plugins/ChaosVD/Source/ChaosVD/Private/Widgets/SChaosVDSceneQueryBrowser.h
--- a/Engine/Plugins/ChaosVD/Source/ChaosVD/Private/Widgets/SChaosVDSceneQueryBrowser.h
+++ b/Engine/Plugins/ChaosVD/Source/ChaosVD/Private/Widgets/SChaosVDSceneQueryBrowser.h
@@ -121,7 +121,7 @@ protected:
 
 	TSharedPtr<TArray<TSharedPtr<FChaosVDSceneQueryTreeItem>>> FilteredCachedTreeItems;
 	TArray<TSharedPtr<FChaosVDSceneQueryTreeItem>> UnfilteredCachedTreeItems;
-	TMap<int32, TSharedPtr<FChaosVDSceneQueryTreeItem>> CachedTreeItemsByID;
+	TMap<uint32, TSharedPtr<FChaosVDSceneQueryTreeItem>> CachedTreeItemsByID;
 	TMap<int32, bool> CachedSolverVisibilityByID;
 
 	TMap<FName, TSharedPtr<FString>> CachedNameToStringMap;
diff --git a/Engine/Plugins/ChaosVD/Source/ChaosVD/Private/Widgets/SChaosVDSolverPlaybackControls.cpp b/Engine/Plugins/ChaosVD/Source/ChaosVD/Private/Widgets/SChaosVDSolverPlaybackControls.cpp
--- a/Engine/Plugins/ChaosVD/Source/ChaosVD/Private/Widgets/SChaosVDSolverPlaybackControls.cpp
+++ b/Engine/Plugins/ChaosVD/Source/ChaosVD/Private/Widgets/SChaosVDSolverPlaybackControls.cpp
@@ -216,6 +216,11 @@ FReply SChaosVDSolverPlaybackControls::ToggleSolverSyncLink() const
 	return FReply::Handled();
 }
 
+bool SChaosVDSolverPlaybackControls::CanChangeVisibility() const
+{
+	return SolverTrackInfoRef->bSupportsVisibilityChange;
+}
+
 const FSlateBrush* SChaosVDSolverPlaybackControls::GetBrushForCurrentVisibility() const
 {
 	return bIsVisible ? SolverVisibleIconBrush : SolverHiddenIconBrush;
@@ -264,12 +269,13 @@ TSharedPtr<SWidget> SChaosVDSolverPlaybackControls::CreateVisibilityWidget()
 	return SNew(SButton)
 			.OnClicked_Raw(this, &SChaosVDSolverPlaybackControls::ToggleSolverVisibility)
 			.ToolTipText_Raw(this, &SChaosVDSolverPlaybackControls::GetVisibilityButtonToolTipText)
+			.IsEnabled_Raw(this,&SChaosVDSolverPlaybackControls::CanChangeVisibility)
 			[
 				SNew(SImage)
 				.Image_Raw(this,&SChaosVDSolverPlaybackControls::GetBrushForCurrentVisibility)
 				.DesiredSizeOverride(FVector2D(16.0f,16.0f))
 				.ColorAndOpacity(FSlateColor::UseForeground())
-			];	
+			];
 }
 
 TSharedPtr<SWidget> SChaosVDSolverPlaybackControls::CreateSyncLinkWidget()
@@ -287,6 +293,11 @@ TSharedPtr<SWidget> SChaosVDSolverPlaybackControls::CreateSyncLinkWidget()
 
 FText SChaosVDSolverPlaybackControls::GetVisibilityButtonToolTipText() const
 {
+	if (SolverTrackInfoRef->bSupportsVisibilityChange)
+	{
+		return LOCTEXT("VisibilityControlDisabledButtonToolTipText", "This track does not support visibility changes");
+	}
+
 	return bIsVisible ? LOCTEXT("HideVisibilityButtonToolTipText", "Click to hide all the visualization data corresponding to this solver track") : LOCTEXT("ShowVisibilityButtonToolTipText", "Click to show all the visualization data corresponding to this solver track");
 }
 
diff --git a/Engine/Plugins/ChaosVD/Source/ChaosVD/Private/Widgets/SChaosVDSolverPlaybackControls.h b/Engine/Plugins/ChaosVD/Source/ChaosVD/Private/Widgets/SChaosVDSolverPlaybackControls.h
--- a/Engine/Plugins/ChaosVD/Source/ChaosVD/Private/Widgets/SChaosVDSolverPlaybackControls.h
+++ b/Engine/Plugins/ChaosVD/Source/ChaosVD/Private/Widgets/SChaosVDSolverPlaybackControls.h
@@ -67,6 +67,7 @@ private:
 	FReply ToggleSolverVisibility() const;
 	FReply ToggleSolverSyncLink() const;
 
+	bool CanChangeVisibility() const;
 	const FSlateBrush* GetBrushForCurrentVisibility() const;
 	const FSlateBrush* GetBrushForCurrentLinkState() const;
 
diff --git a/Engine/Plugins/ChaosVD/Source/ChaosVD/Public/ChaosVDPlaybackController.h b/Engine/Plugins/ChaosVD/Source/ChaosVD/Public/ChaosVDPlaybackController.h
--- a/Engine/Plugins/ChaosVD/Source/ChaosVD/Public/ChaosVDPlaybackController.h
+++ b/Engine/Plugins/ChaosVD/Source/ChaosVD/Public/ChaosVDPlaybackController.h
@@ -50,6 +50,7 @@ struct FChaosVDTrackInfo
 	bool bTrackSyncEnabled = true;
 	bool bIsServer = false;
 	bool bHasNetworkSyncData = false;
+	bool bSupportsVisibilityChange = true;
 
 	bool operator==(const FChaosVDTrackInfo& Other) const;
 
@@ -449,8 +450,6 @@ protected:
 	/** Map containing all track info, by track type*/
 	TMap<EChaosVDTrackType, TrackInfoByIDMap> TrackInfoPerType;
 
-	TWeakPtr<FChaosVDTrackInfo> CachedServerTrack;
-
 	/** Ptr to the loaded recording */
 	TSharedPtr<FChaosVDRecording> LoadedRecording;
 
diff --git a/Engine/Plugins/ChaosVD/Source/ChaosVD/Public/Trace/ChaosVDTraceProvider.h b/Engine/Plugins/ChaosVD/Source/ChaosVD/Public/Trace/ChaosVDTraceProvider.h
--- a/Engine/Plugins/ChaosVD/Source/ChaosVD/Public/Trace/ChaosVDTraceProvider.h
+++ b/Engine/Plugins/ChaosVD/Source/ChaosVD/Public/Trace/ChaosVDTraceProvider.h
@@ -135,7 +135,23 @@ public:
 	
 	int32 CHAOSVD_API GetRemappedSolverID(int32 SolverID);
 
+	bool DoesOwnRecordingInstance() const
+	{
+		return !bHasRecordingOverride;
+	}
+
 private:
+	
+	int32 GetCurrentGameThreadTrackID() const
+	{
+		return RemappedGameThreadTrackID;
+	}
+	
+	void SetCurrentGameThreadTrackID(int32 NewID)
+	{
+		RemappedGameThreadTrackID = NewID;
+	}
+
 	TraceServices::IAnalysisSession& Session;
 
 	TSharedPtr<FChaosVDRecording> InternalRecording;
@@ -163,6 +179,8 @@ private:
 	TMap<int32, int32> CurrentNetworkTickOffsets;
 	TSortedMap<int32, int32> RemappedSolversID;
 
+	int32 RemappedGameThreadTrackID = INDEX_NONE;
+
 	TSet<FString> MissingDataProcessors;
 	TSet<FString> TypesFailedToSerialize;
 
diff --git a/Engine/Source/Runtime/Experimental/ChaosVDData/Public/ChaosVDRecording.h b/Engine/Source/Runtime/Experimental/ChaosVDData/Public/ChaosVDRecording.h
--- a/Engine/Source/Runtime/Experimental/ChaosVDData/Public/ChaosVDRecording.h
+++ b/Engine/Source/Runtime/Experimental/ChaosVDData/Public/ChaosVDRecording.h
@@ -20,6 +20,8 @@
 #include "DataWrappers/ChaosVDAccelerationStructureDataWrappers.h"
 #include "DataWrappers/ChaosVDDebugShapeDataWrapper.h"
 
+#include "ChaosVDRecording.generated.h"
+
 namespace Chaos::VisualDebugger
 {
 	class FChaosVDSerializableNameTable;
@@ -183,10 +185,17 @@ typedef TArray<FChaosVDFrameStageData, TInlineAllocator<16>> FChaosVDStepsContai
 
 typedef TArray<FChaosVDFrameStageData, TInlineAllocator<16>> FChaosVDFrameStagesContainer;
 
-struct CHAOSVDDATA_API FChaosVDSolverFrameData
+UENUM()
+enum class EChaosVDSolverFrameAttributes : uint16
 {
+	None = 0,
+	HasGTDataToReRoute = 1 << 0,
+};
+ENUM_CLASS_FLAGS(EChaosVDSolverFrameAttributes)
 
-PRAGMA_DISABLE_DEPRECATION_WARNINGS
+struct CHAOSVDDATA_API FChaosVDSolverFrameData
+{
+	PRAGMA_DISABLE_DEPRECATION_WARNINGS
 	FChaosVDSolverFrameData() = default;
 	FChaosVDSolverFrameData(const FChaosVDSolverFrameData& Other) = default;
 	FChaosVDSolverFrameData(FChaosVDSolverFrameData&& Other) noexcept = default;
@@ -208,6 +217,7 @@ PRAGMA_ENABLE_DEPRECATION_WARNINGS
 	TSet<int32> ParticlesDestroyedIDs;
 	double StartTime = -1.0;
 	double EndTime = -1.0;
+
 	TArray<TSharedPtr<FChaosVDCharacterGroundConstraint>> RecordedCharacterGroundConstraints;
 	
 	/** Calculates and returns the frame time for this recorded frame.
@@ -256,9 +266,26 @@ PRAGMA_ENABLE_DEPRECATION_WARNINGS
 	{
 		return CustomData;
 	}
+	
+	EChaosVDSolverFrameAttributes GetAttributes() const
+	{
+		return FrameAttributes;	
+	}
+	
+	void AddAttributes(EChaosVDSolverFrameAttributes NewAttributes)
+	{
+		EnumAddFlags(FrameAttributes, NewAttributes);	
+	}
+
+	void RemoveAttributes(EChaosVDSolverFrameAttributes AttributesToRemove)
+	{
+		EnumRemoveFlags(FrameAttributes, AttributesToRemove);	
+	}
 
 private:
 	FChaosVDCustomFrameData CustomData;
+	
+	EChaosVDSolverFrameAttributes FrameAttributes = EChaosVDSolverFrameAttributes::None;
 };
 
 struct FChaosVDGameFrameData
@@ -305,11 +332,13 @@ PRAGMA_ENABLE_DEPRECATION_WARNINGS
 	UE_DEPRECATED(5.6, "RecordedNonSolverTransformsByID is deprecated and will be removed in a future release.")
 	TMap<FName, FChaosVDTrackedTransform> RecordedNonSolverTransformsByID;
 
+	UE_DEPRECATED(5.6, "RecordedDebugDrawBoxesBySolverID is deprecated and will be removed in a future release, use GetCustomDataContainer instead.")
 	TMap<int32, TMap<int32, TSharedPtr<FChaosVDQueryDataWrapper>>> RecordedSceneQueriesBySolverID;
 
 	UE_DEPRECATED(5.5, "RecordedSceneQueries is deprecated and will be removed in a future release, use RecordedSceneQueriesByQueryID instead.")
 	TMap<int32, TSharedPtr<FChaosVDQueryDataWrapper>> RecordedSceneQueries;
 
+	UE_DEPRECATED(5.6, "RecordedDebugDrawBoxesBySolverID is deprecated and will be removed in a future release, use GetCustomDataContainer instead.")
 	TMap<int32, TSharedPtr<FChaosVDQueryDataWrapper>> RecordedSceneQueriesByQueryID;
 
 	UE_DEPRECATED(5.6, "RecordedAABBTreesBySolverID is deprecated and will be removed in a future release, use GetCustomDataContainer instead.")
@@ -339,10 +368,12 @@ private:
 	bool bIsDirty = false;
 };
 
-struct FChaosVDSolverDataOffset
+USTRUCT()
+struct FChaosVDGameFrameDataWrapper
 {
-	double StartTimeOffset;
-	uint64 CycleNumberOffset;
+	GENERATED_BODY()
+
+	TSharedPtr<FChaosVDGameFrameData> FrameData;
 };
 
 enum class EChaosVDRecordingAttributes : uint8
diff --git a/Engine/Source/Runtime/Experimental/ChaosVisualDebugger/Public/DataWrappers/ChaosVDQueryDataWrappers.h b/Engine/Source/Runtime/Experimental/ChaosVisualDebugger/Public/DataWrappers/ChaosVDQueryDataWrappers.h
--- a/Engine/Source/Runtime/Experimental/ChaosVisualDebugger/Public/DataWrappers/ChaosVDQueryDataWrappers.h
+++ b/Engine/Source/Runtime/Experimental/ChaosVisualDebugger/Public/DataWrappers/ChaosVDQueryDataWrappers.h
@@ -471,3 +471,13 @@ struct TStructOpsTypeTraits<FChaosVDQueryDataWrapper> : public TStructOpsTypeTra
 		WithSerializer = true,
 	};
 };
+
+USTRUCT()
+struct FChaosVDSceneQueriesDataContainer
+{
+	GENERATED_BODY()
+
+	TMap<int32, TMap<int32, TSharedPtr<FChaosVDQueryDataWrapper>>> RecordedSceneQueriesBySolverID;
+
+	TMap<int32, TSharedPtr<FChaosVDQueryDataWrapper>> RecordedSceneQueriesByQueryID;
+};
