[ChaosVD] Fix CIS in Shippen Editor configurations

A previous CL wrapped the ChaosVD module with WITH_CHAOS_VISUAL_DEBUGGER , but there is editor code that uses it. WITH_CHAOS_VISUAL_DEBUGGER is 0 in shipping configurations.

diff --git a/Engine/Plugins/ChaosVD/Source/ChaosVD/Private/ChaosVDEngine.cpp b/Engine/Plugins/ChaosVD/Source/ChaosVD/Private/ChaosVDEngine.cpp
--- a/Engine/Plugins/ChaosVD/Source/ChaosVD/Private/ChaosVDEngine.cpp
+++ b/Engine/Plugins/ChaosVD/Source/ChaosVD/Private/ChaosVDEngine.cpp
@@ -26,6 +26,8 @@ void FChaosVDEngine::Initialize()
 	// Listen for the recording stop event to clear the live session flag
 	// TODO: We do something similar for the live flag on the CVD recording instance. We should unify both and have a single place where to check the live state
 	// of a session, so we also have one single place to clear the flag
+
+#if WITH_CHAOS_VISUAL_DEBUGGER
 	LiveSessionStoppedDelegateHandle = FChaosVDRuntimeModule::Get().RegisterRecordingStopCallback(FChaosVDRecordingStateChangedDelegate::FDelegate::CreateLambda([WeakThis = AsWeak()]()
 	{
 		if (const TSharedPtr<FChaosVDEngine> CVDEngine = WeakThis.Pin())
@@ -36,6 +38,7 @@ void FChaosVDEngine::Initialize()
 			}
 		}
 	}));
+#endif
 
 	bIsInitialized = true;
 }
@@ -68,11 +71,13 @@ void FChaosVDEngine::DeInitialize()
 	CloseActiveSessions();
 	
 	CurrentSessionDescriptors.Reset();
-
+	
+#if WITH_CHAOS_VISUAL_DEBUGGER
 	if (FChaosVDRuntimeModule::IsLoaded())
 	{
 		FChaosVDRuntimeModule::Get().RemoveRecordingStopCallback(LiveSessionStoppedDelegateHandle);
 	}
+#endif
 
 	LiveSessionStoppedDelegateHandle = FDelegateHandle();
 
diff --git a/Engine/Plugins/ChaosVD/Source/ChaosVD/Private/ChaosVDPlaybackController.cpp b/Engine/Plugins/ChaosVD/Source/ChaosVD/Private/ChaosVDPlaybackController.cpp
--- a/Engine/Plugins/ChaosVD/Source/ChaosVD/Private/ChaosVDPlaybackController.cpp
+++ b/Engine/Plugins/ChaosVD/Source/ChaosVD/Private/ChaosVDPlaybackController.cpp
@@ -37,16 +37,21 @@ bool FChaosVDTrackInfo::AreSameTrack(const TSharedRef<const FChaosVDTrackInfo>&
 FChaosVDPlaybackController::FChaosVDPlaybackController(const TWeakPtr<FChaosVDScene>& InSceneToControl) : ActiveTrack(MakeShared<FChaosVDTrackInfo>())
 {
 	SceneToControl = InSceneToControl;
+
+#if WITH_CHAOS_VISUAL_DEBUGGER
 	RecordingStoppedHandle = FChaosVDRuntimeModule::Get().RegisterRecordingStopCallback(FChaosVDRecordingStateChangedDelegate::FDelegate::CreateRaw(this, &FChaosVDPlaybackController::HandleDisconnectedFromSession));
+#endif
 }
 
 FChaosVDPlaybackController::~FChaosVDPlaybackController()
 {
+#if WITH_CHAOS_VISUAL_DEBUGGER
 	// There is a chance the Runtime module is unloaded by now if we had the tool open and we are closing the editor
 	if (FChaosVDRuntimeModule::IsLoaded())
 	{
 		FChaosVDRuntimeModule::Get().RemoveRecordingStopCallback(RecordingStoppedHandle);
 	}
+#endif
 
 	UnloadCurrentRecording(EChaosVDUnloadRecordingFlags::Silent);
 }
diff --git a/Engine/Plugins/ChaosVD/Source/ChaosVD/Private/Widgets/SChaosVDRecordingControls.cpp b/Engine/Plugins/ChaosVD/Source/ChaosVD/Private/Widgets/SChaosVDRecordingControls.cpp
--- a/Engine/Plugins/ChaosVD/Source/ChaosVD/Private/Widgets/SChaosVDRecordingControls.cpp
+++ b/Engine/Plugins/ChaosVD/Source/ChaosVD/Private/Widgets/SChaosVDRecordingControls.cpp
@@ -38,8 +38,10 @@ void SChaosVDRecordingControls::Construct(const FArguments& InArgs, const TShare
 		GenerateToolbarWidget()
 	];
 
+#if WITH_CHAOS_VISUAL_DEBUGGER
 	RecordingStartedHandle = FChaosVDRuntimeModule::Get().RegisterRecordingStartedCallback(FChaosVDRecordingStateChangedDelegate::FDelegate::CreateRaw(this, &SChaosVDRecordingControls::HandleRecordingStart));
 	RecordingStoppedHandle = FChaosVDRuntimeModule::Get().RegisterRecordingStopCallback(FChaosVDRecordingStateChangedDelegate::FDelegate::CreateRaw(this, &SChaosVDRecordingControls::HandleRecordingStop));
+#endif
 }
 
 TSharedRef<SWidget> SChaosVDRecordingControls::GenerateToggleRecordingStateButton(EChaosVDRecordingMode RecordingMode, const FText& StartRecordingTooltip)
@@ -223,11 +225,13 @@ bool SChaosVDRecordingControls::HasDataChannelsSupport() const
 
 SChaosVDRecordingControls::~SChaosVDRecordingControls()
 {
+#if WITH_CHAOS_VISUAL_DEBUGGER
 	if (FChaosVDRuntimeModule::IsLoaded())
 	{
 		FChaosVDRuntimeModule::Get().RemoveRecordingStartedCallback(RecordingStartedHandle);
 		FChaosVDRuntimeModule::Get().RemoveRecordingStopCallback(RecordingStoppedHandle);
 	}
+#endif
 
 	if (ConnectionAttemptNotification.IsValid())
 	{
@@ -243,6 +247,7 @@ const FSlateBrush* SChaosVDRecordingControls::GetRecordOrStopButton(EChaosVDReco
 
 void SChaosVDRecordingControls::HandleRecordingStop()
 {
+#if WITH_CHAOS_VISUAL_DEBUGGER
 	const TSharedPtr<SChaosVDMainTab> MainTabSharedPtr = MainTabWeakPtr.Pin();
 	if (!MainTabSharedPtr.IsValid())
 	{
@@ -274,6 +279,7 @@ void SChaosVDRecordingControls::HandleRecordingStop()
 			MainTabSharedPtr->GetChaosVDEngineInstance()->LoadRecording(FChaosVDRuntimeModule::Get().GetLastRecordingFileNamePath());
 		}
 	}
+#endif
 }
 
 void SChaosVDRecordingControls::HandleRecordingStart()
@@ -358,6 +364,7 @@ void SChaosVDRecordingControls::AttemptToConnectToLiveSession()
 
 FReply SChaosVDRecordingControls::ToggleRecordingState(EChaosVDRecordingMode RecordingMode)
 {
+#if WITH_CHAOS_VISUAL_DEBUGGER
 	if (!IsRecording())
 	{
 		TArray<FString, TInlineAllocator<1>> RecordingArgs;
@@ -384,6 +391,7 @@ FReply SChaosVDRecordingControls::ToggleRecordingState(EChaosVDRecordingMode Rec
 	{
 		FChaosVDRuntimeModule::Get().StopRecording();
 	}
+#endif
 
 	return FReply::Handled();
 }
@@ -507,9 +515,13 @@ FText SChaosVDRecordingControls::GetRecordingTimeText() const
 	FNumberFormattingOptions FormatOptions;
 	FormatOptions.MinimumFractionalDigits = 2;
 	FormatOptions.MaximumFractionalDigits = 2;
+
+#if WITH_CHAOS_VISUAL_DEBUGGER
 	FText SecondsText = FText::AsNumber(FChaosVDRuntimeModule::Get().GetAccumulatedRecordingTime(), &FormatOptions);
-	
 	return FText::Format(LOCTEXT("RecordingTimer","{0} s"), SecondsText);
+#else
+	return FText::GetEmpty();
+#endif
 }
 
 void SChaosVDRecordingControls::PushConnectionAttemptNotification()
