[ChaosVD] Implemented Session Discovery and one click workflows for standalone client and server instances using Message Bus.

- Created new a new CVD Remote Session manager that uses message bus to discover and issue commands to other client/server or even editors
- Updated the recording controls used in PIE to start live sessions or recording to files in one click, to make them use the new session manager.
- Added a Target selector to the main toolbar. This is where you can choose what client / server / or editor instance you want to connect to. There are a 2 pre-made targets that allows you to connect to multiple instances (client/server) in one single click, achieving the same workflow available in PIE.
- Updated the new data channel menu to make it use the new session manager, this allows you to enable/disable data channel in a remote instance from CVD's UI.
- Fixed an issue where if you start loading two files in the same frame, the remapping id systems will fail (this was a bug since I introduced the remapping system two weeks ago, but the issue was only possible to hit easily now that we can load multiple files in one click).
- Fixed an issue where if you load multiple files in the same frame, some of the initial data gets lost.
- Updated the Open File picker to allows yo select multiple files at once.
- Fixed an issue where if you try to load in Single source mode, two live sessions, the currently loaded session will not unload first.

diff --git a/Engine/Plugins/ChaosVD/Source/ChaosVD/ChaosVD.Build.cs b/Engine/Plugins/ChaosVD/Source/ChaosVD/ChaosVD.Build.cs
--- a/Engine/Plugins/ChaosVD/Source/ChaosVD/ChaosVD.Build.cs
+++ b/Engine/Plugins/ChaosVD/Source/ChaosVD/ChaosVD.Build.cs
@@ -51,6 +51,8 @@ public class ChaosVD : ModuleRules
 				"CommonMenuExtensions", 
 				"EditorWidgets", 
 				"TedsOutliner",
+				"ChaosSolverEngine",
+				"Sockets",
 			}
 			);
 
diff --git a/Engine/Plugins/ChaosVD/Source/ChaosVD/Private/ChaosVDEngine.cpp b/Engine/Plugins/ChaosVD/Source/ChaosVD/Private/ChaosVDEngine.cpp
--- a/Engine/Plugins/ChaosVD/Source/ChaosVD/Private/ChaosVDEngine.cpp
+++ b/Engine/Plugins/ChaosVD/Source/ChaosVD/Private/ChaosVDEngine.cpp
@@ -7,6 +7,7 @@
 #include "ChaosVDRuntimeModule.h"
 #include "ChaosVDScene.h"
 #include "ChaosVDSettingsManager.h"
+#include "Chaos/ChaosVDEngineEditorBridge.h"
 #include "Settings/ChaosVDMiscSettings.h"
 #include "Trace/ChaosVDTraceManager.h"
 
@@ -23,27 +24,29 @@ void FChaosVDEngine::Initialize()
 
 	PlaybackController = MakeShared<FChaosVDPlaybackController>(CurrentScene);
 
-	// Listen for the recording stop event to clear the live session flag
-	// TODO: We do something similar for the live flag on the CVD recording instance. We should unify both and have a single place where to check the live state
-	// of a session, so we also have one single place to clear the flag
-
-#if WITH_CHAOS_VISUAL_DEBUGGER
-	LiveSessionStoppedDelegateHandle = FChaosVDRuntimeModule::Get().RegisterRecordingStopCallback(FChaosVDRecordingStateChangedDelegate::FDelegate::CreateLambda([WeakThis = AsWeak()]()
+	if (TSharedPtr<FChaosVDRemoteSessionsManager> RemoteSessionManager = FChaosVDEngineEditorBridge::Get().GetRemoteSessionsManager())
 	{
-		if (const TSharedPtr<FChaosVDEngine> CVDEngine = WeakThis.Pin())
+		LiveSessionStoppedDelegateHandle = RemoteSessionManager->OnSessionRecordingStopped().AddLambda([WeakThis = AsWeak()](TWeakPtr<FChaosVDSessionInfo> Session)
 		{
-			if (CVDEngine->HasAnyLiveSessionActive())
+			if (const TSharedPtr<FChaosVDEngine> CVDEngine = WeakThis.Pin())
 			{
-				CVDEngine->CloseActiveSessions();
+				if (CVDEngine->HasAnyLiveSessionActive())
+				{
+					CVDEngine->StopActiveTraceSessions();
+
+					if (CVDEngine->PlaybackController)
+					{
+						CVDEngine->PlaybackController->HandleDisconnectedFromSession();
+					}
+				}
 			}
-		}
-	}));
-#endif
-
+		});
+	}
+	
 	bIsInitialized = true;
 }
 
-void FChaosVDEngine::CloseActiveSessions()
+void FChaosVDEngine::CloseActiveTraceSessions()
 {
 	if (const TSharedPtr<FChaosVDTraceManager> CVDTraceManager = FChaosVDModule::Get().GetTraceManager())
 	{
@@ -55,6 +58,19 @@ void FChaosVDEngine::CloseActiveSessions()
 			SessionDescriptor.bIsLiveSession = false;
 		}
 	}
+
+	CurrentSessionDescriptors.Reset();
+}
+
+void FChaosVDEngine::StopActiveTraceSessions()
+{
+	if (const TSharedPtr<FChaosVDTraceManager> CVDTraceManager = FChaosVDModule::Get().GetTraceManager())
+	{
+		for (FChaosVDTraceSessionDescriptor& SessionDescriptor : CurrentSessionDescriptors)
+		{
+			CVDTraceManager->StopSession(SessionDescriptor.SessionName);
+		}
+	}
 }
 
 void FChaosVDEngine::DeInitialize()
@@ -68,16 +84,12 @@ void FChaosVDEngine::DeInitialize()
 	CurrentScene.Reset();
 	PlaybackController.Reset();
 
-	CloseActiveSessions();
-	
-	CurrentSessionDescriptors.Reset();
-	
-#if WITH_CHAOS_VISUAL_DEBUGGER
-	if (FChaosVDRuntimeModule::IsLoaded())
+	CloseActiveTraceSessions();
+
+	if (TSharedPtr<FChaosVDRemoteSessionsManager> RemoteSessionManager = FChaosVDEngineEditorBridge::Get().GetRemoteSessionsManager())
 	{
-		FChaosVDRuntimeModule::Get().RemoveRecordingStopCallback(LiveSessionStoppedDelegateHandle);
+		RemoteSessionManager->OnSessionRecordingStopped().Remove(LiveSessionStoppedDelegateHandle);
 	}
-#endif
 
 	LiveSessionStoppedDelegateHandle = FDelegateHandle();
 
@@ -151,7 +163,7 @@ void FChaosVDEngine::LoadRecording(const FString& FilePath, EChaosVDLoadRecorded
 	}
 }
 
-bool FChaosVDEngine::ConnectToLiveSession(int32 SessionID, const FString& InSessionAddress, EChaosVDLoadRecordedDataMode LoadingMode)
+bool FChaosVDEngine::ConnectToLiveSession(uint32 SessionID, const FString& InSessionAddress, EChaosVDLoadRecordedDataMode LoadingMode)
 {
 	FChaosVDTraceSessionDescriptor NewSessionFromFileDescriptor;
 
@@ -169,7 +181,7 @@ bool FChaosVDEngine::ConnectToLiveSession(int32 SessionID, const FString& InSess
 
 	if (NewSessionFromFileDescriptor.IsValid())
 	{
-		OpenSession(NewSessionFromFileDescriptor);
+		OpenSession(NewSessionFromFileDescriptor, LoadingMode);
 		bSuccess = true;
 	}
 
@@ -180,10 +192,9 @@ void FChaosVDEngine::OpenSession(const FChaosVDTraceSessionDescriptor& SessionDe
 {
 	if (LoadingMode == EChaosVDLoadRecordedDataMode::SingleSource)
 	{
-		CloseActiveSessions();
-		CurrentSessionDescriptors.Reset();
+		CloseActiveTraceSessions();
 	}
-	
+
 	CurrentSessionDescriptors.Emplace(SessionDescriptor);
 
 	if (CurrentSessionDescriptors.Num() <= 1)
diff --git a/Engine/Plugins/ChaosVD/Source/ChaosVD/Private/ChaosVDEngine.h b/Engine/Plugins/ChaosVD/Source/ChaosVD/Private/ChaosVDEngine.h
--- a/Engine/Plugins/ChaosVD/Source/ChaosVD/Private/ChaosVDEngine.h
+++ b/Engine/Plugins/ChaosVD/Source/ChaosVD/Private/ChaosVDEngine.h
@@ -36,7 +36,8 @@ public:
 	}
 
 	void Initialize();
-	void CloseActiveSessions();
+	void CloseActiveTraceSessions();
+	void StopActiveTraceSessions();
 
 	void DeInitialize();
 
@@ -65,7 +66,7 @@ public:
 	bool HasAnyLiveSessionActive() const;
 
 	void LoadRecording(const FString& FilePath, EChaosVDLoadRecordedDataMode LoadingMode = EChaosVDLoadRecordedDataMode::MultiSource);
-	bool ConnectToLiveSession(int32 SessionID, const FString& InSessionAddress, EChaosVDLoadRecordedDataMode LoadingMode = EChaosVDLoadRecordedDataMode::MultiSource);
+	bool ConnectToLiveSession(uint32 SessionID, const FString& InSessionAddress, EChaosVDLoadRecordedDataMode LoadingMode = EChaosVDLoadRecordedDataMode::MultiSource);
 
 	void OpenSession(const FChaosVDTraceSessionDescriptor& SessionDescriptor, EChaosVDLoadRecordedDataMode LoadingMode = EChaosVDLoadRecordedDataMode::MultiSource);
 	
diff --git a/Engine/Plugins/ChaosVD/Source/ChaosVD/Private/ChaosVDPlaybackController.cpp b/Engine/Plugins/ChaosVD/Source/ChaosVD/Private/ChaosVDPlaybackController.cpp
--- a/Engine/Plugins/ChaosVD/Source/ChaosVD/Private/ChaosVDPlaybackController.cpp
+++ b/Engine/Plugins/ChaosVD/Source/ChaosVD/Private/ChaosVDPlaybackController.cpp
@@ -37,22 +37,10 @@ bool FChaosVDTrackInfo::AreSameTrack(const TSharedRef<const FChaosVDTrackInfo>&
 FChaosVDPlaybackController::FChaosVDPlaybackController(const TWeakPtr<FChaosVDScene>& InSceneToControl) : ActiveTrack(MakeShared<FChaosVDTrackInfo>())
 {
 	SceneToControl = InSceneToControl;
-
-#if WITH_CHAOS_VISUAL_DEBUGGER
-	RecordingStoppedHandle = FChaosVDRuntimeModule::Get().RegisterRecordingStopCallback(FChaosVDRecordingStateChangedDelegate::FDelegate::CreateRaw(this, &FChaosVDPlaybackController::HandleDisconnectedFromSession));
-#endif
 }
 
 FChaosVDPlaybackController::~FChaosVDPlaybackController()
 {
-#if WITH_CHAOS_VISUAL_DEBUGGER
-	// There is a chance the Runtime module is unloaded by now if we had the tool open and we are closing the editor
-	if (FChaosVDRuntimeModule::IsLoaded())
-	{
-		FChaosVDRuntimeModule::Get().RemoveRecordingStopCallback(RecordingStoppedHandle);
-	}
-#endif
-
 	UnloadCurrentRecording(EChaosVDUnloadRecordingFlags::Silent);
 }
 
diff --git a/Engine/Plugins/ChaosVD/Source/ChaosVD/Private/Trace/ChaosVDTraceManager.cpp b/Engine/Plugins/ChaosVD/Source/ChaosVD/Private/Trace/ChaosVDTraceManager.cpp
--- a/Engine/Plugins/ChaosVD/Source/ChaosVD/Private/Trace/ChaosVDTraceManager.cpp
+++ b/Engine/Plugins/ChaosVD/Source/ChaosVD/Private/Trace/ChaosVDTraceManager.cpp
@@ -7,6 +7,7 @@
 #include "Features/IModularFeatures.h"
 #include "Modules/ModuleManager.h"
 #include "Misc/Paths.h"
+#include "ProfilingDebugging/TraceAuxiliary.h"
 #include "Trace/ChaosVDTraceModule.h"
 #include "Trace/StoreClient.h"
 #include "TraceServices/AnalysisService.h"
@@ -34,9 +35,11 @@ FString FChaosVDTraceManager::LoadTraceFile(const FString& InTraceFilename, cons
 	ITraceServicesModule& TraceServicesModule = FModuleManager::LoadModuleChecked<ITraceServicesModule>("TraceServices");
 	if (const TSharedPtr<TraceServices::IAnalysisService> TraceAnalysisService = TraceServicesModule.GetAnalysisService())
 	{
+		SetPendingExternalRecordingToProcess(ExistingRecordingPtr);
+
 		if (const TSharedPtr<const TraceServices::IAnalysisSession> NewSession = TraceAnalysisService->StartAnalysis(*InTraceFilename))
 		{
-			HandleTraceSessionStarted(InTraceFilename, NewSession.ToSharedRef(), ExistingRecordingPtr);
+			AnalysisSessionByName.Add(InTraceFilename, NewSession.ToSharedRef());
 
 			return NewSession->GetName();
 		}
@@ -45,26 +48,45 @@ FString FChaosVDTraceManager::LoadTraceFile(const FString& InTraceFilename, cons
 	return FString();
 }
 
-void FChaosVDTraceManager::HandleTraceSessionStarted(const FString& InTraceName, const TSharedRef<const TraceServices::IAnalysisSession>& InNewSession, const TSharedPtr<FChaosVDRecording>& InExistingRecordingPtr)
+const UE::Trace::FStoreClient::FSessionInfo* FChaosVDTraceManager::GetTraceSessionInfo(FStringView InSessionHost, FGuid TraceGuid)
 {
-	if (InExistingRecordingPtr)
+	if (InSessionHost.IsEmpty())
 	{
-		if (FChaosVDTraceProvider* ChaosVDProvider = ConstCastSharedRef<TraceServices::IAnalysisSession>(InNewSession)->EditProvider<FChaosVDTraceProvider>(FChaosVDTraceProvider::ProviderName))
-		{
-			ChaosVDProvider->SetExternalRecordingInstanceForSession(InExistingRecordingPtr.ToSharedRef());
-		}
+		UE_LOG(LogChaosVDEditor, Error, TEXT("[%s] Failed to connect to trace store. Provided session host is empty"), ANSI_TO_TCHAR(__FUNCTION__));
+		return nullptr;
+	}
+
+	using namespace UE::Trace;
+	const FStoreClient* StoreClient = FStoreClient::Connect(InSessionHost.GetData());
+
+	if (!StoreClient)
+	{
+		UE_LOG(LogChaosVDEditor, Error, TEXT("[%s] Failed to connect to trace store at [%s]"), ANSI_TO_TCHAR(__FUNCTION__), InSessionHost.GetData())
+		return nullptr;
 	}
-			
-	AnalysisSessionByName.Add(InTraceName, InNewSession);
+
+	return StoreClient->GetSessionInfoByGuid(TraceGuid);
+}
+
+void FChaosVDTraceManager::SetPendingExternalRecordingToProcess(const TSharedPtr<FChaosVDRecording>& ExistingRecordingPtr)
+{
+	TWeakPtr<FChaosVDRecording>& PendingRecording = FChaosVDTraceManagerThreadContext::Get().PendingExternalRecordingWeakPtr;
+	ensureMsgf(!PendingRecording.Pin(), TEXT("Attempted to start a secondary trace session before a pending recording instance was processed"));
+	PendingRecording = ExistingRecordingPtr;
 }
 
 FString FChaosVDTraceManager::ConnectToLiveSession(FStringView InSessionHost, uint32 SessionID, const TSharedPtr<FChaosVDRecording>& ExistingRecordingPtr)
-{	
+{
+	FString SessionName;
+
+	if (InSessionHost.IsEmpty())
+	{
+		return SessionName;
+	}
+
 	using namespace UE::Trace;
 	FStoreClient* StoreClient = FStoreClient::Connect(InSessionHost.GetData());
 
-	FString SessionName;
-
 	if (!StoreClient)
 	{
 		return SessionName;
@@ -95,10 +117,13 @@ FString FChaosVDTraceManager::ConnectToLiveSession(FStringView InSessionHost, ui
 	{
 		// Close this session in case we were already analysing it
 		CloseSession(TraceName);
+		
+		SetPendingExternalRecordingToProcess(ExistingRecordingPtr);
 	
 		if (const TSharedPtr<const TraceServices::IAnalysisSession> NewSession = TraceAnalysisService->StartAnalysis(SessionID, *TraceName, MoveTemp(TraceData)))
 		{
-			HandleTraceSessionStarted(TraceName, NewSession.ToSharedRef(), ExistingRecordingPtr);
+			AnalysisSessionByName.Add(TraceName, NewSession.ToSharedRef());
+
 			SessionName = NewSession->GetName();
 		}
 	}
@@ -119,7 +144,7 @@ FString FChaosVDTraceManager::GetLocalTraceStoreDirPath()
 	const UE::Trace::FStoreClient::FStatus* Status = StoreClient->GetStatus();
 	if (!Status)
 	{
-		UE_LOG(LogChaosVDEditor, Error, TEXT("[%s] Failed to to get Trace Store staus"), ANSI_TO_TCHAR(__FUNCTION__));
+		UE_LOG(LogChaosVDEditor, Error, TEXT("[%s] Failed to to get Trace Store status"), ANSI_TO_TCHAR(__FUNCTION__));
 		return TEXT("");
 	}
 
@@ -144,7 +169,18 @@ void FChaosVDTraceManager::CloseSession(const FString& InSessionName)
 		{
 			(*Session)->Stop(true);
 		}
-		
+
 		AnalysisSessionByName.Remove(InSessionName);
 	}
 }
+
+void FChaosVDTraceManager::StopSession(const FString& InSessionName)
+{
+	if (const TSharedPtr<const TraceServices::IAnalysisSession>* Session = AnalysisSessionByName.Find(InSessionName))
+	{
+		if (Session->IsValid())
+		{
+			(*Session)->Stop(true);
+		}
+	}
+}
diff --git a/Engine/Plugins/ChaosVD/Source/ChaosVD/Private/Trace/ChaosVDTraceManager.h b/Engine/Plugins/ChaosVD/Source/ChaosVD/Private/Trace/ChaosVDTraceManager.h
--- a/Engine/Plugins/ChaosVD/Source/ChaosVD/Private/Trace/ChaosVDTraceManager.h
+++ b/Engine/Plugins/ChaosVD/Source/ChaosVD/Private/Trace/ChaosVDTraceManager.h
@@ -24,6 +24,17 @@ struct FChaosVDTraceSessionDescriptor
 	bool IsValid() const { return !SessionName.IsEmpty(); }
 };
 
+/** Objects that allows us to use TLS to temporarily store and access a ptr to an existing instance.
+ * This is temporary to workaround the lack of an API method we need in the trace API, and will be removed in the future.
+ * either when we add that to the API, or find another way to pass an existing CVD recording to the trace provider before anaislis starts
+ */
+class FChaosVDTraceManagerThreadContext : public TThreadSingleton<FChaosVDTraceManagerThreadContext>
+{
+public:
+
+	TWeakPtr<FChaosVDRecording> PendingExternalRecordingWeakPtr;
+};
+
 /** Manager class used by Chaos VD to interact/control UE Trace systems */
 class FChaosVDTraceManager
 {
@@ -53,12 +64,18 @@ public:
 	/** Stops and de-registers a trace session registered with the provided session name */
 	void CloseSession(const FString& InSessionName);
 
+	/** Stops a trace session registered with the provided session name */
+	void StopSession(const FString& InSessionName);
+
 	template<typename TVisitor>
 	static void EnumerateActiveSessions(FStringView InSessionHost, TVisitor Callback);
+	static const UE::Trace::FStoreClient::FSessionInfo* GetTraceSessionInfo(FStringView InSessionHost, FGuid TraceGuid);
 
 private:
-	
-	void HandleTraceSessionStarted(const FString& InTraceName, const TSharedRef<const TraceServices::IAnalysisSession>& InNewSession, const TSharedPtr<FChaosVDRecording>& InExistingRecordingPtr);
+
+	/** Temporary workaround method to set an existing recording structure in CVD's trace provider before the trace analysis starts
+	 * in the Trace analysis thread.*/
+	void SetPendingExternalRecordingToProcess(const TSharedPtr<FChaosVDRecording>& ExistingRecordingPtr);
 
 	/** The trace analysis session. */
 	TMap<FString, TSharedPtr<const TraceServices::IAnalysisSession>> AnalysisSessionByName;
diff --git a/Engine/Plugins/ChaosVD/Source/ChaosVD/Private/Trace/ChaosVDTraceModule.cpp b/Engine/Plugins/ChaosVD/Source/ChaosVD/Private/Trace/ChaosVDTraceModule.cpp
--- a/Engine/Plugins/ChaosVD/Source/ChaosVD/Private/Trace/ChaosVDTraceModule.cpp
+++ b/Engine/Plugins/ChaosVD/Source/ChaosVD/Private/Trace/ChaosVDTraceModule.cpp
@@ -20,6 +20,19 @@ void FChaosVDTraceModule::OnAnalysisBegin(TraceServices::IAnalysisSession& InSes
 {
 	const TSharedRef<FChaosVDTraceProvider> Provider = MakeShared<FChaosVDTraceProvider>(InSession);
 
+	// TODO: For multi-file support ideally we should be able to set this in the same call when we start the session but there is no API in trace for that.
+	// We can't do it immediately after we start analysis either because analysis starts in a separate thread and when we get a chance to lock it
+	// it might be too late (specially when we open multiple trace files in the same frame).
+
+	//This workaround works because this is called in the GT
+	ensure(IsInGameThread());
+	TWeakPtr<FChaosVDRecording>& PendingRecording = FChaosVDTraceManagerThreadContext::Get().PendingExternalRecordingWeakPtr;
+	if (TSharedPtr<FChaosVDRecording> RecordingInstance = PendingRecording.Pin())
+	{
+		Provider->SetExternalRecordingInstanceForSession(RecordingInstance.ToSharedRef());
+		PendingRecording.Reset();
+	}
+
 	InSession.AddProvider(FChaosVDTraceProvider::ProviderName, Provider, Provider);
 	InSession.AddAnalyzer(new FChaosVDTraceAnalyzer(InSession, Provider));
 }
diff --git a/Engine/Plugins/ChaosVD/Source/ChaosVD/Private/Trace/ChaosVDTraceProvider.cpp b/Engine/Plugins/ChaosVD/Source/ChaosVD/Private/Trace/ChaosVDTraceProvider.cpp
--- a/Engine/Plugins/ChaosVD/Source/ChaosVD/Private/Trace/ChaosVDTraceProvider.cpp
+++ b/Engine/Plugins/ChaosVD/Source/ChaosVD/Private/Trace/ChaosVDTraceProvider.cpp
@@ -122,21 +122,22 @@ FString FChaosVDTraceProvider::GenerateFormattedStringListFromSet(const TSet<FSt
 int32 FChaosVDTraceProvider::RemapSolverID(int32 SolverID)
 {
 	int32 RemappedSolverID = SolverID;
-	if (InternalRecording->HasSolverID(SolverID))
+
 	{
-		do
+		// Lock the recording until we manage to reserve a unique solver ID replacement
+		FWriteScopeLock WriteLock(InternalRecording->RecordingDataLock);
+
+		while (InternalRecording->HasSolverID_AssumesLocked(RemappedSolverID))
 		{
-			RemappedSolverID = FMath::RandRange(RemappedSolverID, std::numeric_limits<int32>::max());
+			constexpr int32 MaxValue = std::numeric_limits<int32>::max() -1;
+			RemappedSolverID = FMath::RandRange(RemappedSolverID, MaxValue);
 		}
-		while (InternalRecording->HasSolverID(RemappedSolverID));
 
-		RemappedSolversID.Emplace(SolverID, RemappedSolverID);
-	}
-	else
-	{
-		RemappedSolversID.Emplace(SolverID, SolverID);
+		InternalRecording->ReserveSolverID_AssumesLocked(RemappedSolverID);
 	}
 
+	RemappedSolversIDs.Emplace(SolverID, RemappedSolverID);
+
 	return RemappedSolverID;
 }
 
@@ -144,16 +145,13 @@ int32 FChaosVDTraceProvider::GetRemappedSolverID(int32 SolverID)
 {
 	int32 RemappedSolverID = SolverID;
 
-	if (InternalRecording && EnumHasAnyFlags(InternalRecording->GetAttributes(), EChaosVDRecordingAttributes::Merged))
+	if (int32* SolverIDPtr = RemappedSolversIDs.Find(SolverID))
+	{
+		RemappedSolverID = *SolverIDPtr;
+	}
+	else
 	{
-		if (int32* SolverIDPtr = RemappedSolversID.Find(SolverID))
-        {
-        	RemappedSolverID = *SolverIDPtr;
-        }
-		else
-        {
-			RemappedSolverID = INDEX_NONE; 
-        }
+		RemappedSolverID = INDEX_NONE; 
 	}
 
 	return RemappedSolverID;
diff --git a/Engine/Plugins/ChaosVD/Source/ChaosVD/Private/Widgets/SChaosVDMainTab.cpp b/Engine/Plugins/ChaosVD/Source/ChaosVD/Private/Widgets/SChaosVDMainTab.cpp
--- a/Engine/Plugins/ChaosVD/Source/ChaosVD/Private/Widgets/SChaosVDMainTab.cpp
+++ b/Engine/Plugins/ChaosVD/Source/ChaosVD/Private/Widgets/SChaosVDMainTab.cpp
@@ -668,14 +668,23 @@ void SChaosVDMainTab::BrowseChaosVDRecordingFromFolder(FStringView FolderPath, E
 			FolderPath.GetData(),
 			TEXT(""),
 			*ExtensionStr,
-			EFileDialogFlags::None,
+			EFileDialogFlags::Multiple,
 			OutOpenFilenames
 		);
 	}
 
-	if (OutOpenFilenames.Num() > 0)
+	// Ideally, we should not need to do this, but the UI to support multi source is not robust yet, The UI improvement task to avoid this scenario is planned as UE-197418
+	if (const bool bShouldOverrideMode = OutOpenFilenames.Num() > 0 && LoadingMode == EChaosVDLoadRecordedDataMode::SingleSource)
 	{
-		LoadCVDFile(OutOpenFilenames[0], LoadingMode);
+		GetChaosVDEngineInstance()->CloseActiveTraceSessions();
+		LoadingMode = EChaosVDLoadRecordedDataMode::MultiSource;
+	
+		UE_LOG(LogChaosVDEditor, Warning, TEXT("[%s] Single source mode was selected with multiple files. Overriding mode to Multi Source..."), ANSI_TO_TCHAR(__FUNCTION__));
+	}
+
+	for (const FString& Filename : OutOpenFilenames)
+	{
+		LoadCVDFile(Filename, LoadingMode);
 	}
 }
 
diff --git a/Engine/Plugins/ChaosVD/Source/ChaosVD/Private/Widgets/SChaosVDRecordingControls.cpp b/Engine/Plugins/ChaosVD/Source/ChaosVD/Private/Widgets/SChaosVDRecordingControls.cpp
--- a/Engine/Plugins/ChaosVD/Source/ChaosVD/Private/Widgets/SChaosVDRecordingControls.cpp
+++ b/Engine/Plugins/ChaosVD/Source/ChaosVD/Private/Widgets/SChaosVDRecordingControls.cpp
@@ -5,15 +5,19 @@
 #include "AsyncCompilationHelpers.h"
 #include "ChaosVDEngine.h"
 #include "ChaosVDModule.h"
+#include "ChaosVDRecordingDetails.h"
 #include "ChaosVDStyle.h"
-#include "ChaosVDRuntimeModule.h"
 #include "ChaosVisualDebugger/ChaosVisualDebuggerTrace.h"
 #include "Editor.h"
+#include "SChaosVDNameListPicker.h"
+#include "SEnumCombo.h"
+#include "SocketSubsystem.h"
 #include "Input/Reply.h"
 #include "Misc/MessageDialog.h"
 #include "StatusBarSubsystem.h"
 #include "ToolMenu.h"
 #include "ToolMenus.h"
+#include "Chaos/ChaosVDEngineEditorBridge.h"
 #include "Framework/MultiBox/MultiBoxBuilder.h"
 #include "Misc/ScopedSlowTask.h"
 #include "Widgets/Images/SImage.h"
@@ -25,6 +29,8 @@
 
 #define LOCTEXT_NAMESPACE "ChaosVisualDebugger"
 
+const FName SChaosVDRecordingControls::RecordingControlsToolbarName = FName("ChaosVD.MainToolBar.RecordingControls");
+
 void SChaosVDRecordingControls::Construct(const FArguments& InArgs, const TSharedRef<SChaosVDMainTab>& InMainTabSharedRef)
 {
 	MainTabWeakPtr = InMainTabSharedRef;
@@ -38,10 +44,13 @@ void SChaosVDRecordingControls::Construct(const FArguments& InArgs, const TShare
 		GenerateToolbarWidget()
 	];
 
-#if WITH_CHAOS_VISUAL_DEBUGGER
-	RecordingStartedHandle = FChaosVDRuntimeModule::Get().RegisterRecordingStartedCallback(FChaosVDRecordingStateChangedDelegate::FDelegate::CreateRaw(this, &SChaosVDRecordingControls::HandleRecordingStart));
-	RecordingStoppedHandle = FChaosVDRuntimeModule::Get().RegisterRecordingStopCallback(FChaosVDRecordingStateChangedDelegate::FDelegate::CreateRaw(this, &SChaosVDRecordingControls::HandleRecordingStop));
-#endif
+	CurrentSelectedSessionId = FChaosVDRemoteSessionsManager::LocalEditorSessionID;
+
+	if (TSharedPtr<FChaosVDRemoteSessionsManager> RemoteSessionManager = FChaosVDEngineEditorBridge::Get().GetRemoteSessionsManager())
+	{
+		RemoteSessionManager->OnSessionRecordingStarted().AddSP(this, &SChaosVDRecordingControls::HandleRecordingStart);
+		RemoteSessionManager->OnSessionRecordingStopped().AddSP(this, &SChaosVDRecordingControls::HandleRecordingStop);
+	}
 }
 
 TSharedRef<SWidget> SChaosVDRecordingControls::GenerateToggleRecordingStateButton(EChaosVDRecordingMode RecordingMode, const FText& StartRecordingTooltip)
@@ -102,11 +111,92 @@ TSharedRef<SWidget> SChaosVDRecordingControls::GenerateToggleRecordingStateButto
 			];
 }
 
+FText SChaosVDRecordingControls::GetCurrentSelectedSessionName() const
+{
+	if (TSharedPtr<FChaosVDSessionInfo> CurrentSessionPtr = GetCurrentSessionInfo())
+	{
+		return FText::AsCultureInvariant(CurrentSessionPtr->SessionName);
+	}
+
+	static FText InvalidSessionName = LOCTEXT("InvalidSessionLabel", "Invalid Session");
+
+	return InvalidSessionName;
+}
+
+TSharedRef<SWidget> SChaosVDRecordingControls::GenerateTargetSessionSelector()
+{
+	return SNew(SHorizontalBox)
+		+SHorizontalBox::Slot()
+		.AutoWidth()
+		[
+			SNew(SComboButton)
+			.MenuPlacement(MenuPlacement_AboveAnchor).ComboButtonStyle(&FAppStyle::Get().GetWidgetStyle<FComboButtonStyle>("ComboButton"))
+			.OnGetMenuContent(this, &SChaosVDRecordingControls::GenerateTargetSessionDropdown)
+			.HasDownArrow(true)
+			.ButtonContent()
+			[
+				SNew(STextBlock)
+				.Text_Raw(this, &SChaosVDRecordingControls::GetCurrentSelectedSessionName)
+				.TextStyle(&FAppStyle::Get().GetWidgetStyle<FTextBlockStyle>("NormalText"))
+		]];
+}
+
+TSharedRef<SWidget> SChaosVDRecordingControls::GenerateTargetSessionDropdown()
+{
+	using namespace Chaos::VisualDebugger;
+
+	FMenuBuilder MenuBuilder(true, nullptr);
+
+	MenuBuilder.BeginSection("CVDRecordingWidgetTargets", LOCTEXT("CVDRecordingTargetsMenu", "Available Targets"));
+	{
+		if (TSharedPtr<FChaosVDRemoteSessionsManager> RemoteSessionManager = FChaosVDEngineEditorBridge::Get().GetRemoteSessionsManager())
+		{
+			RemoteSessionManager->EnumerateActiveSessions([this, &MenuBuilder](const TSharedRef<FChaosVDSessionInfo>& InSessionInfoRef)
+			{
+				if (EnumHasAnyFlags(InSessionInfoRef->GetSessionTypeAttributes(), EChaosVDRemoteSessionAttributes::IsMultiSessionWrapper))
+				{
+					return true;
+				}
+	
+				FText SessionNameAsText = FText::AsCultureInvariant(InSessionInfoRef->SessionName);
+				MenuBuilder.AddMenuEntry(
+				SessionNameAsText,
+				FText::Format(LOCTEXT("SingleTargetItemTooltip", "Select {0} session as current target"), SessionNameAsText),
+				GetIconForSession(InSessionInfoRef->InstanceId),
+				FUIAction(FExecuteAction::CreateSP(this, &SChaosVDRecordingControls::SelectTargetSession, InSessionInfoRef->InstanceId), EUIActionRepeatMode::RepeatDisabled));
+				return true;
+			});
+		}
+	}
+	
+	MenuBuilder.EndSection();
+
+	MenuBuilder.BeginSection("CVDRecordingWidgetTargetsMulti", LOCTEXT("CVDRecordingMultiTargetsMenu", "Multi Target"));
+
+	FText AllRemoteTargetsLabel = LOCTEXT("AllRemoteOption", "All Remote");
+	MenuBuilder.AddMenuEntry(
+			AllRemoteTargetsLabel,
+			LOCTEXT("MultiRemoteTargetTooltip", "Select this to act on all remote targets"),
+			GetIconForSession(FChaosVDRemoteSessionsManager::AllRemoteSessionsWrapperGUID),
+			FUIAction(FExecuteAction::CreateSP(this, &SChaosVDRecordingControls::SelectTargetSession, FChaosVDRemoteSessionsManager::AllRemoteSessionsWrapperGUID), EUIActionRepeatMode::RepeatDisabled));
+
+	FText AllTargets = LOCTEXT("AllTargetsOption", "All");
+	MenuBuilder.AddMenuEntry(
+		AllTargets,
+		LOCTEXT("MultiTargetTooltip", "Select this to act on all targets, both Local and Remote"),
+		GetIconForSession(FChaosVDRemoteSessionsManager::AllSessionsWrapperGUID),
+		FUIAction(FExecuteAction::CreateSP(this, &SChaosVDRecordingControls::SelectTargetSession, FChaosVDRemoteSessionsManager::AllSessionsWrapperGUID), EUIActionRepeatMode::RepeatDisabled));
+
+	MenuBuilder.EndSection();
+
+	return MenuBuilder.MakeWidget();
+}
+
 TSharedRef<SWidget> SChaosVDRecordingControls::GenerateRecordingTimeTextBlock()
 {
 	return SNew(SBox)
 			.VAlign(VAlign_Center)
-			.Visibility_Lambda([this]() { return IsRecording() ? EVisibility::Visible : EVisibility::Collapsed; })
+			.Visibility_Raw(this, &SChaosVDRecordingControls::GetRecordingTimeTextBlockVisibility)
 			.Padding(12.0f, 0.0f, 0.0f, 0.0f)
 			[
 				SNew(STextBlock)
@@ -130,6 +220,14 @@ TSharedRef<SWidget> SChaosVDRecordingControls::GenerateToolbarWidget()
 	return UToolMenus::Get()->GenerateWidget(RecordingControlsToolbarName, MenuContext);
 }
 
+EVisibility SChaosVDRecordingControls::GetRecordingTimeTextBlockVisibility() const
+{
+	TSharedPtr<FChaosVDSessionInfo> SessionInfo = GetCurrentSessionInfo();
+	bool bIsVisible = SessionInfo && !EnumHasAnyFlags(SessionInfo->GetSessionTypeAttributes(), EChaosVDRemoteSessionAttributes::IsMultiSessionWrapper) && IsRecording();
+
+	return  bIsVisible ? EVisibility::Visible : EVisibility::Collapsed;
+}
+
 TSharedRef<SWidget> SChaosVDRecordingControls::GenerateDataChannelsButton()
 {
 	return SNew(SComboButton)
@@ -153,90 +251,152 @@ TSharedRef<SWidget> SChaosVDRecordingControls::GenerateDataChannelsButton()
 			];
 }
 
+TSharedRef<SWidget> SChaosVDRecordingControls::GenerateLoadingModeSelector()
+{
+	TAttribute<int32> GetCurrentValueAttribute;
+	GetCurrentValueAttribute.BindSPLambda(this, [this]()
+	{
+		return static_cast<int32>(CurrentLoadingMode);
+	});
+
+	SEnumComboBox::FOnEnumSelectionChanged EnumValueChangedDelegate = SEnumComboBox::FOnEnumSelectionChanged::CreateSPLambda(this, [this](int32 NewValue, ESelectInfo::Type SelectionType)
+	{
+		CurrentLoadingMode = static_cast<EChaosVDLoadRecordedDataMode>(NewValue);
+	});
+	
+	return SNew(SEnumComboBox, StaticEnum<EChaosVDLoadRecordedDataMode>())
+		.IsEnabled_Raw(this, &SChaosVDRecordingControls::CanChangeLoadingMode)
+		.CurrentValue(GetCurrentValueAttribute)
+		.OnEnumSelectionChanged(EnumValueChangedDelegate);
+}
+
 TSharedRef<SWidget> SChaosVDRecordingControls::GenerateDataChannelsMenu()
 {
 	using namespace Chaos::VisualDebugger;
 
 	FMenuBuilder MenuBuilder(true, nullptr);
 
-#if WITH_CHAOS_VISUAL_DEBUGGER
 	MenuBuilder.BeginSection("CVDRecordingWidget", LOCTEXT("CVDRecordingMenuChannels", "Data Channels"));
 	{
-		FChaosVDDataChannelsManager::Get().EnumerateChannels([this, &MenuBuilder](const TSharedRef<FCVDDataChannel>& Channel)
+		if (TSharedPtr<FChaosVDSessionInfo> SessionInfo = GetCurrentSessionInfo())
 		{
-			MenuBuilder.AddMenuEntry(
-				Channel->GetDisplayName(),
-				FText::Format(LOCTEXT("ChannelDesc", "Enable/disable the {0} channel"), Channel->GetDisplayName()),
+			for (const TPair<FString, FChaosVDDataChannelState>& DataChannelStateWithName : SessionInfo->DataChannelsStatesByName)
+			{
+				FText ChannelNamesAsText = FText::AsCultureInvariant(DataChannelStateWithName.Key);
+				MenuBuilder.AddMenuEntry(
+				ChannelNamesAsText,
+				FText::Format(LOCTEXT("ChannelDesc", "Enable/disable the {0} channel"), ChannelNamesAsText),
 				FSlateIcon(),
-				FUIAction(FExecuteAction::CreateSP(this, &SChaosVDRecordingControls::ToggleChannelEnabledState, Channel.ToWeakPtr()),
-					FCanExecuteAction::CreateSP(this, &SChaosVDRecordingControls::CanChangeChannelEnabledState, Channel.ToWeakPtr()), FIsActionChecked::CreateSP(this, &SChaosVDRecordingControls::IsChannelEnabled,  Channel.ToWeakPtr())), NAME_None, EUserInterfaceActionType::ToggleButton);
-
-			return true;
-		});
+				FUIAction(FExecuteAction::CreateSP(this, &SChaosVDRecordingControls::ToggleChannelEnabledState, DataChannelStateWithName.Key),
+					FCanExecuteAction::CreateSP(this, &SChaosVDRecordingControls::CanChangeChannelEnabledState, DataChannelStateWithName.Key), FIsActionChecked::CreateSP(this, &SChaosVDRecordingControls::IsChannelEnabled, DataChannelStateWithName.Key)), NAME_None, EUserInterfaceActionType::ToggleButton);
+			}
+		}
 	}
-#endif //WITH_CHAOS_VISUAL_DEBUGGER
 
 	MenuBuilder.EndSection();
 
 	return MenuBuilder.MakeWidget();
 }
 
-#if WITH_CHAOS_VISUAL_DEBUGGER
-
-void SChaosVDRecordingControls::ToggleChannelEnabledState(TWeakPtr<FCVDDataChannel> Channel)
+void SChaosVDRecordingControls::ToggleChannelEnabledState(FString ChannelName)
 {
-	if (const TSharedPtr<FCVDDataChannel>& ChannelPtr = Channel.Pin())
+	TSharedPtr<FChaosVDRemoteSessionsManager> RemoteSessionManager = FChaosVDEngineEditorBridge::Get().GetRemoteSessionsManager();
+	TSharedPtr<FChaosVDSessionInfo> SessionInfo = GetCurrentSessionInfo();
+	if (RemoteSessionManager && SessionInfo)
 	{
-		const bool bIsEnabled = ChannelPtr->IsChannelEnabled();
-		ChannelPtr->SetChannelEnabled(!bIsEnabled);
+		if (FChaosVDDataChannelState* ChannelState = SessionInfo->DataChannelsStatesByName.Find(ChannelName))
+		{
+			ChannelState->bWaitingUpdatedState = true;
+			RemoteSessionManager->SendDataChannelStateChangeCommand(SessionInfo->Address, {ChannelState->ChannelName, !ChannelState->bIsEnabled });
+		}
 	}
 }
 
-bool SChaosVDRecordingControls::IsChannelEnabled(TWeakPtr<FCVDDataChannel> Channel)
+bool SChaosVDRecordingControls::IsChannelEnabled(FString ChannelName)
 {
-	if (const TSharedPtr<Chaos::VisualDebugger::FChaosVDOptionalDataChannel>& ChannelPtr = Channel.Pin())
+	if (TSharedPtr<FChaosVDSessionInfo> SessionInfo = GetCurrentSessionInfo())
 	{
-		return ChannelPtr->IsChannelEnabled();
+		if (FChaosVDDataChannelState* ChannelState = SessionInfo->DataChannelsStatesByName.Find(ChannelName))
+		{
+			return ChannelState->bIsEnabled;
+		}
 	}
-	
+
 	return false;
 }
 
-bool SChaosVDRecordingControls::CanChangeChannelEnabledState(TWeakPtr<FCVDDataChannel> Channel)
+bool SChaosVDRecordingControls::CanChangeChannelEnabledState(FString ChannelName)
 {
-	if (const TSharedPtr<FCVDDataChannel>& ChannelPtr = Channel.Pin())
+	if (TSharedPtr<FChaosVDSessionInfo> SessionInfo = GetCurrentSessionInfo())
 	{
-		return ChannelPtr->CanChangeEnabledState();
+		if (FChaosVDDataChannelState* ChannelState = SessionInfo->DataChannelsStatesByName.Find(ChannelName))
+		{
+			return ChannelState->bCanChangeChannelState && !ChannelState->bWaitingUpdatedState;
+		}
 	}
 	
 	return false;
 }
 
-#endif // WITH_CHAOS_VISUAL_DEBUGGER
+void SChaosVDRecordingControls::SelectTargetSession(FGuid SessionId)
+{
+	CurrentSelectedSessionId = SessionId;
+}
 
-bool SChaosVDRecordingControls::HasDataChannelsSupport() const
+FSlateIcon SChaosVDRecordingControls::GetIconForSession(FGuid SessionId)
 {
-#if WITH_CHAOS_VISUAL_DEBUGGER
-	return !FChaosVDModule::IsStandaloneChaosVisualDebugger();
-#else
-	return false;
-#endif
+	if (TSharedPtr<FChaosVDSessionInfo> SessionInfo = GetSessionInfo(SessionId))
+	{
+		return SessionInfo->IsRecording() ? FSlateIcon(FChaosVDStyle::GetStyleSetName(),FName("RecordIcon"), FName("RecordIcon")) : FSlateIcon();
+	}
+
+	return FSlateIcon();
 }
 
-SChaosVDRecordingControls::~SChaosVDRecordingControls()
+TSharedPtr<FChaosVDSessionInfo> SChaosVDRecordingControls::GetCurrentSessionInfo() const
+{
+	return GetSessionInfo(CurrentSelectedSessionId);
+}
+
+TSharedPtr<FChaosVDSessionInfo> SChaosVDRecordingControls::GetSessionInfo(FGuid Id) const
 {
-#if WITH_CHAOS_VISUAL_DEBUGGER
-	if (FChaosVDRuntimeModule::IsLoaded())
+	TSharedPtr<FChaosVDRemoteSessionsManager> RemoteSessionManager = FChaosVDEngineEditorBridge::Get().GetRemoteSessionsManager();
+	TSharedPtr<FChaosVDSessionInfo> SessionInfo = RemoteSessionManager ? RemoteSessionManager->GetSessionInfo(Id).Pin() : nullptr;
+
+	return SessionInfo;
+}
+
+bool SChaosVDRecordingControls::HasDataChannelsSupport() const
+{
+	if (TSharedPtr<FChaosVDSessionInfo> SessionInfo = GetCurrentSessionInfo())
 	{
-		FChaosVDRuntimeModule::Get().RemoveRecordingStartedCallback(RecordingStartedHandle);
-		FChaosVDRuntimeModule::Get().RemoveRecordingStopCallback(RecordingStoppedHandle);
+		return !SessionInfo->DataChannelsStatesByName.IsEmpty();
 	}
-#endif
 
-	if (ConnectionAttemptNotification.IsValid())
+	return false;
+}
+
+bool SChaosVDRecordingControls::CanChangeLoadingMode() const
+{
+	if (TSharedPtr<FChaosVDSessionInfo> CurrentSession = GetCurrentSessionInfo())
 	{
-		ConnectionAttemptNotification->ExpireAndFadeout();
+		// In multi session mode targets, the loading mode is controlled automatically
+		if (EnumHasAnyFlags(CurrentSession->GetSessionTypeAttributes(), EChaosVDRemoteSessionAttributes::IsMultiSessionWrapper))
+		{
+			return false;
+		}
+		else if (const TSharedPtr<SChaosVDMainTab> MainTabSharedPtr = MainTabWeakPtr.Pin())
+		{
+			// If nothing is loaded yet, it does not make sense change the loading mode
+			return !MainTabSharedPtr->GetChaosVDEngineInstance()->GetCurrentSessionDescriptors().IsEmpty();
+		}
 	}
+
+	return false;
+}
+
+SChaosVDRecordingControls::~SChaosVDRecordingControls()
+{
 }
 
 const FSlateBrush* SChaosVDRecordingControls::GetRecordOrStopButton(EChaosVDRecordingMode RecordingMode) const
@@ -245,16 +405,24 @@ const FSlateBrush* SChaosVDRecordingControls::GetRecordOrStopButton(EChaosVDReco
 	return bRecordingButtonHovered && IsRecording() ? FChaosVDStyle::Get().GetBrush("StopIcon") : RecordIconBrush;
 }
 
-void SChaosVDRecordingControls::HandleRecordingStop()
+void SChaosVDRecordingControls::HandleRecordingStop(TWeakPtr<FChaosVDSessionInfo> SessionInfo)
 {
-#if WITH_CHAOS_VISUAL_DEBUGGER
 	const TSharedPtr<SChaosVDMainTab> MainTabSharedPtr = MainTabWeakPtr.Pin();
 	if (!MainTabSharedPtr.IsValid())
 	{
 		return;
 	}
 
-	const bool bIsLiveSession = MainTabSharedPtr->GetChaosVDEngineInstance()->HasAnyLiveSessionActive();
+	TSharedPtr<FChaosVDSessionInfo> SessionInfoPtr = SessionInfo.Pin();
+
+	if (!ensure(SessionInfoPtr))
+	{
+		return;
+	}
+
+	FText CurrentTraceTarget = FText::AsCultureInvariant(SessionInfoPtr->LastKnownRecordingState.TraceDetails.TraceTarget);
+
+	const bool bIsLiveSession = SessionInfoPtr->GetRecordingMode() == EChaosVDRecordingMode::Live;
 
 	if (UStatusBarSubsystem* StatusBarSubsystem = GEditor ? GEditor->GetEditorSubsystem<UStatusBarSubsystem>() : nullptr)
 	{
@@ -267,22 +435,21 @@ void SChaosVDRecordingControls::HandleRecordingStop()
 		}
 		else
 		{
-			const FText RecordingPathMessage = FText::Format(LOCTEXT("RecordingSavedPathMessage"," Recoring saved at {0} "), FText::AsCultureInvariant(FChaosVDRuntimeModule::Get().GetLastRecordingFileNamePath()));
+			const FText RecordingPathMessage = FText::Format(LOCTEXT("RecordingSavedPathMessage"," Recording saved at {0} "), CurrentTraceTarget);
 			RecordingPathMessageHandle = StatusBarSubsystem->PushStatusBarMessage(StatusBarID, RecordingPathMessage);
 		}
 	}
-	
-	if (!bIsLiveSession)
+
+	if (!bIsLiveSession && !EnumHasAnyFlags(SessionInfoPtr->GetSessionTypeAttributes(), EChaosVDRemoteSessionAttributes::IsMultiSessionWrapper))
 	{
 		if (FMessageDialog::Open(EAppMsgType::YesNo, LOCTEXT("OpenLastRecordingMessage", "Do you want to load the recorded file now? ")) == EAppReturnType::Yes)
 		{
-			MainTabSharedPtr->GetChaosVDEngineInstance()->LoadRecording(FChaosVDRuntimeModule::Get().GetLastRecordingFileNamePath());
+			MainTabSharedPtr->GetChaosVDEngineInstance()->LoadRecording(CurrentTraceTarget.ToString());
 		}
 	}
-#endif
 }
 
-void SChaosVDRecordingControls::HandleRecordingStart()
+void SChaosVDRecordingControls::HandleRecordingStart(TWeakPtr<FChaosVDSessionInfo> SessionInfo)
 {
 	UStatusBarSubsystem* StatusBarSubsystem = GEditor ? GEditor->GetEditorSubsystem<UStatusBarSubsystem>() : nullptr;
 	if (!StatusBarSubsystem)
@@ -302,59 +469,57 @@ void SChaosVDRecordingControls::HandleRecordingStart()
 		LiveSessionEndedMessageHandle = FStatusBarMessageHandle();
 	}
 
-	RecordingMessageHandle = StatusBarSubsystem->PushStatusBarMessage(StatusBarID, LOCTEXT("RecordingMessgae", "Recording..."));
+	RecordingMessageHandle = StatusBarSubsystem->PushStatusBarMessage(StatusBarID, LOCTEXT("RecordingMessage", "Recording..."));
 }
 
-void SChaosVDRecordingControls::AttemptToConnectToLiveSession()
+void SChaosVDRecordingControls::ExecuteAsyncConnectionAttemptTaskWithRetry(int32 RemainingRetries, TFunction<bool()> RecordingStartAttemptCallback)
 {
-	if(!bAutoConnectionAttemptInProgress)
-	{
-		bAutoConnectionAttemptInProgress = true;
-		PushConnectionAttemptNotification();
-	}
+	TSharedPtr<SNotificationItem> AttemptNotification = PushConnectionAttemptNotification();
 
-	// We need to wait at least one tick before attempting to connect
-	FTSTicker::GetCoreTicker().AddTicker(FTickerDelegate::CreateLambda([WeakThis = AsWeak()](float DeltaTime)
+	ExecutePostRecordingAsyncTaskWithRetry_Internal(RemainingRetries, RecordingStartAttemptCallback, AttemptNotification);
+}
+
+void SChaosVDRecordingControls::ExecutePostRecordingAsyncTaskWithRetry_Internal(int32 RemainingRetries, TFunction<bool()> RecordingStartAttemptCallback, const TSharedPtr<SNotificationItem>& InProgressNotification)
+{
+	// We need to wait at least one tick before attempting to connect to give it time to the trace to be initialized, write to disk, and for the
+	// session manager to hear back from a remote instance
+
+	FTSTicker::GetCoreTicker().AddTicker(FTickerDelegate::CreateLambda([InProgressNotification, RemainingRetries, WeakThis = AsWeak(), RecordingStartAttemptCallback](float DeltaTime)
 	{
 		if (const TSharedPtr<SChaosVDRecordingControls> RecordingControlsPtr = StaticCastSharedPtr<SChaosVDRecordingControls>(WeakThis.Pin()))
 		{
-			if (const TSharedPtr<SChaosVDMainTab> MainTabSharedPtr = RecordingControlsPtr->MainTabWeakPtr.Pin())
+			if (!RecordingStartAttemptCallback)
 			{
-				RecordingControlsPtr->CurrentConnectionAttempts++;
-
-				RecordingControlsPtr->UpdateConnectionAttemptNotification();
-
-				static FString SessionAddress(TEXT("127.0.0.1"));
-			
-				int32 SessionID = 0;
+				RecordingControlsPtr->HandleConnectionAttemptResult(EChaosVDLiveConnectionAttemptResult::Failed, InProgressNotification);
+				UE_LOG(LogChaosVDEditor, Error, TEXT("[%s] Failed to connect to session | Invalid callback provided..."), ANSI_TO_TCHAR(__FUNCTION__));
+				return false;
+			}
 
-				FChaosVDTraceManager::EnumerateActiveSessions(SessionAddress, [&SessionID](const UE::Trace::FStoreClient::FSessionInfo& InSessionInfo)
-				{
-					SessionID = InSessionInfo.GetTraceId();
+			if (const TSharedPtr<SChaosVDMainTab> MainTabSharedPtr = RecordingControlsPtr->MainTabWeakPtr.Pin())
+			{
+				const int32 NewRemainingRetries = RemainingRetries -1;
 
-					// CVD stops all active sessions before staring a recording, so we know our session ID will be first;
-					return false;
-				});
+				RecordingControlsPtr->UpdateConnectionAttemptNotification(InProgressNotification, NewRemainingRetries);
 
 				// CVD needs the trace session name to be able to load a live session. Although the session exist, the session name might not be written right away
 				// Trace files don't really have metadata, it is all part of the same stream, so we need to wait until it is written which might take a few ticks.
 				// Therefore if it is not ready, try again a few times.
-				if (!MainTabSharedPtr->ConnectToLiveSession(SessionID, SessionAddress, EChaosVDLoadRecordedDataMode::SingleSource))
+				if (!RecordingStartAttemptCallback())
 				{
-					if (RecordingControlsPtr->CurrentConnectionAttempts <= RecordingControlsPtr->MaxAutoplayConnectionAttempts)
+					if (NewRemainingRetries > 0)
 					{
 						UE_LOG(LogChaosVDEditor, Verbose, TEXT("[%s] Failed to connect to live session | Attempting again in [%f]..."), ANSI_TO_TCHAR(__FUNCTION__), RecordingControlsPtr->IntervalBetweenAutoplayConnectionAttemptsSeconds);
-						RecordingControlsPtr->AttemptToConnectToLiveSession();
+						RecordingControlsPtr->ExecutePostRecordingAsyncTaskWithRetry_Internal(NewRemainingRetries, RecordingStartAttemptCallback, InProgressNotification);
 					}
 					else
 					{
-						RecordingControlsPtr->HandleConnectionAttemptResult(EChaosVDLiveConnectionAttemptResult::Failed);
-						UE_LOG(LogChaosVDEditor, Error, TEXT("[%s] Failed to connect to live session | [%d] attempts exhausted..."), ANSI_TO_TCHAR(__FUNCTION__), RecordingControlsPtr->MaxAutoplayConnectionAttempts);	
+						RecordingControlsPtr->HandleConnectionAttemptResult(EChaosVDLiveConnectionAttemptResult::Failed, InProgressNotification);
+						UE_LOG(LogChaosVDEditor, Error, TEXT("[%s] Failed to connect to live session | attempts exhausted..."), ANSI_TO_TCHAR(__FUNCTION__));	
 					}
 				}
 				else
 				{
-					RecordingControlsPtr->HandleConnectionAttemptResult(EChaosVDLiveConnectionAttemptResult::Success);
+					RecordingControlsPtr->HandleConnectionAttemptResult(EChaosVDLiveConnectionAttemptResult::Success, InProgressNotification);
 				}
 			}
 		}
@@ -362,72 +527,169 @@ void SChaosVDRecordingControls::AttemptToConnectToLiveSession()
 	}), IntervalBetweenAutoplayConnectionAttemptsSeconds);
 }
 
-FReply SChaosVDRecordingControls::ToggleRecordingState(EChaosVDRecordingMode RecordingMode)
+SChaosVDRecordingControls::FLiveConnectionSettings SChaosVDRecordingControls::GetTargetLiveConnectionSettings(const TSharedRef<FChaosVDSessionInfo>& InSessionInfoRef)
 {
-#if WITH_CHAOS_VISUAL_DEBUGGER
-	if (!IsRecording())
+	FLiveConnectionSettings ConnectionSettings;
+
+	if (!ensure(!EnumHasAnyFlags(InSessionInfoRef->GetSessionTypeAttributes(), EChaosVDRemoteSessionAttributes::IsMultiSessionWrapper)))
+	{
+		UE_LOG(LogChaosVDEditor, Error, TEXT("[%s] Cannot not be called with a multi session wrapper."), ANSI_TO_TCHAR(__FUNCTION__));
+
+		return ConnectionSettings;
+	}
+
+	if (InSessionInfoRef->IsConnected() && InSessionInfoRef->GetRecordingMode() == EChaosVDRecordingMode::Live)
+	{
+		FChaosVDTraceDetails& RecordingSessionDetails = InSessionInfoRef->LastKnownRecordingState.TraceDetails;
+
+		if (const UE::Trace::FStoreClient::FSessionInfo* TraceSessionInfo = FChaosVDTraceManager::GetTraceSessionInfo(RecordingSessionDetails.TraceTarget, RecordingSessionDetails.TraceGuid))
+		{
+			ConnectionSettings.SessionAddress = RecordingSessionDetails.TraceTarget;
+			ConnectionSettings.TraceID = TraceSessionInfo->GetTraceId();
+			return ConnectionSettings;
+		}
+	}
+
+	return ConnectionSettings;
+}
+
+void SChaosVDRecordingControls::ToggleMultiSessionSessionRecordingState(EChaosVDRecordingMode RecordingMode, const TSharedRef<FChaosVDMultiSessionInfo>& InSessionInfoRef)
+{
+	const TSharedPtr<SChaosVDMainTab> MainTabSharedPtr = MainTabWeakPtr.Pin();
+	if (!ensure(MainTabSharedPtr))
+	{
+		return;
+	}
+
+	bool bNewRecordingState = !IsRecording();
+
+	if (bNewRecordingState)
+	{
+		CurrentLoadingMode = EChaosVDLoadRecordedDataMode::MultiSource;
+		MainTabSharedPtr->GetChaosVDEngineInstance()->CloseActiveTraceSessions();
+	}
+
+	InSessionInfoRef->EnumerateInnerSessions([this, RecordingMode, bNewRecordingState](const TSharedRef<FChaosVDSessionInfo>& InInnerSessionRef)
+	{
+		SetSessionRecordingState(bNewRecordingState, RecordingMode, InInnerSessionRef);
+		return true;
+	});
+}
+
+void SChaosVDRecordingControls::ToggleSingleSessionRecordingState(EChaosVDRecordingMode RecordingMode, const TSharedRef<FChaosVDSessionInfo>& SessionInfoRef)
+{
+	SetSessionRecordingState(!SessionInfoRef->IsRecording(), RecordingMode, SessionInfoRef);
+}
+
+void SChaosVDRecordingControls::SetSessionRecordingState(bool bIsRecording, EChaosVDRecordingMode RecordingMode, const TSharedRef<FChaosVDSessionInfo>& SessionInfoRef)
+{
+	TSharedPtr<FChaosVDRemoteSessionsManager> RemoteSessionManager = FChaosVDEngineEditorBridge::Get().GetRemoteSessionsManager();
+	if (!RemoteSessionManager)
+	{
+		UE_LOG(LogChaosVDEditor, Error, TEXT("[%s] Session Manager is not available"), ANSI_TO_TCHAR(__FUNCTION__));
+		return;
+	}
+
+	if (bIsRecording)
 	{
 		TArray<FString, TInlineAllocator<1>> RecordingArgs;
 
 		if (RecordingMode == EChaosVDRecordingMode::Live)
 		{
-			RecordingArgs.Emplace(TEXT("Server"));
+			FChaosVDStartRecordingCommandMessage RecordingParams;
+			RecordingParams.RecordingMode = EChaosVDRecordingMode::Live;
+
+			check(GLog);
+			bool bOutCanBindAll = false;
+			//TODO: Add a way to specify a local address in case we have multiple adapters?
+			TSharedRef<FInternetAddr> LocalIP = ISocketSubsystem::Get(PLATFORM_SOCKETSUBSYSTEM)->GetLocalHostAddr(*GLog, bOutCanBindAll);
 
-			FChaosVDRuntimeModule::Get().StartRecording(RecordingArgs);
+			constexpr bool bAppendPort = false;
+			RecordingParams.Target = LocalIP->ToString(bAppendPort);
 
-			// Only attempt to connect if we managed to start a recording.
-			// If we failed, the runtime module takes care of showing a pop-up error in the editor already
-			if (FChaosVDRuntimeModule::Get().IsRecording())
+			RemoteSessionManager->SendStartRecordingCommand(SessionInfoRef->Address, RecordingParams);
+
+			// Once the start recording command is issue, we can try to connect to the created session (if everything went well).
+			// If it didn't go well, this will take care of update the UI to notify the user
+			constexpr int32 RemainingRetries = 4;
+			ExecuteAsyncConnectionAttemptTaskWithRetry(RemainingRetries, [WeakThis = AsWeak(), SessionInfoAsWeak = SessionInfoRef.ToWeakPtr()]()
 			{
-				AttemptToConnectToLiveSession();
-			}	
+				TSharedPtr<FChaosVDSessionInfo> SessionInfoPtr = SessionInfoAsWeak.Pin();
+
+				if (!SessionInfoPtr)
+				{
+					return false;
+				}
+				
+				TSharedPtr<SChaosVDRecordingControls> Controls =  StaticCastSharedPtr<SChaosVDRecordingControls>(WeakThis.Pin());
+				const TSharedPtr<SChaosVDMainTab> MainTabSharedPtr = Controls ? Controls->MainTabWeakPtr.Pin() : nullptr;
+				if (!MainTabSharedPtr)
+				{
+					return false;
+				}
+
+				FLiveConnectionSettings ConnectionDetails = Controls->GetTargetLiveConnectionSettings(SessionInfoPtr.ToSharedRef());
+				return MainTabSharedPtr->ConnectToLiveSession(ConnectionDetails.TraceID, ConnectionDetails.SessionAddress, Controls->CurrentLoadingMode);
+			});
 		}
 		else
 		{
-			FChaosVDRuntimeModule::Get().StartRecording(RecordingArgs);
+			FChaosVDStartRecordingCommandMessage RecordingParams;
+			RecordingParams.RecordingMode = EChaosVDRecordingMode::File;
+			RemoteSessionManager->SendStartRecordingCommand(SessionInfoRef->Address, RecordingParams);
+
+			// Once the start recording command is issued, we need to check if the recording started, which might take a few frames.
+			// This will take care of retrying, waiting and update the UI to notify the user if needed.
+			constexpr int32 RemainingRetries = 4;
+			ExecuteAsyncConnectionAttemptTaskWithRetry(RemainingRetries, [WeakThis = AsWeak()]()
+			{
+				TSharedPtr<SChaosVDRecordingControls> Controls = StaticCastSharedPtr<SChaosVDRecordingControls>(WeakThis.Pin());
+				return Controls ? Controls->IsRecording() : false;
+			});
 		}
 	}
 	else
 	{
-		FChaosVDRuntimeModule::Get().StopRecording();
+		RemoteSessionManager->SendStopRecordingCommand(SessionInfoRef->Address);
 	}
-#endif
-
-	return FReply::Handled();
 }
 
-bool SChaosVDRecordingControls::IsRecordingToggleButtonEnabled(EChaosVDRecordingMode RecordingMode) const
+FReply SChaosVDRecordingControls::ToggleRecordingState(EChaosVDRecordingMode RecordingMode)
 {
-	if (bAutoConnectionAttemptInProgress)
+	TSharedPtr<FChaosVDSessionInfo> SessionInfoPtr = GetCurrentSessionInfo();
+	if (!SessionInfoPtr)
 	{
-		return false;
+		return FReply::Handled();
 	}
 
-	// Currently we don't support start a recording in standalone mode
-	// but once we add a way to communicate with other running instances we can re-enable these buttons
-	if (FChaosVDModule::IsStandaloneChaosVisualDebugger())
+	if (EnumHasAnyFlags(SessionInfoPtr->GetSessionTypeAttributes(), EChaosVDRemoteSessionAttributes::IsMultiSessionWrapper))
 	{
-		return false;
+		ToggleMultiSessionSessionRecordingState(RecordingMode, StaticCastSharedPtr<FChaosVDMultiSessionInfo>(SessionInfoPtr).ToSharedRef());
 	}
-
-	const TSharedPtr<SChaosVDMainTab> MainTabSharedPtr = MainTabWeakPtr.Pin();
-	if (!MainTabSharedPtr.IsValid())
+	else
 	{
-		return false;
-
+		ToggleSingleSessionRecordingState(RecordingMode, SessionInfoPtr.ToSharedRef());
 	}
 
-	const bool bIsLiveSession = MainTabSharedPtr->GetChaosVDEngineInstance()->HasAnyLiveSessionActive();
+	return FReply::Handled();
+}
 
-	const bool bIsRecording = IsRecording();
+bool SChaosVDRecordingControls::IsRecordingToggleButtonEnabled(EChaosVDRecordingMode RecordingMode) const
+{
+	if (CurrentSelectedSessionId == FChaosVDRemoteSessionsManager::InvalidSessionGUID)
+	{
+		return false;
+	}
 
-	if (RecordingMode == EChaosVDRecordingMode::File)
+	if (!IsRecording())
 	{
-		return (bIsRecording && !bIsLiveSession) || !bIsRecording;
+		return true;
 	}
-	else if (RecordingMode == EChaosVDRecordingMode::Live && GEditor)
+
+	// If we are recording, don't show the stop button for the mode that is disabled
+	if (TSharedPtr<FChaosVDSessionInfo> SessionInfo = GetCurrentSessionInfo())
 	{
-		return (bIsRecording && bIsLiveSession) || (!bIsRecording && GEditor->IsPlayingSessionInEditor());
+		return SessionInfo->GetRecordingMode() == RecordingMode;
 	}
 
 	return false;
@@ -435,10 +697,18 @@ bool SChaosVDRecordingControls::IsRecordingToggleButtonEnabled(EChaosVDRecording
 
 EVisibility SChaosVDRecordingControls::IsRecordingToggleButtonVisible(EChaosVDRecordingMode RecordingMode) const
 {
+	if (!IsRecording())
+	{
+		return EVisibility::Visible;
+	}
+
 	// If we are recording, don't show the stop button for the mode that is disabled
-	const bool bIsRecording = IsRecording();
-	const bool bShouldButtonBeVisible = bIsRecording ? bIsRecording && IsRecordingToggleButtonEnabled(RecordingMode) : true;
-	return bShouldButtonBeVisible ? EVisibility::Visible : EVisibility::Collapsed;
+	if (TSharedPtr<FChaosVDSessionInfo> SessionInfo = GetCurrentSessionInfo())
+	{
+		return SessionInfo->GetRecordingMode() == RecordingMode ? EVisibility::Visible : EVisibility::Collapsed;
+	}
+
+	return EVisibility::Visible;
 }
 
 void SChaosVDRecordingControls::RegisterMenus()
@@ -459,8 +729,30 @@ void SChaosVDRecordingControls::RegisterMenus()
 
 		TSharedRef<SWidget> RecordToFileButton = SNew(SBox).Padding(4.0f,0.0f)[ RecordingControls->GenerateToggleRecordingStateButton(EChaosVDRecordingMode::File, LOCTEXT("RecordToFileButtonDesc", "Starts a recording for the current session, saving it directly to file")) ];
 		TSharedRef<SWidget> RecordToLiveButton = SNew(SBox).Padding(4.0f,0.0f)[ RecordingControls->GenerateToggleRecordingStateButton(EChaosVDRecordingMode::Live, LOCTEXT("RecordLiveButtonDesc", "Starts a recording and automatically connects to it playing it back in real time")) ];
+		TSharedRef<SWidget> SessionsDropdown = SNew(SBox).Padding(4.0f,0.0f)[ RecordingControls->GenerateTargetSessionSelector() ];
 		TSharedRef<SWidget> RecordingTime = RecordingControls->GenerateRecordingTimeTextBlock();
 		TSharedRef<SWidget> DataChannelsButton = SNew(SBox).Padding(4.0f,0.0f)[ RecordingControls->GenerateDataChannelsButton() ];
+		TSharedRef<SWidget> LoadingModeSelector = SNew(SBox).Padding(4.0f,0.0f)[ RecordingControls->GenerateLoadingModeSelector() ];
+
+		InSection.AddSeparator("RecordingControlsDivider");
+
+		InSection.AddEntry(
+			FToolMenuEntry::InitWidget(
+				"AvailableSessions",
+				SessionsDropdown,
+				FText::GetEmpty(),
+				false,
+				false
+			));
+		
+		InSection.AddEntry(
+			FToolMenuEntry::InitWidget(
+				"LoadingModeSelector",
+				LoadingModeSelector,
+				FText::GetEmpty(),
+				false,
+				false
+			));
 
 		InSection.AddEntry(
 			FToolMenuEntry::InitWidget(
@@ -500,74 +792,72 @@ void SChaosVDRecordingControls::RegisterMenus()
 	}));
 }
 
-
 bool SChaosVDRecordingControls::IsRecording() const
 {
-#if WITH_CHAOS_VISUAL_DEBUGGER
-	return FChaosVisualDebuggerTrace::IsTracing();
-#else
-	return false;
-#endif
+	TSharedPtr<FChaosVDSessionInfo> SessionInfo = GetCurrentSessionInfo();
+	return SessionInfo ? SessionInfo->IsRecording() : false;
 }
 
 FText SChaosVDRecordingControls::GetRecordingTimeText() const
 {
-	FNumberFormattingOptions FormatOptions;
-	FormatOptions.MinimumFractionalDigits = 2;
-	FormatOptions.MaximumFractionalDigits = 2;
+	if (TSharedPtr<FChaosVDSessionInfo> SessionInfo = GetCurrentSessionInfo())
+	{
+		FNumberFormattingOptions FormatOptions;
+		FormatOptions.MinimumFractionalDigits = 2;
+		FormatOptions.MaximumFractionalDigits = 2;
 
-#if WITH_CHAOS_VISUAL_DEBUGGER
-	FText SecondsText = FText::AsNumber(FChaosVDRuntimeModule::Get().GetAccumulatedRecordingTime(), &FormatOptions);
-	return FText::Format(LOCTEXT("RecordingTimer","{0} s"), SecondsText);
-#else
-	return FText::GetEmpty();
-#endif
+		FText SecondsText = FText::AsNumber(SessionInfo->LastKnownRecordingState.ElapsedTime, &FormatOptions);
+		return FText::Format(LOCTEXT("RecordingTimer","{0} s"), SecondsText);
+	}
+
+	return LOCTEXT("RecordingTimerError","Failed to get time information");
 }
 
-void SChaosVDRecordingControls::PushConnectionAttemptNotification()
+TSharedPtr<SNotificationItem> SChaosVDRecordingControls::PushConnectionAttemptNotification()
 {
-	FNotificationInfo Info(LOCTEXT("ConnectingToLiceSessionMessge", "Connecting to Live Session ..."));
+	FNotificationInfo Info(LOCTEXT("ConnectingToSessionMessage", "Connecting Session ..."));
 	Info.bFireAndForget = false;
 	Info.FadeOutDuration = 3.0f;
 	Info.ExpireDuration = 0.0f;
 
-	ConnectionAttemptNotification = FSlateNotificationManager::Get().AddNotification(Info);
+	TSharedPtr<SNotificationItem> ConnectionAttemptNotification = FSlateNotificationManager::Get().AddNotification(Info);
 	
 	if (ConnectionAttemptNotification.IsValid())
 	{
 		ConnectionAttemptNotification->SetCompletionState(SNotificationItem::CS_Pending);
+		return ConnectionAttemptNotification;
 	}
+
+	return nullptr;
 }
 
-void SChaosVDRecordingControls::UpdateConnectionAttemptNotification()
+void SChaosVDRecordingControls::UpdateConnectionAttemptNotification(const TSharedPtr<SNotificationItem>& InNotification, int32 AttemptsRemaining)
 {
-	if (ConnectionAttemptNotification.IsValid())
+	if (InNotification)
 	{
-		ConnectionAttemptNotification->SetSubText(FText::FormatOrdered(LOCTEXT("SessionConnectionAttemptSubText", "Attempt {0} / {1}"), CurrentConnectionAttempts, MaxAutoplayConnectionAttempts));
+		InNotification->SetSubText(FText::FormatOrdered(LOCTEXT("SessionConnectionAttemptSubText", "Attempts Remaining {0}"), AttemptsRemaining));
 	}
+	
 }
 
-void SChaosVDRecordingControls::HandleConnectionAttemptResult(EChaosVDLiveConnectionAttemptResult Result)
+void SChaosVDRecordingControls::HandleConnectionAttemptResult(EChaosVDLiveConnectionAttemptResult Result, const TSharedPtr<SNotificationItem>& InNotification)
 {
-	CurrentConnectionAttempts = 0;
-	bAutoConnectionAttemptInProgress = false;
-
-	if (ConnectionAttemptNotification.IsValid())
+	if (InNotification)
 	{
 		if (Result == EChaosVDLiveConnectionAttemptResult::Success)
 		{
-			ConnectionAttemptNotification->SetText(LOCTEXT("SessionConnectionSuccess", "Connected!"));
-			ConnectionAttemptNotification->SetSubText(FText::GetEmpty());
-			ConnectionAttemptNotification->SetCompletionState(SNotificationItem::ECompletionState::CS_Success);
+			InNotification->SetText(LOCTEXT("SessionConnectionSuccess", "Connected!"));
+			InNotification->SetSubText(FText::GetEmpty());
+			InNotification->SetCompletionState(SNotificationItem::ECompletionState::CS_Success);
 		}
 		else
 		{
-			ConnectionAttemptNotification->SetText(LOCTEXT("SessionConnectionFailedText", "Failed to connect"));
-			ConnectionAttemptNotification->SetSubText(LOCTEXT("SessionConnectionFailedSubText", "See the logs for more details..."));
-			ConnectionAttemptNotification->SetCompletionState(SNotificationItem::ECompletionState::CS_Fail);
+			InNotification->SetText(LOCTEXT("SessionConnectionFailedText", "Failed to connect"));
+			InNotification->SetSubText(LOCTEXT("SessionConnectionFailedSubText", "See the logs for more details..."));
+			InNotification->SetCompletionState(SNotificationItem::ECompletionState::CS_Fail);
 		}
 
-		ConnectionAttemptNotification->ExpireAndFadeout();
+		InNotification->ExpireAndFadeout();
 	}
 }
 
diff --git a/Engine/Plugins/ChaosVD/Source/ChaosVD/Private/Widgets/SChaosVDRecordingControls.h b/Engine/Plugins/ChaosVD/Source/ChaosVD/Private/Widgets/SChaosVDRecordingControls.h
--- a/Engine/Plugins/ChaosVD/Source/ChaosVD/Private/Widgets/SChaosVDRecordingControls.h
+++ b/Engine/Plugins/ChaosVD/Source/ChaosVD/Private/Widgets/SChaosVDRecordingControls.h
@@ -3,6 +3,8 @@
 #pragma once
 
 #include "StatusBarSubsystem.h"
+#include "Chaos/ChaosVDRemoteSessionsManager.h"
+#include "ChaosVDEngine.h"
 #include "Templates/SharedPointer.h"
 #include "Widgets/SCompoundWidget.h"
 
@@ -18,6 +20,7 @@ class FReply;
 class SButton;
 class SChaosVDMainTab;
 struct FSlateBrush;
+struct FSlateIcon;
 
 UCLASS()
 class CHAOSVD_API UChaosVDRecordingToolbarMenuContext : public UObject
@@ -29,13 +32,6 @@ public:
 	TWeakPtr<class SChaosVDRecordingControls> RecordingControlsWidget;
 };
 
-UENUM()
-enum class EChaosVDRecordingMode
-{
-	File,
-	Live
-};
-
 UENUM()
 enum class EChaosVDLiveConnectionAttemptResult
 {
@@ -50,7 +46,7 @@ class SChaosVDRecordingControls : public SCompoundWidget
 public:
 	SLATE_BEGIN_ARGS( SChaosVDRecordingControls ){}
 	SLATE_END_ARGS()
-	
+
 	void Construct(const FArguments& InArgs, const TSharedRef<SChaosVDMainTab>& InMainTabSharedRef);
 
 	virtual ~SChaosVDRecordingControls() override;
@@ -58,26 +54,46 @@ public:
 protected:
 
 	TSharedRef<SWidget> GenerateToggleRecordingStateButton(EChaosVDRecordingMode RecordingMode, const FText& StartRecordingTooltip);
+	TSharedRef<SWidget> GenerateTargetSessionSelector();
+	TSharedRef<SWidget> GenerateTargetSessionDropdown();
 	TSharedRef<SWidget> GenerateDataChannelsMenu();
 	TSharedRef<SWidget> GenerateDataChannelsButton();
+	TSharedRef<SWidget> GenerateLoadingModeSelector();
 	TSharedRef<SWidget> GenerateRecordingTimeTextBlock();
 	TSharedRef<SWidget> GenerateToolbarWidget();
 
-#if WITH_CHAOS_VISUAL_DEBUGGER
-	void ToggleChannelEnabledState(TWeakPtr<FCVDDataChannel> Channel);
-	bool IsChannelEnabled(TWeakPtr<FCVDDataChannel> Channel);
-	bool CanChangeChannelEnabledState(TWeakPtr<FCVDDataChannel> Channel);
-#endif
+	EVisibility GetRecordingTimeTextBlockVisibility() const;
+
+	void ToggleChannelEnabledState(FString ChannelName);
+	bool IsChannelEnabled(FString ChannelName);
+	bool CanChangeChannelEnabledState(FString ChannelName);
+
+	void SelectTargetSession(FGuid SessionId);
+
+	FSlateIcon GetIconForSession(FGuid SessionId);
+
+	TSharedPtr<FChaosVDSessionInfo> GetCurrentSessionInfo() const;
+	TSharedPtr<FChaosVDSessionInfo> GetSessionInfo(FGuid Id) const;
 
 	bool HasDataChannelsSupport() const;
+
+	bool CanChangeLoadingMode() const;
 	
 	const FSlateBrush* GetRecordOrStopButton(EChaosVDRecordingMode RecordingMode) const;
 	
-	void HandleRecordingStop();
-	void HandleRecordingStart();
-	void AttemptToConnectToLiveSession();
+	void HandleRecordingStop(TWeakPtr<FChaosVDSessionInfo> SessionInfo);
+	void HandleRecordingStart(TWeakPtr<FChaosVDSessionInfo> SessionInfo);
+
+	void ExecuteAsyncConnectionAttemptTaskWithRetry(int32 RemainingRetries, TFunction<bool()> RecordingStartAttemptCallback);
+
+private:
+	void ExecutePostRecordingAsyncTaskWithRetry_Internal(int32 RemainingRetries, TFunction<bool()> RecordingStartAttemptCallback, const TSharedPtr<SNotificationItem>& InProgressNotification);
+protected:
 
 	FReply ToggleRecordingState(EChaosVDRecordingMode RecordingMode);
+	void ToggleMultiSessionSessionRecordingState(EChaosVDRecordingMode RecordingMode, const TSharedRef<FChaosVDMultiSessionInfo>& InSessionInfoRef);
+	void ToggleSingleSessionRecordingState(EChaosVDRecordingMode RecordingMode, const TSharedRef<FChaosVDSessionInfo>& SessionInfoRef);
+	void SetSessionRecordingState(bool bIsRecording, EChaosVDRecordingMode RecordingMode, const TSharedRef<FChaosVDSessionInfo>& SessionInfoRef);
 
 	bool IsRecordingToggleButtonEnabled(EChaosVDRecordingMode RecordingMode) const;
 	EVisibility IsRecordingToggleButtonVisible(EChaosVDRecordingMode RecordingMode) const;
@@ -88,11 +104,23 @@ protected:
 
 	FText GetRecordingTimeText() const;
 
-	void PushConnectionAttemptNotification();
-	void UpdateConnectionAttemptNotification();
+	TSharedPtr<SNotificationItem> PushConnectionAttemptNotification();
+	void UpdateConnectionAttemptNotification(const TSharedPtr<SNotificationItem>& InNotification, int32 AttemptsRemaining);
+
+	void HandleConnectionAttemptResult(EChaosVDLiveConnectionAttemptResult Result, const TSharedPtr<SNotificationItem>& InNotification);
+
+	struct FLiveConnectionSettings
+	{
+		int32 TraceID = INDEX_NONE;
+		FString SessionAddress;
+	};
+
+	FLiveConnectionSettings GetTargetLiveConnectionSettings(const TSharedRef<FChaosVDSessionInfo>& InSessionInfoRef);
+
+	FText GetCurrentSelectedSessionName() const;
+
+	FGuid CurrentSelectedSessionId = FGuid();
 
-	void HandleConnectionAttemptResult(EChaosVDLiveConnectionAttemptResult Result);
-	
 	FName StatusBarID;
 	
 	FStatusBarMessageHandle RecordingMessageHandle;
@@ -103,16 +131,21 @@ protected:
 
 	TWeakPtr<SChaosVDMainTab> MainTabWeakPtr;
 
-	int32 MaxAutoplayConnectionAttempts = 4;
 	float IntervalBetweenAutoplayConnectionAttemptsSeconds = 1.0f;
-	int32 CurrentConnectionAttempts = 0;
-	bool bAutoConnectionAttemptInProgress = false;
 
 	bool bRecordingButtonHovered = false;
 
 	FCurveSequence RecordingAnimation;
 
-	TSharedPtr<SNotificationItem> ConnectionAttemptNotification;
+	EChaosVDLoadRecordedDataMode CurrentLoadingMode = EChaosVDLoadRecordedDataMode::SingleSource;
+
+	struct FPostRecordingCommandAsyncTask
+	{
+		FGuid OwningSessionId;
+		int32 RemainingAttempts;
+		
+		TFunction<bool()> TaskCallback;
+	};
 
-	static inline const FName RecordingControlsToolbarName = FName("ChaosVD.MainToolBar.RecordingControls");
+	static const FName RecordingControlsToolbarName;
 };
diff --git a/Engine/Plugins/ChaosVD/Source/ChaosVD/Public/Trace/ChaosVDTraceProvider.h b/Engine/Plugins/ChaosVD/Source/ChaosVD/Public/Trace/ChaosVDTraceProvider.h
--- a/Engine/Plugins/ChaosVD/Source/ChaosVD/Public/Trace/ChaosVDTraceProvider.h
+++ b/Engine/Plugins/ChaosVD/Source/ChaosVD/Public/Trace/ChaosVDTraceProvider.h
@@ -134,6 +134,7 @@ public:
 	virtual void EditAccessCheck() const override{};
 	
 	int32 CHAOSVD_API GetRemappedSolverID(int32 SolverID);
+	int32 CHAOSVD_API GetRemappedSolverID_AssumesLocked(int32 SolverID);
 
 	bool DoesOwnRecordingInstance() const
 	{
@@ -177,7 +178,7 @@ private:
 	bool bHasRecordingOverride = false;
 
 	TMap<int32, int32> CurrentNetworkTickOffsets;
-	TSortedMap<int32, int32> RemappedSolversID;
+	TSortedMap<int32, int32> RemappedSolversIDs;
 
 	int32 RemappedGameThreadTrackID = INDEX_NONE;
 
diff --git a/Engine/Source/Runtime/Experimental/Chaos/Private/ChaosVisualDebugger/ChaosVDOptionalDataChannel.cpp b/Engine/Source/Runtime/Experimental/Chaos/Private/ChaosVisualDebugger/ChaosVDOptionalDataChannel.cpp
--- a/Engine/Source/Runtime/Experimental/Chaos/Private/ChaosVisualDebugger/ChaosVDOptionalDataChannel.cpp
+++ b/Engine/Source/Runtime/Experimental/Chaos/Private/ChaosVisualDebugger/ChaosVDOptionalDataChannel.cpp
@@ -50,6 +50,8 @@ namespace Chaos::VisualDebugger
 	void FChaosVDOptionalDataChannel::SetChannelEnabled(bool bNewEnabled)
 	{
 		bIsEnabled = bNewEnabled;
+
+		FChaosVDDataChannelsManager::Get().OnChannelStateChanged().Broadcast(AsWeak());
 	}
 
 	FChaosVDDataChannelsManager& FChaosVDDataChannelsManager::Get()
diff --git a/Engine/Source/Runtime/Experimental/Chaos/Private/ChaosVisualDebugger/ChaosVDSerializedNameTable.cpp b/Engine/Source/Runtime/Experimental/Chaos/Private/ChaosVisualDebugger/ChaosVDSerializedNameTable.cpp
--- a/Engine/Source/Runtime/Experimental/Chaos/Private/ChaosVisualDebugger/ChaosVDSerializedNameTable.cpp
+++ b/Engine/Source/Runtime/Experimental/Chaos/Private/ChaosVisualDebugger/ChaosVDSerializedNameTable.cpp
@@ -9,6 +9,8 @@
 
 namespace Chaos::VisualDebugger
 {
+	FStringView FChaosVDSerializedNameEntry::WrapperTypeName = TEXT("FChaosVDSerializedNameEntry");
+
 	uint64 FChaosVDSerializableNameTable::AddNameToTable(FName Name)
 	{
 		const uint64 NameID = Name.ToUnstableInt();
diff --git a/Engine/Source/Runtime/Experimental/Chaos/Public/ChaosVisualDebugger/ChaosVDOptionalDataChannel.h b/Engine/Source/Runtime/Experimental/Chaos/Public/ChaosVisualDebugger/ChaosVDOptionalDataChannel.h
--- a/Engine/Source/Runtime/Experimental/Chaos/Public/ChaosVisualDebugger/ChaosVDOptionalDataChannel.h
+++ b/Engine/Source/Runtime/Experimental/Chaos/Public/ChaosVisualDebugger/ChaosVDOptionalDataChannel.h
@@ -5,6 +5,7 @@
 #if WITH_CHAOS_VISUAL_DEBUGGER
 
 #include "Templates/SharedPointer.h"
+#include "Delegates/Delegate.h"
 
 class FText;
 class FName;
@@ -98,6 +99,16 @@ namespace Chaos::VisualDebugger
 			return nullptr;
 		}
 
+	private:
+
+		DECLARE_MULTICAST_DELEGATE_OneParam(FChannelStateChangedDelegate, TWeakPtr<FChaosVDOptionalDataChannel> DataChannelChanged);
+		FChannelStateChangedDelegate ChannelStateChangedDelegate;
+
+	public:
+		FChannelStateChangedDelegate& OnChannelStateChanged()
+		{
+			return ChannelStateChangedDelegate;
+		}
 	private:
 		
 		/** Registers a CVD Data channel instance - Do not call directly */
diff --git a/Engine/Source/Runtime/Experimental/Chaos/Public/ChaosVisualDebugger/ChaosVDSerializedNameTable.h b/Engine/Source/Runtime/Experimental/Chaos/Public/ChaosVisualDebugger/ChaosVDSerializedNameTable.h
--- a/Engine/Source/Runtime/Experimental/Chaos/Public/ChaosVisualDebugger/ChaosVDSerializedNameTable.h
+++ b/Engine/Source/Runtime/Experimental/Chaos/Public/ChaosVisualDebugger/ChaosVDSerializedNameTable.h
@@ -18,7 +18,7 @@ namespace Chaos::VisualDebugger
 	/** Serialized Name entry that can be loaded by its recorded ID */
 	struct CHAOS_API FChaosVDSerializedNameEntry
 	{
-		inline static FStringView WrapperTypeName = TEXT("FChaosVDSerializedNameEntry");
+		static FStringView WrapperTypeName;
 
 		uint64 NameID = 0;
 		int32 NameNumber = NAME_NO_NUMBER;
diff --git a/Engine/Source/Runtime/Experimental/ChaosSolverEngine/ChaosSolverEngine.Build.cs b/Engine/Source/Runtime/Experimental/ChaosSolverEngine/ChaosSolverEngine.Build.cs
--- a/Engine/Source/Runtime/Experimental/ChaosSolverEngine/ChaosSolverEngine.Build.cs
+++ b/Engine/Source/Runtime/Experimental/ChaosSolverEngine/ChaosSolverEngine.Build.cs
@@ -17,11 +17,13 @@ namespace UnrealBuildTool.Rules
 					"DeveloperSettings", 
 					"DataflowCore",
 					"DataflowEngine",
-					"DataflowSimulation"
+					"DataflowSimulation",
+					"ChaosVDRuntime" // If CVD is compiled out, this only provides access to its structs data types
 				}
 				);
 
 			SetupModulePhysicsSupport(Target);
+			SetupModuleChaosVisualDebuggerSupport(Target);
 			PrivateDefinitions.Add("CHAOS_INCLUDE_LEVEL_1=1");
 		}
 	}
diff --git a/Engine/Source/Runtime/Experimental/ChaosSolverEngine/Private/Chaos/ChaosSolverEnginePlugin.cpp b/Engine/Source/Runtime/Experimental/ChaosSolverEngine/Private/Chaos/ChaosSolverEnginePlugin.cpp
--- a/Engine/Source/Runtime/Experimental/ChaosSolverEngine/Private/Chaos/ChaosSolverEnginePlugin.cpp
+++ b/Engine/Source/Runtime/Experimental/ChaosSolverEngine/Private/Chaos/ChaosSolverEnginePlugin.cpp
@@ -10,7 +10,7 @@
 #include "Chaos/ChaosDebugDrawSubsystem.h"
 #include "Chaos/ChaosSolverActor.h"
 #include "ChaosSolversModule.h"
-#include "ChaosVDRecordingStateScreenMessageHandler.h"
+#include "Chaos/ChaosVDEngineEditorBridge.h"
 
 class FChaosSolverEnginePlugin : public IChaosSolverEnginePlugin
 {
@@ -34,7 +34,10 @@ void FChaosSolverEnginePlugin::StartupModule()
 #endif
 
 #if WITH_CHAOS_VISUAL_DEBUGGER
-	FChaosVDRecordingStateScreenMessageHandler::Get().Initialize();
+	FCoreDelegates::OnFEngineLoopInitComplete.AddLambda([]()
+	{
+		FChaosVDEngineEditorBridge::Get().Initialize();
+	});
 #endif
 
 }
@@ -46,7 +49,7 @@ void FChaosSolverEnginePlugin::ShutdownModule()
 #endif
 
 #if WITH_CHAOS_VISUAL_DEBUGGER
-	FChaosVDRecordingStateScreenMessageHandler::Get().TearDown();
+	FChaosVDEngineEditorBridge::Get().TearDown();
 #endif
 
 }
diff --git a/Engine/Source/Runtime/Experimental/ChaosSolverEngine/Private/Chaos/ChaosVDEngineEditorBridge.cpp b/Engine/Source/Runtime/Experimental/ChaosSolverEngine/Private/Chaos/ChaosVDEngineEditorBridge.cpp
new file mode 100644
--- /dev/null
+++ b/Engine/Source/Runtime/Experimental/ChaosSolverEngine/Private/Chaos/ChaosVDEngineEditorBridge.cpp
@@ -0,0 +1,206 @@
+// Copyright Epic Games, Inc. All Rights Reserved.
+
+#include "Chaos/ChaosVDEngineEditorBridge.h"
+
+#include "GameFramework/PlayerController.h"
+
+#if WITH_CHAOS_VISUAL_DEBUGGER
+
+#include "ChaosVDRuntimeModule.h"
+#include "ChaosVisualDebugger/ChaosVDTraceMacros.h"
+#include "ChaosVisualDebugger/ChaosVisualDebuggerTrace.h"
+#include "DataWrappers/ChaosVDCollisionDataWrappers.h"
+#include "Engine/Engine.h"
+#include "Serialization/MemoryWriter.h"
+
+void FChaosVDEngineEditorBridge::AddOnScreenRecordingMessage()
+{
+	if (!GEngine)
+	{
+		return;
+	}
+
+	static FText ChaosVDRecordingStartedMessage = NSLOCTEXT("ChaosVisualDebugger", "OnScreenChaosVDRecordingStartedMessage", "Chaos Visual Debugger recording in progress...");
+
+	if (CVDRecordingMessageKey == 0)
+	{
+		CVDRecordingMessageKey = GetTypeHash(ChaosVDRecordingStartedMessage.ToString());
+	}
+	
+	// Add a long duration value, we will remove the message manually when the recording stops
+	constexpr float MessageDurationSeconds = 3600.0f;
+	GEngine->AddOnScreenDebugMessage(CVDRecordingMessageKey, MessageDurationSeconds, FColor::Red,ChaosVDRecordingStartedMessage.ToString());
+}
+
+void FChaosVDEngineEditorBridge::RemoveOnScreenRecordingMessage()
+{
+	if (!GEngine)
+	{
+		return;
+	}
+
+	if (CVDRecordingMessageKey != 0)
+	{
+		GEngine->RemoveOnScreenDebugMessage(CVDRecordingMessageKey);
+	}
+}
+
+void FChaosVDEngineEditorBridge::HandleCVDRecordingStarted()
+{
+	SerializeCollisionChannelsNames();
+
+	AddOnScreenRecordingMessage();
+
+	BroadcastSessionStatus(FApp::GetDeltaTime());
+
+	constexpr float UpdateInterval = 0.5f;
+	RecordingStatusUpdateHandle = FTSTicker::GetCoreTicker().AddTicker(FTickerDelegate::CreateRaw(this, &FChaosVDEngineEditorBridge::BroadcastSessionStatus), UpdateInterval);
+}
+
+void FChaosVDEngineEditorBridge::HandleCVDRecordingStopped()
+{
+	RemoveOnScreenRecordingMessage();
+
+	FTSTicker::GetCoreTicker().RemoveTicker(RecordingStatusUpdateHandle);
+	BroadcastSessionStatus(FApp::GetDeltaTime());
+}
+
+void FChaosVDEngineEditorBridge::HandleCVDRecordingStartFailed(const FText& InFailureReason) const
+{
+#if !WITH_EDITOR
+	// In non-editor builds we don't have an error pop-up, therefore we want to show the error message on screen
+	FText ErrorMessage = FText::FormatOrdered(NSLOCTEXT("ChaosVisualDebugger","StartRecordingFailedOnScreenMessage", "Failed to start CVD recording. {0}"), InFailureReason);
+
+	constexpr float MessageDurationSeconds = 4.0f;
+	GEngine->AddOnScreenDebugMessage(CVDRecordingMessageKey, MessageDurationSeconds, FColor::Red, ErrorMessage.ToString());
+#endif
+}
+
+void FChaosVDEngineEditorBridge::HandlePIEStarted(UGameInstance* GameInstance)
+{
+	// If we were already recording, show the message
+	if (FChaosVDRuntimeModule::Get().IsRecording())
+	{
+		HandleCVDRecordingStarted();
+	}
+}
+
+void FChaosVDEngineEditorBridge::HandleDataChannelChanged(TWeakPtr<FChaosVDDataDataChannel> ChannelWeakPtr)
+{
+	if (TSharedPtr<FChaosVDDataDataChannel> DataChannelPtr = ChannelWeakPtr.Pin())
+	{
+		FChaosVDChannelStateChangeResponseMessage NewChannelState;
+		NewChannelState.InstanceID = FApp::GetInstanceId();
+		NewChannelState.NewState.bIsEnabled = DataChannelPtr->IsChannelEnabled();
+		NewChannelState.NewState.ChannelName = DataChannelPtr->GetId().ToString();
+		NewChannelState.NewState.bCanChangeChannelState = DataChannelPtr->CanChangeEnabledState();
+		
+		RemoteSessionsManager->PublishDataChannelStateChangeUpdate(NewChannelState);
+	}
+}
+
+void FChaosVDEngineEditorBridge::SerializeCollisionChannelsNames()
+{
+	TArray<uint8> CollisionChannelsDataBuffer;
+	FMemoryWriter MemWriterAr(CollisionChannelsDataBuffer);
+
+	FChaosVDCollisionChannelsInfoContainer CollisionChannelInfoContainer;
+
+	if (UCollisionProfile* CollisionProfileData = UCollisionProfile::Get())
+	{
+		constexpr int32 MaxSupportedChannels = 32;
+		for (int32 ChannelIndex = 0; ChannelIndex < MaxSupportedChannels; ++ChannelIndex)
+		{
+			FChaosVDCollisionChannelInfo Info;
+			Info.DisplayName = CollisionProfileData->ReturnChannelNameFromContainerIndex(ChannelIndex).ToString();
+			Info.CollisionChannel = ChannelIndex;
+			Info.bIsTraceType = CollisionProfileData->ConvertToTraceType(static_cast<ECollisionChannel>(ChannelIndex)) != TraceTypeQuery_MAX;
+			CollisionChannelInfoContainer.CustomChannelsNames[ChannelIndex] = Info;
+		}
+	}
+
+	Chaos::VisualDebugger::WriteDataToBuffer(CollisionChannelsDataBuffer, CollisionChannelInfoContainer);
+
+	CVD_TRACE_BINARY_DATA(CollisionChannelsDataBuffer, FChaosVDCollisionChannelsInfoContainer::WrapperTypeName, EChaosVDTraceBinaryDataOptions::ForceTrace)
+}
+
+FChaosVDEngineEditorBridge& FChaosVDEngineEditorBridge::Get()
+{
+	static FChaosVDEngineEditorBridge CVDEngineEditorBridge;
+	return CVDEngineEditorBridge;
+}
+
+bool FChaosVDEngineEditorBridge::BroadcastSessionStatus(float DeltaTime)
+{
+	FChaosVDRuntimeModule& RuntimeModule = FChaosVDRuntimeModule::Get();
+
+	FChaosVDRecordingStatusMessage StatusMessage;
+	StatusMessage.InstanceId = FApp::GetInstanceId();
+	StatusMessage.bIsRecording = RuntimeModule.IsRecording();
+	StatusMessage.ElapsedTime = RuntimeModule.GetAccumulatedRecordingTime();
+	StatusMessage.TraceDetails = RuntimeModule.GetCurrentTraceSessionDetails();
+
+	RemoteSessionsManager->PublishRecordingStatusUpdate(StatusMessage);
+
+	return true;
+}
+
+void FChaosVDEngineEditorBridge::Initialize()
+{
+	FChaosVDRuntimeModule& Module = FChaosVDRuntimeModule::Get();
+
+	FChaosVDTraceDetails TraceDetails = Module.GetCurrentTraceSessionDetails();
+
+	RemoteSessionsManager->Initialize();
+
+	RecordingStartedHandle = FChaosVDRuntimeModule::Get().RegisterRecordingStartedCallback(FChaosVDRecordingStateChangedDelegate::FDelegate::CreateRaw(this, &FChaosVDEngineEditorBridge::HandleCVDRecordingStarted));
+	RecordingStoppedHandle = FChaosVDRuntimeModule::Get().RegisterRecordingStopCallback(FChaosVDRecordingStateChangedDelegate::FDelegate::CreateRaw(this, &FChaosVDEngineEditorBridge::HandleCVDRecordingStopped));
+	RecordingStartFailedHandle = FChaosVDRuntimeModule::Get().RegisterRecordingStartFailedCallback(FChaosVDRecordingStartFailedDelegate::FDelegate::CreateRaw(this, &FChaosVDEngineEditorBridge::HandleCVDRecordingStartFailed));
+
+	Chaos::VisualDebugger::FChaosVDDataChannelsManager::Get().OnChannelStateChanged().AddRaw(this, &FChaosVDEngineEditorBridge::HandleDataChannelChanged);
+
+#if WITH_EDITOR
+	PIEStartedHandle = FWorldDelegates::OnPIEStarted.AddRaw(this, &FChaosVDEngineEditorBridge::HandlePIEStarted);
+#endif
+
+	// If we were already recording, show the message
+	if (FChaosVDRuntimeModule::Get().IsRecording())
+	{
+		HandleCVDRecordingStarted();
+	}
+}
+
+void FChaosVDEngineEditorBridge::TearDown()
+{
+	FTSTicker::GetCoreTicker().RemoveTicker(RecordingStatusUpdateHandle);
+
+	// Note: This works during engine shutdown because the Module Manager doesn't free the dll on module unload to account for use cases like this
+	// If this appears in a callstack crash it means that assumption changed or was not correct to begin with.
+	// A possible solution is just check if the module is loaded querying the module manager just using the module's name
+	if (FChaosVDRuntimeModule::IsLoaded())
+	{
+		FChaosVDRuntimeModule::Get().RemoveRecordingStartedCallback(RecordingStartedHandle);
+		FChaosVDRuntimeModule::Get().RemoveRecordingStopCallback(RecordingStoppedHandle);
+
+#if WITH_EDITOR
+		 FWorldDelegates::OnPIEStarted.Remove(PIEStartedHandle);
+#endif
+
+		// Make sure of removing the message from the screen in case the recording didn't quite stopped yet
+		if (FChaosVDRuntimeModule::Get().IsRecording())
+		{
+			HandleCVDRecordingStopped();
+		}
+
+		RemoteSessionsManager->Shutdown();
+	}
+}
+#else
+
+FChaosVDEngineEditorBridge& FChaosVDEngineEditorBridge::Get()
+{
+	static FChaosVDEngineEditorBridge CVDEngineEditorBridge;
+	return CVDEngineEditorBridge;
+}
+
+#endif
diff --git a/Engine/Source/Runtime/Experimental/ChaosSolverEngine/Private/Chaos/ChaosVDRecordingStateScreenMessageHandler.cpp b/Engine/Source/Runtime/Experimental/ChaosSolverEngine/Private/Chaos/ChaosVDRecordingStateScreenMessageHandler.cpp
deleted file mode 100644
diff --git a/Engine/Source/Runtime/Experimental/ChaosSolverEngine/Private/Chaos/ChaosVDRecordingStateScreenMessageHandler.h b/Engine/Source/Runtime/Experimental/ChaosSolverEngine/Private/Chaos/ChaosVDRecordingStateScreenMessageHandler.h
deleted file mode 100644
diff --git a/Engine/Source/Runtime/Experimental/ChaosSolverEngine/Private/Chaos/ChaosVDRemoteSessionsManager.cpp b/Engine/Source/Runtime/Experimental/ChaosSolverEngine/Private/Chaos/ChaosVDRemoteSessionsManager.cpp
new file mode 100644
--- /dev/null
+++ b/Engine/Source/Runtime/Experimental/ChaosSolverEngine/Private/Chaos/ChaosVDRemoteSessionsManager.cpp
@@ -0,0 +1,552 @@
+// Copyright Epic Games, Inc. All Rights Reserved.
+
+#include "Chaos/ChaosVDRemoteSessionsManager.h"
+
+#include "ChaosVDRuntimeModule.h"
+#include "IMessageBus.h"
+#include "IMessagingModule.h"
+#include "MessageEndpoint.h"
+#include "MessageEndpointBuilder.h"
+#include "ChaosVisualDebugger/ChaosVDOptionalDataChannel.h"
+#include "Misc/App.h"
+
+const FGuid FChaosVDRemoteSessionsManager::AllRemoteSessionsWrapperGUID = FGuid::NewGuid();
+const FGuid FChaosVDRemoteSessionsManager::AllSessionsWrapperGUID = FGuid::NewGuid();
+const FGuid FChaosVDRemoteSessionsManager::InvalidSessionGUID = FGuid();
+
+const FString FChaosVDRemoteSessionsManager::LocalEditorSessionName = TEXT("Local Editor");
+const FGuid FChaosVDRemoteSessionsManager::LocalEditorSessionID = IsController() ? FApp::GetInstanceId() : InvalidSessionGUID;
+const FName FChaosVDRemoteSessionsManager::MessageBusEndPointName = FName("CVDSessionManagerEndPoint");
+const FString FChaosVDRemoteSessionsManager::AllRemoteSessionsTargetName = TEXT("AllRemote");
+const FString FChaosVDRemoteSessionsManager::AllSessionsTargetName = TEXT("AllSessions");
+
+DEFINE_LOG_CATEGORY(LogChaosVDRemoteSession)
+
+bool FChaosVDSessionInfo::IsRecording() const
+{
+	return LastKnownRecordingState.bIsRecording;
+}
+
+EChaosVDRecordingMode FChaosVDSessionInfo::GetRecordingMode() const
+{
+	return LastKnownRecordingState.TraceDetails.Mode;
+}
+
+bool FChaosVDSessionInfo::IsConnected() const
+{
+	return LastKnownRecordingState.TraceDetails.bIsConnected;
+}
+
+bool FChaosVDMultiSessionInfo::IsRecording() const
+{
+	bool bIsRecording = false;
+	EnumerateInnerSessions([&bIsRecording](const TSharedRef<FChaosVDSessionInfo>& InSessionRef)
+	{
+		if (InSessionRef->IsRecording())
+		{
+			bIsRecording = true;
+			return false;
+		}
+		return true;
+	});
+
+	return bIsRecording;
+}
+
+EChaosVDRecordingMode FChaosVDMultiSessionInfo::GetRecordingMode() const
+{
+	EChaosVDRecordingMode FirstInstanceRecordingMode = EChaosVDRecordingMode::Invalid;
+	EnumerateInnerSessions([&FirstInstanceRecordingMode](const TSharedRef<FChaosVDSessionInfo>& InSessionRef)
+	{
+		FirstInstanceRecordingMode = InSessionRef->GetRecordingMode();
+		return false;
+	});
+
+	return FirstInstanceRecordingMode;
+}
+
+bool FChaosVDMultiSessionInfo::IsConnected() const
+{
+	bool bIsAnyConnected = false;
+	EnumerateInnerSessions([&bIsAnyConnected](const TSharedRef<FChaosVDSessionInfo>& InSessionRef)
+	{
+		if (InSessionRef->IsRecording())
+		{
+			bIsAnyConnected = true;
+			return false;
+		}
+		return true;
+	});
+
+	return bIsAnyConnected;
+}
+
+FChaosVDRemoteSessionsManager::FChaosVDRemoteSessionsManager()
+{
+}
+
+constexpr bool FChaosVDRemoteSessionsManager::IsController()
+{
+#if WITH_EDITOR
+	return true;
+#else
+	return false;
+#endif
+}
+
+TWeakPtr<FChaosVDSessionInfo> FChaosVDRemoteSessionsManager::GetSessionInfo(FGuid Id)
+{
+	if (TSharedPtr<FChaosVDSessionInfo>* FoundSessionPtrPtr = ActiveSessionsByInstanceId.Find(Id))
+	{
+		return *FoundSessionPtrPtr;
+	}
+
+	return nullptr;
+}
+
+TSharedPtr<FChaosVDSessionInfo> FChaosVDRemoteSessionsManager::CreatedWrapperSessionInfo(FGuid InstanceId, const FString& SessionName)
+{
+	TSharedPtr<FChaosVDMultiSessionInfo> NewSessionInfo = MakeShared<FChaosVDMultiSessionInfo>();
+	
+	NewSessionInfo->InstanceId = InstanceId;
+	NewSessionInfo->SessionName = SessionName;
+
+	return NewSessionInfo;
+}
+
+TSharedPtr<FMessageEndpoint> FChaosVDRemoteSessionsManager::CreateEndPoint(const TSharedRef<IMessageBus>& InMessageBus)
+{
+	FMessageEndpointBuilder EndPointBuilder = FMessageEndpoint::Builder(MessageBusEndPointName, InMessageBus)
+	                                          .Handling<FChaosVDSessionPing>(this, &FChaosVDRemoteSessionsManager::HandleSessionPingMessage)
+	                                          .Handling<FChaosVDStartRecordingCommandMessage>(this, &FChaosVDRemoteSessionsManager::HandleRecordingStartCommandMessage)
+	                                          .Handling<FChaosVDStopRecordingCommandMessage>(this, &FChaosVDRemoteSessionsManager::HandleRecordingStopCommandMessage)
+	                                          .Handling<FChaosVDChannelStateChangeCommandMessage>(this, &FChaosVDRemoteSessionsManager::HandleChangeDataChannelStateCommandMessage)
+	                                          .Handling<FChaosVDFullSessionInfoRequestMessage>(this, &FChaosVDRemoteSessionsManager::HandleFullSessionStateRequestMessage);
+
+	if (IsController())
+	{
+		EndPointBuilder.Handling<FChaosVDSessionPong>(this, &FChaosVDRemoteSessionsManager::HandleSessionPongMessage)
+						.Handling<FChaosVDRecordingStatusMessage>(this, &FChaosVDRemoteSessionsManager::HandleRecordingStatusUpdateMessage)
+						.Handling<FChaosVDFullSessionInfoResponseMessage>(this, &FChaosVDRemoteSessionsManager::HandleFullSessionStateResponseMessage)
+						.Handling<FChaosVDChannelStateChangeResponseMessage>(this, &FChaosVDRemoteSessionsManager::HandleChangeDataChannelStateResponseMessage)
+						.Handling<FChaosVDChannelStateChangeResponseMessage>(this, &FChaosVDRemoteSessionsManager::HandleChangeDataChannelStateResponseMessage);
+	}
+
+	return EndPointBuilder.Build();
+}
+
+void FChaosVDRemoteSessionsManager::Initialize()
+{
+	TSharedPtr<IMessageBus> CVDMessageBus = IMessagingModule::Get().GetDefaultBus();
+
+	if (!ensure(CVDMessageBus))
+	{
+		return;
+	}
+
+	MessageBusPtr = CVDMessageBus;
+	MessageEndpoint = CreateEndPoint(CVDMessageBus.ToSharedRef());
+	
+	if (!ensure(MessageEndpoint))
+	{
+		return;
+	}
+
+	if (IsController())
+	{
+		MessageEndpoint->Subscribe<FChaosVDSessionPong>();
+		MessageEndpoint->Subscribe<FChaosVDRecordingStatusMessage>();
+		MessageEndpoint->Subscribe<FChaosVDFullSessionInfoResponseMessage>();
+		MessageEndpoint->Subscribe<FChaosVDChannelStateChangeResponseMessage>();
+	}
+
+	MessageEndpoint->Subscribe<FChaosVDSessionPing>();
+	MessageEndpoint->Subscribe<FChaosVDStartRecordingCommandMessage>();
+	MessageEndpoint->Subscribe<FChaosVDStopRecordingCommandMessage>();
+	MessageEndpoint->Subscribe<FChaosVDChannelStateChangeCommandMessage>();
+	MessageEndpoint->Subscribe<FChaosVDFullSessionInfoRequestMessage>();
+
+	ActiveSessionsByInstanceId.Add(AllRemoteSessionsWrapperGUID, CreatedWrapperSessionInfo(AllRemoteSessionsWrapperGUID, AllRemoteSessionsTargetName));
+	ActiveSessionsByInstanceId.Add(AllSessionsWrapperGUID, CreatedWrapperSessionInfo(AllSessionsWrapperGUID, AllSessionsTargetName));
+
+	constexpr float TickInterval = 1.0f;
+	TickHandle = FTSTicker::GetCoreTicker().AddTicker(FTickerDelegate::CreateRaw(this, &FChaosVDRemoteSessionsManager::Tick), TickInterval);
+}
+
+void FChaosVDRemoteSessionsManager::Shutdown()
+{
+	if (MessageEndpoint)
+	{
+		MessageEndpoint->Unsubscribe<FChaosVDSessionPong>();
+		MessageEndpoint->Unsubscribe<FChaosVDRecordingStatusMessage>();
+		MessageEndpoint->Unsubscribe<FChaosVDSessionPing>();
+		MessageEndpoint->Unsubscribe<FChaosVDStartRecordingCommandMessage>();
+		MessageEndpoint->Unsubscribe<FChaosVDStopRecordingCommandMessage>();
+		MessageEndpoint->Unsubscribe<FChaosVDChannelStateChangeCommandMessage>();
+		MessageEndpoint->Unsubscribe<FChaosVDFullSessionInfoRequestMessage>();
+		MessageEndpoint->Unsubscribe<FChaosVDFullSessionInfoResponseMessage>();
+		MessageEndpoint->Unsubscribe<FChaosVDChannelStateChangeResponseMessage>();
+	}
+
+	FTSTicker::GetCoreTicker().RemoveTicker(TickHandle);
+}
+
+void FChaosVDRemoteSessionsManager::PublishRecordingStatusUpdate(const FChaosVDRecordingStatusMessage& UpdateMessage)
+{
+	if (ensure(MessageEndpoint))
+	{
+		MessageEndpoint->Publish(FMessageEndpoint::MakeMessage<FChaosVDRecordingStatusMessage>(UpdateMessage), EMessageScope::Network);
+	}
+}
+
+void FChaosVDRemoteSessionsManager::PublishDataChannelStateChangeUpdate(const FChaosVDChannelStateChangeResponseMessage& InNewStateData)
+{
+	if (ensure(MessageEndpoint))
+	{
+		MessageEndpoint->Publish(FMessageEndpoint::MakeMessage<FChaosVDChannelStateChangeResponseMessage>(InNewStateData), EMessageScope::Network);
+	}
+}
+
+void FChaosVDRemoteSessionsManager::SendStartRecordingCommand(const FMessageAddress& InDestinationAddress, const FChaosVDStartRecordingCommandMessage& RecordingStartCommandParams)
+{
+	if (!MessageEndpoint)
+	{
+		UE_LOG(LogChaosVDRemoteSession, Error, TEXT("[%s] Failed to send command | Invalid endpoint."), ANSI_TO_TCHAR(__FUNCTION__));
+		return;
+	}
+
+	MessageEndpoint->Send(
+	FMessageEndpoint::MakeMessage<FChaosVDStartRecordingCommandMessage>(RecordingStartCommandParams),
+	FChaosVDStartRecordingCommandMessage::StaticStruct(),
+	EMessageFlags::Reliable,
+	nullptr,
+	TArrayBuilder<FMessageAddress>().Add(InDestinationAddress),
+	FTimespan::Zero(),
+	FDateTime::MaxValue());
+}
+
+void FChaosVDRemoteSessionsManager::SendStopRecordingCommand(const FMessageAddress& InDestinationAddress)
+{
+	if (!MessageEndpoint)
+	{
+		UE_LOG(LogChaosVDRemoteSession, Error, TEXT("[%s] Failed to send command | Invalid endpoint."), ANSI_TO_TCHAR(__FUNCTION__));
+		return;
+	}
+
+	MessageEndpoint->Send(
+	FMessageEndpoint::MakeMessage<FChaosVDStopRecordingCommandMessage>(),
+	FChaosVDStopRecordingCommandMessage::StaticStruct(),
+	EMessageFlags::Reliable,
+	nullptr,
+	TArrayBuilder<FMessageAddress>().Add(InDestinationAddress),
+	FTimespan::Zero(),
+	FDateTime::MaxValue());
+}
+
+void FChaosVDRemoteSessionsManager::SendDataChannelStateChangeCommand(const FMessageAddress& InDestinationAddress,const FChaosVDChannelStateChangeCommandMessage& InNewStateData)
+{
+	if (!MessageEndpoint)
+	{
+		UE_LOG(LogChaosVDRemoteSession, Error, TEXT("[%s] Failed to send command | Invalid endpoint."), ANSI_TO_TCHAR(__FUNCTION__));
+		return;
+	}
+
+	MessageEndpoint->Send(
+	FMessageEndpoint::MakeMessage<FChaosVDChannelStateChangeCommandMessage>(InNewStateData),
+	FChaosVDChannelStateChangeCommandMessage::StaticStruct(),
+	EMessageFlags::Reliable,
+	nullptr,
+	TArrayBuilder<FMessageAddress>().Add(InDestinationAddress),
+	FTimespan::Zero(),
+	FDateTime::MaxValue());
+}
+
+void FChaosVDRemoteSessionsManager::SendFullSessionStateRequestCommand(const FMessageAddress& InDestinationAddress)
+{
+	if (!MessageEndpoint)
+	{
+		UE_LOG(LogChaosVDRemoteSession, Error, TEXT("[%s] Failed to send command | Invalid endpoint."), ANSI_TO_TCHAR(__FUNCTION__));
+		return;
+	}
+
+	MessageEndpoint->Send(
+	FMessageEndpoint::MakeMessage<FChaosVDFullSessionInfoRequestMessage>(),
+	FChaosVDFullSessionInfoRequestMessage::StaticStruct(),
+	EMessageFlags::Reliable,
+	nullptr,
+	TArrayBuilder<FMessageAddress>().Add(InDestinationAddress),
+	FTimespan::Zero(),
+	FDateTime::MaxValue());
+}
+
+bool FChaosVDRemoteSessionsManager::Tick(float DeltaTime)
+{
+	if (IsController())
+	{
+		SendPing();
+		RemoveExpiredSessions();
+	}
+
+	return true;
+}
+
+void FChaosVDRemoteSessionsManager::SendPing()
+{
+	if (ensure(MessageEndpoint))
+	{
+		if (FChaosVDSessionPing* SessionPingData = FMessageEndpoint::MakeMessage<FChaosVDSessionPing>())
+		{
+			SessionPingData->ControllerInstanceId = FApp::GetInstanceId();
+			MessageEndpoint->Publish(SessionPingData, EMessageScope::Network);
+		}
+	}
+}
+
+void FChaosVDRemoteSessionsManager::SendPong(const FChaosVDSessionPing& InMessage)
+{
+	if (!ensure(MessageEndpoint))
+	{
+		return;
+	}
+	
+	if (FChaosVDSessionPong* PongMessage = FMessageEndpoint::MakeMessage<FChaosVDSessionPong>())
+	{
+		PongMessage->InstanceId = FApp::GetInstanceId();
+		PongMessage->SessionId = FApp::GetSessionId();
+
+		if (InMessage.ControllerInstanceId == PongMessage->InstanceId)
+		{
+			PongMessage->SessionName = LocalEditorSessionName;
+		}
+		else
+		{
+			FString AppSessionName = FApp::GetSessionName();
+			PongMessage->SessionName = AppSessionName == TEXT("None") || AppSessionName.IsEmpty() ? FString::Format(TEXT("{0} {1} {2}"), { FApp::GetProjectName(), FString(LexToString(FApp::GetBuildTargetType())), FString::FromInt(FPlatformProcess::GetCurrentProcessId()) }) : AppSessionName;	
+		}
+
+		MessageEndpoint->Publish(PongMessage, EMessageScope::Network);
+	}
+}
+
+void FChaosVDRemoteSessionsManager::RegisterSessionInMultiSessionWrapper(const TSharedRef<FChaosVDSessionInfo>& InSessionInfoRef)
+{
+	if (InSessionInfoRef->SessionName != LocalEditorSessionName )
+	{
+		StaticCastSharedPtr<FChaosVDMultiSessionInfo>(ActiveSessionsByInstanceId.FindChecked(AllRemoteSessionsWrapperGUID))->InnerSessionsByInstanceID.Emplace(InSessionInfoRef->InstanceId, InSessionInfoRef);
+	}
+
+	StaticCastSharedPtr<FChaosVDMultiSessionInfo>(ActiveSessionsByInstanceId.FindChecked(AllSessionsWrapperGUID))->InnerSessionsByInstanceID.Emplace(InSessionInfoRef->InstanceId, InSessionInfoRef);
+}
+
+void FChaosVDRemoteSessionsManager::DeRegisterSessionInMultiSessionWrapper(const TSharedRef<FChaosVDSessionInfo>& InSessionInfoRef)
+{
+	StaticCastSharedPtr<FChaosVDMultiSessionInfo>(ActiveSessionsByInstanceId.FindChecked(AllRemoteSessionsWrapperGUID))->InnerSessionsByInstanceID.Remove(InSessionInfoRef->InstanceId);
+	StaticCastSharedPtr<FChaosVDMultiSessionInfo>(ActiveSessionsByInstanceId.FindChecked(AllSessionsWrapperGUID))->InnerSessionsByInstanceID.Remove(InSessionInfoRef->InstanceId);
+}
+
+void FChaosVDRemoteSessionsManager::HandleSessionPongMessage(const FChaosVDSessionPong& InMessage, const TSharedRef<IMessageContext>& InContext)
+{
+	TSharedPtr<FChaosVDSessionInfo>& SessionInfoPtr = ActiveSessionsByInstanceId.FindOrAdd(InMessage.InstanceId);
+
+	if (!SessionInfoPtr)
+	{
+		SessionInfoPtr = MakeShared<FChaosVDSessionInfo>();
+		SessionInfoPtr->Address = InContext->GetSender();
+		SessionInfoPtr->InstanceId = InMessage.InstanceId;
+		SessionInfoPtr->SessionName = InMessage.SessionName;
+
+		RegisterSessionInMultiSessionWrapper(SessionInfoPtr.ToSharedRef());
+
+		// This is the first time we see this session, so we need to request the rest of its state so we can properly populate the UI
+		SendFullSessionStateRequestCommand(SessionInfoPtr->Address);
+	}
+
+	SessionInfoPtr->LastPingTime = FDateTime::UtcNow();
+
+	PendingRecordingStatusMessages.RemoveAndCopyValue(InMessage.InstanceId, SessionInfoPtr->LastKnownRecordingState);
+	
+	SessionsUpdatedDelegate.Broadcast();
+}
+
+void FChaosVDRemoteSessionsManager::HandleSessionPingMessage(const FChaosVDSessionPing& InMessage, const TSharedRef<IMessageContext>& InContext)
+{
+	SendPong(InMessage);
+}
+
+void FChaosVDRemoteSessionsManager::HandleRecordingStatusUpdateMessage(const FChaosVDRecordingStatusMessage& Message, const TSharedRef<IMessageContext>& InContext)
+{
+	if (TSharedPtr<FChaosVDSessionInfo>* SessionInfoPtrPtr = ActiveSessionsByInstanceId.Find(Message.InstanceId))
+	{
+		TSharedPtr<FChaosVDSessionInfo>& SessionInfoPtr = *SessionInfoPtrPtr;
+		check(SessionInfoPtr);
+
+		if (SessionInfoPtr->LastKnownRecordingState.bIsRecording != Message.bIsRecording)
+		{
+			if (Message.bIsRecording)
+			{
+				RecordingStartedDelegate.Broadcast(SessionInfoPtr);
+			}
+			else
+			{
+				RecordingStoppedDelegate.Broadcast(SessionInfoPtr);
+			}
+		}
+
+		SessionInfoPtr->LastKnownRecordingState = Message;
+	}
+	else
+	{
+		PendingRecordingStatusMessages.FindOrAdd(Message.InstanceId) = Message;
+	}
+}
+
+void FChaosVDRemoteSessionsManager::HandleRecordingStartCommandMessage(const FChaosVDStartRecordingCommandMessage& InMessage, const TSharedRef<IMessageContext>& InContext)
+{
+#if WITH_CHAOS_VISUAL_DEBUGGER
+	UE_AUTORTFM_ONCOMMIT(InMessage)
+	{
+		TArray<FString> RecordingParams;
+
+		if (InMessage.RecordingMode == EChaosVDRecordingMode::Live)
+		{
+			RecordingParams.Add(TEXT("Server"));
+
+			if (!InMessage.Target.IsEmpty())
+			{
+				RecordingParams.Add(InMessage.Target);
+			}
+		}
+		else if (InMessage.RecordingMode == EChaosVDRecordingMode::File)
+		{
+			RecordingParams.Add(TEXT("File"));
+		}
+
+		FChaosVDRuntimeModule::Get().StartRecording(RecordingParams);
+	};
+
+#endif
+}
+
+void FChaosVDRemoteSessionsManager::HandleRecordingStopCommandMessage(const FChaosVDStopRecordingCommandMessage& InMessage, const TSharedRef<IMessageContext>& InContext)
+{
+#if WITH_CHAOS_VISUAL_DEBUGGER
+	UE_AUTORTFM_ONCOMMIT()
+	{
+		FChaosVDRuntimeModule::Get().StopRecording();
+	};
+#endif
+}
+
+void FChaosVDRemoteSessionsManager::HandleChangeDataChannelStateCommandMessage(const FChaosVDChannelStateChangeCommandMessage& InMessage, const TSharedRef<IMessageContext>& InContext)
+{
+#if WITH_CHAOS_VISUAL_DEBUGGER
+	UE_AUTORTFM_ONCOMMIT(InMessage)
+	{
+		using namespace Chaos::VisualDebugger;
+		if (TSharedPtr<FChaosVDOptionalDataChannel> ChannelInstance = FChaosVDDataChannelsManager::Get().GetChannelById(FName(InMessage.NewState.ChannelName)))
+		{
+			ChannelInstance->SetChannelEnabled(InMessage.NewState.bIsEnabled);
+		}
+	};
+#endif
+}
+
+void FChaosVDRemoteSessionsManager::HandleChangeDataChannelStateResponseMessage(const FChaosVDChannelStateChangeResponseMessage& InMessage, const TSharedRef<IMessageContext>& InContext)
+{
+	if (TSharedPtr<FChaosVDSessionInfo>* SessionInfoPtrPtr = ActiveSessionsByInstanceId.Find(InMessage.InstanceID))
+	{
+		if (TSharedPtr<FChaosVDSessionInfo>& SessionInfoPtr = *SessionInfoPtrPtr)
+		{
+			if (FChaosVDDataChannelState* FoundChannelState = SessionInfoPtr->DataChannelsStatesByName.Find(InMessage.NewState.ChannelName))
+			{
+				*FoundChannelState = InMessage.NewState;
+			}
+		}
+	}
+}
+
+void FChaosVDRemoteSessionsManager::HandleFullSessionStateRequestMessage(const FChaosVDFullSessionInfoRequestMessage& InMessage, const TSharedRef<IMessageContext>& InContext)
+{
+	if (!ensure(MessageEndpoint))
+	{
+		return;
+	}
+	
+	if (FChaosVDFullSessionInfoResponseMessage* FullSessionStateResponse = FMessageEndpoint::MakeMessage<FChaosVDFullSessionInfoResponseMessage>())
+	{
+		FullSessionStateResponse->InstanceId = FApp::GetInstanceId();
+
+#if WITH_CHAOS_VISUAL_DEBUGGER
+		FullSessionStateResponse->bIsRecording = FChaosVDRuntimeModule::Get().IsRecording();
+
+		using namespace Chaos::VisualDebugger;
+		FChaosVDDataChannelsManager::Get().EnumerateChannels([&FullSessionStateResponse](const TSharedRef<FChaosVDOptionalDataChannel>& Channel)
+		{
+			FChaosVDChannelStateChangeCommandMessage DataChannelStateMessage = {Channel->GetId().ToString(), Channel->IsChannelEnabled(), Channel->CanChangeEnabledState() };
+			FullSessionStateResponse->DataChannelsStates.Emplace(DataChannelStateMessage.NewState);
+			return true;
+		});	
+#endif
+
+		MessageEndpoint->Send(
+			FullSessionStateResponse,
+			FChaosVDFullSessionInfoResponseMessage::StaticStruct(),
+			EMessageFlags::Reliable,
+			nullptr,
+			TArrayBuilder<FMessageAddress>().Add(InContext->GetSender()),
+			FTimespan::Zero(),
+			FDateTime::MaxValue());
+	}
+}
+
+void FChaosVDRemoteSessionsManager::HandleFullSessionStateResponseMessage(const FChaosVDFullSessionInfoResponseMessage& InMessage, const TSharedRef<IMessageContext>& InContext)
+{
+	if (TSharedPtr<FChaosVDSessionInfo>* SessionInfoPtrPtr = ActiveSessionsByInstanceId.Find(InMessage.InstanceId))
+	{
+		if (TSharedPtr<FChaosVDSessionInfo>& SessionInfoPtr = *SessionInfoPtrPtr)
+		{
+			SessionInfoPtr->LastKnownRecordingState.bIsRecording = InMessage.bIsRecording;
+
+			for (const FChaosVDDataChannelState& ChannelState : InMessage.DataChannelsStates)
+			{
+				SessionInfoPtr->DataChannelsStatesByName.Emplace(ChannelState.ChannelName, ChannelState);
+			}
+		}
+	}
+}
+
+void FChaosVDRemoteSessionsManager::RemoveExpiredSessions()
+{
+	bool bAnySessionRemoved = false;
+	FDateTime CurrentTime = FDateTime::UtcNow();
+	for (TMap<FGuid, TSharedPtr<FChaosVDSessionInfo>>::TIterator RemoveIterator = ActiveSessionsByInstanceId.CreateIterator(); RemoveIterator; ++RemoveIterator)
+	{
+		TSharedPtr<FChaosVDSessionInfo>& SessionInfoPtr = RemoveIterator.Value();
+		if (!SessionInfoPtr)
+		{
+			RemoveIterator.RemoveCurrent();
+			bAnySessionRemoved = true;
+			continue;
+		}
+
+		if (!EnumHasAnyFlags(SessionInfoPtr->GetSessionTypeAttributes(), EChaosVDRemoteSessionAttributes::CanExpire))
+		{
+			continue;
+		}
+
+		FTimespan ElapsedTime = CurrentTime - SessionInfoPtr->LastPingTime;
+
+		constexpr float MaxAllowedTimeBetweenPings = 3.0f;
+		if (ElapsedTime > FTimespan::FromSeconds(MaxAllowedTimeBetweenPings))
+		{
+			DeRegisterSessionInMultiSessionWrapper(SessionInfoPtr.ToSharedRef());
+			RemoveIterator.RemoveCurrent();
+			bAnySessionRemoved = true;
+		}
+	}
+
+	if (bAnySessionRemoved)
+	{
+		SessionsUpdatedDelegate.Broadcast();
+	}
+}
diff --git a/Engine/Source/Runtime/Experimental/ChaosSolverEngine/Public/Chaos/ChaosVDEngineEditorBridge.h b/Engine/Source/Runtime/Experimental/ChaosSolverEngine/Public/Chaos/ChaosVDEngineEditorBridge.h
new file mode 100644
--- /dev/null
+++ b/Engine/Source/Runtime/Experimental/ChaosSolverEngine/Public/Chaos/ChaosVDEngineEditorBridge.h
@@ -0,0 +1,86 @@
+// Copyright Epic Games, Inc. All Rights Reserved.
+
+#pragma once
+#include "Chaos/ChaosVDRemoteSessionsManager.h"
+
+#if WITH_CHAOS_VISUAL_DEBUGGER
+
+#include "Delegates/IDelegateInstance.h"
+#include "HAL/Platform.h"
+
+class UGameInstance;
+
+namespace Chaos
+{
+	namespace VisualDebugger
+	{
+		struct FChaosVDOptionalDataChannel;
+	}
+}
+
+/** Object that bridges the gap between the ChaosVD runtime module and the Engine & CVD editor.
+ * As the ChaosVDRuntime module does not have access to the engine module, this object reacts to events and performs necessary operation the runtime module cannot do directly
+ */
+class FChaosVDEngineEditorBridge
+{
+public:
+	
+	FChaosVDEngineEditorBridge() : RemoteSessionsManager(MakeShared<FChaosVDRemoteSessionsManager>())
+	{
+	}
+	
+	~FChaosVDEngineEditorBridge(){}
+
+	CHAOSSOLVERENGINE_API static FChaosVDEngineEditorBridge& Get();
+
+	void Initialize();
+	void TearDown();
+
+	TSharedPtr<FChaosVDRemoteSessionsManager> GetRemoteSessionsManager()
+	{
+		return RemoteSessionsManager;
+	}
+
+	using FChaosVDDataDataChannel = Chaos::VisualDebugger::FChaosVDOptionalDataChannel;
+
+private:
+	void AddOnScreenRecordingMessage();
+	void RemoveOnScreenRecordingMessage();
+	void HandleCVDRecordingStarted();
+	void HandleCVDRecordingStopped();
+	void HandleCVDRecordingStartFailed(const FText& InFailureReason) const;
+	void HandlePIEStarted(UGameInstance* GameInstance);
+
+	void HandleDataChannelChanged(TWeakPtr<FChaosVDDataDataChannel> ChannelWeakPtr);
+
+	void SerializeCollisionChannelsNames();
+	
+	bool BroadcastSessionStatus(float DeltaTime);
+
+	FDelegateHandle RecordingStartedHandle;
+	FDelegateHandle RecordingStoppedHandle;
+	FDelegateHandle RecordingStartFailedHandle;
+	uint64 CVDRecordingMessageKey = 0;
+
+#if WITH_EDITOR
+	FDelegateHandle PIEStartedHandle;
+#endif
+
+	TSharedRef<FChaosVDRemoteSessionsManager> RemoteSessionsManager;
+	
+	FTSTicker::FDelegateHandle RecordingStatusUpdateHandle;
+};
+#else
+
+class FChaosVDEngineEditorBridge
+{
+public:
+	CHAOSSOLVERENGINE_API static FChaosVDEngineEditorBridge& Get();
+	
+	TSharedPtr<FChaosVDRemoteSessionsManager> GetRemoteSessionsManager()
+	{
+		return nullptr;
+	}
+};
+
+#endif
diff --git a/Engine/Source/Runtime/Experimental/ChaosSolverEngine/Public/Chaos/ChaosVDRemoteSessionsManager.h b/Engine/Source/Runtime/Experimental/ChaosSolverEngine/Public/Chaos/ChaosVDRemoteSessionsManager.h
new file mode 100644
--- /dev/null
+++ b/Engine/Source/Runtime/Experimental/ChaosSolverEngine/Public/Chaos/ChaosVDRemoteSessionsManager.h
@@ -0,0 +1,415 @@
+// Copyright Epic Games, Inc. All Rights Reserved.
+
+#pragma once
+#include "ChaosVDRuntimeModule.h"
+#include "ChaosVDRecordingDetails.h"
+#include "Containers/Ticker.h"
+#include "IMessageContext.h"
+#include "Misc/DateTime.h"
+#include "Misc/Guid.h"
+#include "Templates/SharedPointer.h"
+
+#include "ChaosVDRemoteSessionsManager.generated.h"
+
+class IMessageContext;
+struct FChaosVDSessionPong;
+class FMessageEndpoint;
+class IMessageBus;
+
+USTRUCT()
+struct FChaosVDSessionPing
+{
+	GENERATED_BODY()
+	
+	UPROPERTY()
+	FGuid ControllerInstanceId;
+};
+
+USTRUCT()
+struct FChaosVDSessionPong
+{
+	GENERATED_BODY()
+	
+	UPROPERTY()
+	FGuid InstanceId;
+	
+	UPROPERTY()
+	FGuid SessionId;
+
+	UPROPERTY()
+	FString SessionName;
+};
+
+USTRUCT()
+struct FChaosVDStartRecordingCommandMessage
+{
+	GENERATED_BODY()
+
+	UPROPERTY()
+	EChaosVDRecordingMode RecordingMode = EChaosVDRecordingMode::Invalid;
+	
+	UPROPERTY()
+	FString Target;
+};
+
+USTRUCT()
+struct FChaosVDStopRecordingCommandMessage
+{
+	GENERATED_BODY()
+};
+
+USTRUCT()
+struct FChaosVDRecordingStatusMessage
+{
+	GENERATED_BODY()
+	
+	FChaosVDRecordingStatusMessage()
+	{
+	}
+
+	UPROPERTY()
+	FGuid InstanceId = FGuid();
+	
+	UPROPERTY()
+	bool bIsRecording = false;
+
+	UPROPERTY()
+	float ElapsedTime = false;
+	
+	UPROPERTY()
+	FChaosVDTraceDetails TraceDetails;
+};
+
+USTRUCT()
+struct FChaosVDDataChannelState
+{
+	GENERATED_BODY()
+
+	UPROPERTY()
+	FString ChannelName;
+
+	UPROPERTY()
+	bool bIsEnabled = false;
+
+	UPROPERTY()
+	bool bCanChangeChannelState = false;
+
+	bool bWaitingUpdatedState = false;
+};
+
+USTRUCT()
+struct FChaosVDChannelStateChangeCommandMessage
+{
+	GENERATED_BODY()
+	UPROPERTY()
+	FChaosVDDataChannelState NewState;
+};
+
+USTRUCT()
+struct FChaosVDChannelStateChangeResponseMessage
+{
+	GENERATED_BODY()
+
+	UPROPERTY()
+	FGuid InstanceID = FGuid();
+
+	UPROPERTY()
+	FChaosVDDataChannelState NewState;
+};
+
+USTRUCT()
+struct FChaosVDFullSessionInfoRequestMessage
+{
+	GENERATED_BODY()
+};
+
+USTRUCT()
+struct FChaosVDFullSessionInfoResponseMessage
+{
+	GENERATED_BODY()
+
+	UPROPERTY()
+	FGuid InstanceId;
+
+	UPROPERTY()
+	TArray<FChaosVDDataChannelState> DataChannelsStates;
+
+	UPROPERTY()
+	bool bIsRecording = false;
+};
+
+UENUM()
+enum class EChaosVDRemoteSessionAttributes
+{
+	None = 0,
+	SupportsDataChannelChange = 1 << 0,
+	CanExpire = 1 << 1,
+	IsMultiSessionWrapper = 1 << 2,
+};
+
+ENUM_CLASS_FLAGS(EChaosVDRemoteSessionAttributes)
+
+/**
+ * Session object that contains all the information needed to communicate with a remote instance, and the state of that instance
+ */
+struct FChaosVDSessionInfo
+{
+	explicit FChaosVDSessionInfo() :InstanceId(FGuid()),
+									SessionTypeAttributes(EChaosVDRemoteSessionAttributes::CanExpire | EChaosVDRemoteSessionAttributes::SupportsDataChannelChange)
+	{
+	}
+	
+	virtual ~FChaosVDSessionInfo() = default;
+
+protected:
+
+	explicit FChaosVDSessionInfo(EChaosVDRemoteSessionAttributes InSessionTypeAttributes) : SessionTypeAttributes(InSessionTypeAttributes)
+	{
+	}
+
+public:
+
+	FGuid InstanceId;
+	FString SessionName;
+	FMessageAddress Address;
+	FDateTime LastPingTime;
+
+	FChaosVDRecordingStatusMessage LastKnownRecordingState;
+	TMap<FString, FChaosVDDataChannelState> DataChannelsStatesByName;
+
+	virtual bool IsRecording() const;
+	
+	virtual EChaosVDRecordingMode GetRecordingMode() const;
+	
+	virtual bool IsConnected() const;
+
+	EChaosVDRemoteSessionAttributes GetSessionTypeAttributes() const
+	{
+		return SessionTypeAttributes;
+	}
+
+protected:
+	const EChaosVDRemoteSessionAttributes SessionTypeAttributes;
+};
+
+/**
+ * Session object that is able to control and provide information about multiple session objects.
+ * Used to the UI can use the same API to control multiple session, than for single sessions.
+ */
+struct FChaosVDMultiSessionInfo : public FChaosVDSessionInfo
+{
+	explicit FChaosVDMultiSessionInfo() : FChaosVDSessionInfo(EChaosVDRemoteSessionAttributes::IsMultiSessionWrapper)
+	{
+	}
+
+	virtual ~FChaosVDMultiSessionInfo() override = default;
+
+	virtual bool IsRecording() const override;
+	virtual EChaosVDRecordingMode GetRecordingMode() const override;
+	virtual bool IsConnected() const override;
+
+	template<typename TCallback>
+	void EnumerateInnerSessions(const TCallback& Callback) const
+	{
+		for (const TPair<FGuid, TWeakPtr<FChaosVDSessionInfo>>& InnerSessionWithID : InnerSessionsByInstanceID)
+		{
+			if (const TSharedPtr<FChaosVDSessionInfo> SessionPtr = InnerSessionWithID.Value.Pin())
+			{
+				if (!Callback(SessionPtr.ToSharedRef()))
+				{
+					return;
+				}
+			}
+		}
+	}
+
+	TMap<FGuid, TWeakPtr<FChaosVDSessionInfo>> InnerSessionsByInstanceID;
+};
+
+DECLARE_LOG_CATEGORY_EXTERN(LogChaosVDRemoteSession, Log, Log);
+
+DECLARE_MULTICAST_DELEGATE_OneParam(FChaosVDRecordingStateChangeDelegate, TWeakPtr<FChaosVDSessionInfo> Session)
+
+/** Object that is able to discover, issue and execute commands back and forth between CVD and client/server/editor instances */
+class FChaosVDRemoteSessionsManager
+{
+public:
+
+	FChaosVDRemoteSessionsManager();
+
+	void Initialize();
+	void Shutdown();
+
+	FSimpleMulticastDelegate& OnSessionsUpdated()
+	{
+		return SessionsUpdatedDelegate;
+	}
+
+
+	/**
+	 * Issues a command to the provided address that will start a CVD recording
+	 * @param InDestinationAddress Message bus address that will execute the command
+	 * @param RecordingStartCommandParams Desired parameters to be used to start the recording 
+	 */
+	CHAOSSOLVERENGINE_API void SendStartRecordingCommand(const FMessageAddress& InDestinationAddress, const FChaosVDStartRecordingCommandMessage& RecordingStartCommandParams);
+
+	/**
+	 * Issues a command to the provided address to stop a CVD recording 
+	 * 	@param InDestinationAddress Message bus address that will execute the command
+	 */
+	CHAOSSOLVERENGINE_API void SendStopRecordingCommand(const FMessageAddress& InDestinationAddress);
+
+	/**
+	 * Issues a command to the provided address to change the state of a data channel
+	 * 	@param InDestinationAddress Message bus address that will execute the command
+	 * 	@param InNewStateData New data channel state to apply in the receiving instance
+	 */
+	CHAOSSOLVERENGINE_API void SendDataChannelStateChangeCommand(const FMessageAddress& InDestinationAddress, const FChaosVDChannelStateChangeCommandMessage& InNewStateData);
+
+	/**
+	 * Returns the session info object for the provided ID 
+	 * @param Id CVD SessionID
+	 */
+	CHAOSSOLVERENGINE_API TWeakPtr<FChaosVDSessionInfo> GetSessionInfo(FGuid Id);
+
+	/**
+	 * Iterates trough all active and valid cvd sessions, and executes the provided callback to it.
+	 * if the callbacks returns false, the iteration will stop
+	* */
+	template<typename CallbackType>
+	void EnumerateActiveSessions(const CallbackType& Callback)
+	{
+		for (const TPair<FGuid, TSharedPtr<FChaosVDSessionInfo>>& ActiveSession : ActiveSessionsByInstanceId)
+		{
+			if (ActiveSession.Value)
+			{
+				bool bContinue = Callback(ActiveSession.Value.ToSharedRef());
+
+				if (!bContinue)
+				{
+					return;
+				}
+			}
+		}
+	}
+
+
+	/**
+	 * Broadcast to the network a recording state update
+	 * @param UpdateMessage latest recording state of the issuing instance
+	 */
+	void PublishRecordingStatusUpdate(const FChaosVDRecordingStatusMessage& UpdateMessage);
+	/**
+	 * Broadcast to the network a data channel state update
+	 * @param InNewStateData latest data channel state of the issuing instance
+	 */
+	void PublishDataChannelStateChangeUpdate(const FChaosVDChannelStateChangeResponseMessage& InNewStateData);
+
+
+	/**
+	 * Delegate that broadcast when a recording was started in a session (either local or remote)
+	 */
+	FChaosVDRecordingStateChangeDelegate& OnSessionRecordingStarted()
+	{
+		return RecordingStartedDelegate;	
+	}
+
+	/**
+	 * Delegate that broadcast when a recording stops in a session (either local or remote)
+	 */
+	FChaosVDRecordingStateChangeDelegate& OnSessionRecordingStopped()
+	{
+		return RecordingStoppedDelegate;
+	}
+	
+	CHAOSSOLVERENGINE_API static const FString AllRemoteSessionsTargetName;
+	CHAOSSOLVERENGINE_API static const FString AllSessionsTargetName;
+	CHAOSSOLVERENGINE_API static const FString LocalEditorSessionName;
+	CHAOSSOLVERENGINE_API static const FName MessageBusEndPointName;
+	CHAOSSOLVERENGINE_API static const FGuid AllRemoteSessionsWrapperGUID;
+	CHAOSSOLVERENGINE_API static const FGuid AllSessionsWrapperGUID;
+	CHAOSSOLVERENGINE_API static const FGuid InvalidSessionGUID;
+	CHAOSSOLVERENGINE_API static const FGuid LocalEditorSessionID;
+
+private:
+	
+	/**
+	 * Returns true if this instance has controller capabilities (is either an editor or CVD Standalone, which is also an editor)
+	 */
+	static constexpr bool IsController();
+
+	/**
+	 * Sends a request to obtain the full session information to the provided message bus address
+	 * @param InDestinationAddress 
+	 */
+	void SendFullSessionStateRequestCommand(const FMessageAddress& InDestinationAddress);
+
+	/**
+	 * Registers a session object that is able to control multiple session instances
+	 * @param InSessionInfoRef Multi Session object
+	 */
+	void RegisterSessionInMultiSessionWrapper(const TSharedRef<FChaosVDSessionInfo>& InSessionInfoRef);
+	
+	/**
+	 * Deegisters a session object that is able to control multiple session instances
+	 * @param InSessionInfoRef Multi Session object
+	 */
+	void DeRegisterSessionInMultiSessionWrapper(const TSharedRef<FChaosVDSessionInfo>& InSessionInfoRef);
+
+	/**
+	 * Creates a session object that is able to control multiple other session objects.
+	 * @param InstanceId Instance ID for this object
+	 * @param SessionName Session name that will be used in the UI
+	 * @return 
+	 */
+	TSharedPtr<FChaosVDSessionInfo> CreatedWrapperSessionInfo(FGuid InstanceId, const FString& SessionName);
+
+	/***
+	 * Creates the messagebus endpoint this session manager will use
+	 */
+	TSharedPtr<FMessageEndpoint> CreateEndPoint(const TSharedRef<IMessageBus>& InMessageBus);
+
+	bool Tick(float DeltaTime);
+
+	/**
+	 * Usually used by the controller. Broadcast to the network this controller exists.
+	 */
+	void SendPing();
+
+	/**
+	 * Broadcast to the network a small subset of this instance information, in to a received session ping.
+	 * @param InMessage Ping message that was received
+	 */
+	void SendPong(const FChaosVDSessionPing& InMessage);
+	
+	void HandleSessionPongMessage(const FChaosVDSessionPong& InMessage, const TSharedRef<IMessageContext>& InContext);
+	void HandleSessionPingMessage(const FChaosVDSessionPing& InMessage, const TSharedRef<IMessageContext>& InContext);
+	void HandleRecordingStatusUpdateMessage(const FChaosVDRecordingStatusMessage& InMessage, const TSharedRef<IMessageContext>& InContext);
+	void HandleRecordingStartCommandMessage(const FChaosVDStartRecordingCommandMessage& InMessage, const TSharedRef<IMessageContext>& InContext);
+	void HandleRecordingStopCommandMessage(const FChaosVDStopRecordingCommandMessage& InMessage, const TSharedRef<IMessageContext>& InContext);
+	void HandleChangeDataChannelStateCommandMessage(const FChaosVDChannelStateChangeCommandMessage& InMessage, const TSharedRef<IMessageContext>& InContext);
+	void HandleChangeDataChannelStateResponseMessage(const FChaosVDChannelStateChangeResponseMessage& InMessage, const TSharedRef<IMessageContext>& InContext);
+	void HandleFullSessionStateRequestMessage(const FChaosVDFullSessionInfoRequestMessage& InMessage, const TSharedRef<IMessageContext>& InContext);
+	void HandleFullSessionStateResponseMessage(const FChaosVDFullSessionInfoResponseMessage& InMessage, const TSharedRef<IMessageContext>& InContext);
+
+	void RemoveExpiredSessions();
+
+	/** Holds the time at which the last ping was sent. */
+    FDateTime LastPingTime;
+
+    /** Holds a pointer to the message bus. */
+    TWeakPtr<IMessageBus> MessageBusPtr;
+
+    /** Holds the messaging endpoint. */
+    TSharedPtr<FMessageEndpoint> MessageEndpoint;
+
+	TMap<FGuid, TSharedPtr<FChaosVDSessionInfo>> ActiveSessionsByInstanceId;
+	TMap<FGuid, FChaosVDRecordingStatusMessage> PendingRecordingStatusMessages;
+
+	FSimpleMulticastDelegate SessionsUpdatedDelegate;
+
+	FChaosVDRecordingStateChangeDelegate RecordingStartedDelegate;
+	FChaosVDRecordingStateChangeDelegate RecordingStoppedDelegate;
+	
+	FTSTicker::FDelegateHandle TickHandle;
+};
diff --git a/Engine/Source/Runtime/Experimental/ChaosVDData/Private/ChaosVDRecording.cpp b/Engine/Source/Runtime/Experimental/ChaosVDData/Private/ChaosVDRecording.cpp
--- a/Engine/Source/Runtime/Experimental/ChaosVDData/Private/ChaosVDRecording.cpp
+++ b/Engine/Source/Runtime/Experimental/ChaosVDData/Private/ChaosVDRecording.cpp
@@ -330,6 +330,8 @@ void FChaosVDRecording::AddFrameForSolver(const int32 SolverID, FChaosVDSolverFr
 		{
 			CurrentFrameNumber = 0;
 			RecordedFramesDataPerSolver.Add(SolverID, {}).Emplace(MoveTemp(InFrameData));
+
+			CommitSolverID_AssumesLocked(SolverID);
 		}
 	}
 
@@ -372,6 +374,12 @@ FChaosVDGameFrameData* FChaosVDRecording::GetLastGameFrameData_AssumesLocked()
 int32 FChaosVDRecording::GetLowestGameFrameNumberAtCycle(uint64 Cycle)
 {
 	FReadScopeLock ReadLock(RecordingDataLock);
+	return GetLowestGameFrameNumberAtCycle_AssumesLocked(Cycle);
+	
+}
+
+int32 FChaosVDRecording::GetLowestGameFrameNumberAtCycle_AssumesLocked(uint64 Cycle)
+{
 	return Algo::LowerBoundBy(GameFrames, Cycle, &FChaosVDGameFrameData::FirstCycle);
 }
 
@@ -468,7 +476,36 @@ void FChaosVDRecording::SetCollisionChannelsInfoContainer(const TSharedPtr<FChao
 
 bool FChaosVDRecording::HasSolverID(int32 SolverID)
 {
-	return RecordedFramesDataPerSolver.Contains(SolverID);
+	FReadScopeLock ReadLock(RecordingDataLock);
+	return HasSolverID_AssumesLocked(SolverID);
+}
+
+bool FChaosVDRecording::HasSolverID_AssumesLocked(int32 SolverID)
+{
+	return SolverIDs.Contains(SolverID) || ReservedSolverIDs.Contains(SolverID);
+}
+
+void FChaosVDRecording::ReserveSolverID(int32 SolverID)
+{
+	FWriteScopeLock ReadLock(RecordingDataLock);
+	ReserveSolverID_AssumesLocked(SolverID);
+}
+
+void FChaosVDRecording::ReserveSolverID_AssumesLocked(int32 SolverID)
+{
+	ReservedSolverIDs.Emplace(SolverID);
+}
+
+void FChaosVDRecording::CommitSolverID(int32 SolverID)
+{
+	FWriteScopeLock ReadLock(RecordingDataLock);
+	CommitSolverID_AssumesLocked(SolverID);
+}
+
+void FChaosVDRecording::CommitSolverID_AssumesLocked(int32 SolverID)
+{
+	SolverIDs.Emplace(SolverID);
+	ReservedSolverIDs.Remove(SolverID);
 }
 
 void FChaosVDRecording::AddImplicitObject(const uint32 ID, const Chaos::FImplicitObject* InImplicitObject)
diff --git a/Engine/Source/Runtime/Experimental/ChaosVDData/Public/ChaosVDRecording.h b/Engine/Source/Runtime/Experimental/ChaosVDData/Public/ChaosVDRecording.h
--- a/Engine/Source/Runtime/Experimental/ChaosVDData/Public/ChaosVDRecording.h
+++ b/Engine/Source/Runtime/Experimental/ChaosVDData/Public/ChaosVDRecording.h
@@ -544,6 +544,7 @@ struct CHAOSVDDATA_API FChaosVDRecording
 	 * @return Found Game frame number. INDEX_NONE if no frame is found for the specified cycle
 	 */
 	int32 GetLowestGameFrameNumberAtCycle(uint64 Cycle);
+	int32 GetLowestGameFrameNumberAtCycle_AssumesLocked(uint64 Cycle);
 
 	/**
 	 * Searches and returns the lowest game frame number at the specified cycle
@@ -618,15 +619,23 @@ struct CHAOSVDDATA_API FChaosVDRecording
 
 	void AddAttributes(EChaosVDRecordingAttributes Attributes)
 	{
+		FWriteScopeLock WriteLock(RecordingDataLock);
 		EnumAddFlags(RecordingAttributes, Attributes);
 	}
 
 	void RemoveAttributes(EChaosVDRecordingAttributes Attributes)
 	{
+		FWriteScopeLock WriteLock(RecordingDataLock);
 		EnumAddFlags(RecordingAttributes, Attributes);
 	}
 
 	EChaosVDRecordingAttributes GetAttributes() const
+	{
+		FReadScopeLock ReadLock(RecordingDataLock);
+		return GetAttributes_AssumesLocked();
+	}
+
+	EChaosVDRecordingAttributes GetAttributes_AssumesLocked() const
 	{
 		return RecordingAttributes;
 	}
@@ -666,6 +675,13 @@ struct CHAOSVDDATA_API FChaosVDRecording
 	void SetCollisionChannelsInfoContainer(const TSharedPtr<FChaosVDCollisionChannelsInfoContainer>& InCollisionChannelsInfo);
 
 	bool HasSolverID(int32 SolverID);
+	bool HasSolverID_AssumesLocked(int32 SolverID);
+
+	void ReserveSolverID(int32 SolverID);
+	void ReserveSolverID_AssumesLocked(int32 SolverID);
+
+	void CommitSolverID(int32 SolverID);
+	void CommitSolverID_AssumesLocked(int32 SolverID);
 
 protected:
 
@@ -707,6 +723,9 @@ protected:
 
 	TSharedPtr<FChaosVDCollisionChannelsInfoContainer> CollisionChannelsInfoContainer;
 
+	TSet<int32> ReservedSolverIDs;
+	TSet<int32> SolverIDs;
+
 	friend class FChaosVDTraceProvider;
 	friend class FChaosVDTraceImplicitObjectProcessor;
 };
diff --git a/Engine/Source/Runtime/Experimental/ChaosVisualDebugger/Private/ChaosVDRuntimeModule.cpp b/Engine/Source/Runtime/Experimental/ChaosVisualDebugger/Private/ChaosVDRuntimeModule.cpp
--- a/Engine/Source/Runtime/Experimental/ChaosVisualDebugger/Private/ChaosVDRuntimeModule.cpp
+++ b/Engine/Source/Runtime/Experimental/ChaosVisualDebugger/Private/ChaosVDRuntimeModule.cpp
@@ -19,6 +19,7 @@
 #include "Misc/ScopedSlowTask.h"
 #endif
 
+#include "ChaosVDRecordingDetails.h"
 #include "Misc/App.h"
 #include "Modules/ModuleManager.h"
 #include "ProfilingDebugging/TraceAuxiliary.h"
@@ -131,6 +132,20 @@ FString FChaosVDRuntimeModule::GetLastRecordingFileNamePath() const
 	return LastRecordingFileNamePath;
 }
 
+FChaosVDTraceDetails FChaosVDRuntimeModule::GetCurrentTraceSessionDetails() const
+{
+	FChaosVDTraceDetails Details;
+	if (bool bIsConnected = FTraceAuxiliary::IsConnected(Details.SessionGuid, Details.TraceGuid))
+	{
+		Details.TraceTarget = FTraceAuxiliary::GetTraceDestinationString();
+		Details.Mode = FTraceAuxiliary::GetConnectionType() == FTraceAuxiliary::EConnectionType::File ? EChaosVDRecordingMode::File : EChaosVDRecordingMode::Live;
+		Details.bIsConnected = bIsConnected;
+		return Details;
+	}
+
+	return Details;
+}
+
 void FChaosVDRuntimeModule::StopTrace()
 {
 	bRequestedStop = true;
@@ -211,15 +226,21 @@ void FChaosVDRuntimeModule::StartRecording(TConstArrayView<FString> Args)
 		bIsRecording = FTraceAuxiliary::Start(FTraceAuxiliary::EConnectionType::File, *LastRecordingFileNamePath, nullptr, &TracingOptions);
 
 		LastRecordingFileNamePath = bIsRecording ? FTraceAuxiliary::GetTraceDestinationString() : TEXT("");
+		
+		CurrentRecordingMode = EChaosVDRecordingMode::File;
 	}
 	else if(Args[0] == TEXT("Server"))
 	{
 		const FString Target = Args.IsValidIndex(1) ? Args[1] : TEXT("127.0.0.1");
 
+		LastRecordingFileNamePath = Target;
+
 		bIsRecording = FTraceAuxiliary::Start(
 		FTraceAuxiliary::EConnectionType::Network,
 		*Target,
 		nullptr, &TracingOptions);
+		
+		CurrentRecordingMode = EChaosVDRecordingMode::Live;
 	}
 	else
 	{
@@ -259,7 +280,9 @@ void FChaosVDRuntimeModule::StartRecording(TConstArrayView<FString> Args)
 		{
 			UE::TReadScopeLock ReadLock(DelegatesRWLock);
 			RecordingStartFailedDelegate.Broadcast(FailureReason);
-		}	
+		}
+
+		CurrentRecordingMode = EChaosVDRecordingMode::Invalid;
 	}
 }
 
@@ -291,14 +314,14 @@ void FChaosVDRuntimeModule::StopRecording()
 		FTSTicker::GetCoreTicker().RemoveTicker(RecordingTimerHandle);
 		RecordingTimerHandle.Reset();
 	}
+	
+	bIsRecording = false;
+	AccumulatedRecordingTime = 0.0f;
 
 	{
 		UE::TReadScopeLock ReadLock(DelegatesRWLock);
 		RecordingStopDelegate.Broadcast();
 	}
-
-	bIsRecording = false;
-	AccumulatedRecordingTime = 0.0f;
 }
 
 void FChaosVDRuntimeModule::HandleTraceStopRequest(FTraceAuxiliary::EConnectionType TraceType, const FString& TraceDestination)
diff --git a/Engine/Source/Runtime/Experimental/ChaosVisualDebugger/Public/ChaosVDRecordingDetails.h b/Engine/Source/Runtime/Experimental/ChaosVisualDebugger/Public/ChaosVDRecordingDetails.h
new file mode 100644
--- /dev/null
+++ b/Engine/Source/Runtime/Experimental/ChaosVisualDebugger/Public/ChaosVDRecordingDetails.h
@@ -0,0 +1,36 @@
+// Copyright Epic Games, Inc. All Rights Reserved.
+
+#pragma once
+
+#include "Containers/UnrealString.h"
+#include "Misc/Guid.h"
+#include "ChaosVDRecordingDetails.generated.h"
+
+UENUM()
+enum class EChaosVDRecordingMode : uint8
+{
+	Invalid,
+	Live,
+	File
+};
+
+USTRUCT()
+struct FChaosVDTraceDetails
+{
+	GENERATED_BODY()
+
+	UPROPERTY()
+	FGuid TraceGuid;
+	
+	UPROPERTY()
+	FGuid SessionGuid;
+
+	UPROPERTY()
+	FString TraceTarget;
+
+	UPROPERTY()
+	bool bIsConnected = false;
+
+	UPROPERTY()
+	EChaosVDRecordingMode Mode = EChaosVDRecordingMode::Invalid;
+};
\ No newline at end of file
diff --git a/Engine/Source/Runtime/Experimental/ChaosVisualDebugger/Public/ChaosVDRuntimeModule.h b/Engine/Source/Runtime/Experimental/ChaosVisualDebugger/Public/ChaosVDRuntimeModule.h
--- a/Engine/Source/Runtime/Experimental/ChaosVisualDebugger/Public/ChaosVDRuntimeModule.h
+++ b/Engine/Source/Runtime/Experimental/ChaosVisualDebugger/Public/ChaosVDRuntimeModule.h
@@ -1,7 +1,10 @@
 // Copyright Epic Games, Inc. All Rights Reserved.
 
 #pragma once
+
 #if WITH_CHAOS_VISUAL_DEBUGGER
+
+#include "ChaosVDRecordingDetails.h"
 #include "Templates/SharedPointer.h"
 #include "Modules/ModuleInterface.h"
 #include "ProfilingDebugging/TraceAuxiliary.h"
@@ -20,6 +23,7 @@ enum class EChaosVDFullCaptureFlags : int32
 };
 ENUM_CLASS_FLAGS(EChaosVDFullCaptureFlags)
 
+
 DECLARE_MULTICAST_DELEGATE(FChaosVDRecordingStateChangedDelegate)
 DECLARE_MULTICAST_DELEGATE_OneParam(FChaosVDCaptureRequestDelegate, EChaosVDFullCaptureFlags)
 DECLARE_MULTICAST_DELEGATE_OneParam(FChaosVDRecordingStartFailedDelegate, const FText&)
@@ -102,6 +106,13 @@ public:
 	/** Returns the full path of the active recording file*/
 	FString GetLastRecordingFileNamePath() const;
 
+	FChaosVDTraceDetails GetCurrentTraceSessionDetails() const;
+
+	EChaosVDRecordingMode GetCurrentRecordingMode() const
+	{
+		return CurrentRecordingMode;
+	}
+
 private:
 
 	/** Stops the current Trace session */
@@ -145,6 +156,8 @@ private:
 
 	TMap<FString, bool> OriginalTraceChannelsState;
 
+	EChaosVDRecordingMode CurrentRecordingMode = EChaosVDRecordingMode::Invalid;
+
 	static FTransactionallySafeRWLock DelegatesRWLock;
 };
 #endif
