[ChaosVD] Fix for an issue where the GT data in multiple files mode is not visualized.

This change was meant to be part of my previous CL.

diff --git a/Engine/Plugins/ChaosVD/Source/ChaosVD/Private/Trace/ChaosVDTraceAnalyzer.cpp b/Engine/Plugins/ChaosVD/Source/ChaosVD/Private/Trace/ChaosVDTraceAnalyzer.cpp
--- a/Engine/Plugins/ChaosVD/Source/ChaosVD/Private/Trace/ChaosVDTraceAnalyzer.cpp
+++ b/Engine/Plugins/ChaosVD/Source/ChaosVD/Private/Trace/ChaosVDTraceAnalyzer.cpp
@@ -47,7 +47,7 @@ void FChaosVDTraceAnalyzer::PushSimulatedSolverTrackForGTData(const UE::Trace::I
 {
 	FChaosVDSolverFrameData NewFrameData;
 
-	int32 CurrentGTTrackID = ChaosVDTraceProvider->GetRemappedSolverID(ChaosVDTraceProvider->GetCurrentGameThreadTrackID());
+	int32 CurrentGTTrackID = ChaosVDTraceProvider->GetCurrentGameThreadTrackID();
 	if (CurrentGTTrackID == INDEX_NONE)
 	{
 		CurrentGTTrackID = ChaosVDTraceProvider->RemapSolverID(0);
@@ -323,13 +323,20 @@ bool FChaosVDTraceAnalyzer::OnEvent(uint16 RouteId, EStyle Style, const FOnEvent
 
 			if (bUsingAutoRTFM)
 			{
-				FText MissingDataProcessorsMessage = NSLOCTEXT("ChaosVisualDebugger","AutoRTFMWarningMessage", "This recording was made with AutoRTFM enabled. \n\nAutoRTFM is not fully supported and framing/timing of recorded data during transactions might be off.");
+				FText AutoRTFMWarning = NSLOCTEXT("ChaosVisualDebugger","AutoRTFMWarningMessage", "This recording was made with AutoRTFM enabled. \n\nAutoRTFM is not fully supported and framing/timing of recorded data during transactions might be off.");
 
-				FTSTicker::GetCoreTicker().AddTicker(FTickerDelegate::CreateLambda([MissingDataProcessorsMessage](float DeltaTime)
+				if (EnumHasAnyFlags(ChaosVDTraceProvider->GetRecordingForSession()->GetAttributes(), EChaosVDRecordingAttributes::Merged))
 				{
-					FMessageDialog::Open(EAppMsgCategory::Warning, EAppMsgType::Ok, MissingDataProcessorsMessage, NSLOCTEXT("ChaosVisualDebugger", "AutoRTFMWarningMessageTitle", "Partially unsupported CVD Recording"));
-					return false;
-				}));
+					UE_LOG(LogChaosVDEditor, Warning, TEXT("[%s]"), *AutoRTFMWarning.ToString())
+				}
+				else
+				{
+					FTSTicker::GetCoreTicker().AddTicker(FTickerDelegate::CreateLambda([AutoRTFMWarning](float DeltaTime)
+					{
+						FMessageDialog::Open(EAppMsgCategory::Warning, EAppMsgType::Ok, AutoRTFMWarning, NSLOCTEXT("ChaosVisualDebugger", "AutoRTFMWarningMessageTitle", "Partially unsupported CVD Recording"));
+						return false;
+					}));
+				}
 			}
 
 			break;
