[ChaosVD] Fix for an issue where the pop up to automatically load a recording made to fail, will fail to unload an existing loaded recording.

diff --git a/Engine/Plugins/ChaosVD/Source/ChaosVD/Private/Widgets/SChaosVDRecordingControls.cpp b/Engine/Plugins/ChaosVD/Source/ChaosVD/Private/Widgets/SChaosVDRecordingControls.cpp
--- a/Engine/Plugins/ChaosVD/Source/ChaosVD/Private/Widgets/SChaosVDRecordingControls.cpp
+++ b/Engine/Plugins/ChaosVD/Source/ChaosVD/Private/Widgets/SChaosVDRecordingControls.cpp
@@ -444,7 +444,7 @@ void SChaosVDRecordingControls::HandleRecordingStop(TWeakPtr<FChaosVDSessionInfo
 	{
 		if (FMessageDialog::Open(EAppMsgType::YesNo, LOCTEXT("OpenLastRecordingMessage", "Do you want to load the recorded file now? ")) == EAppReturnType::Yes)
 		{
-			MainTabSharedPtr->GetChaosVDEngineInstance()->LoadRecording(CurrentTraceTarget.ToString());
+			MainTabSharedPtr->GetChaosVDEngineInstance()->LoadRecording(CurrentTraceTarget.ToString(), EChaosVDLoadRecordedDataMode::SingleSource);
 		}
 	}
 }
