Mover support in CVD

diff --git a/Engine/Plugins/ChaosVD/Source/ChaosVD/Private/ChaosVDBaseSceneObject.h b/Engine/Plugins/ChaosVD/Source/ChaosVD/Private/ChaosVDBaseSceneObject.h
deleted file mode 100644
diff --git a/Engine/Plugins/ChaosVD/Source/ChaosVD/Private/ChaosVDCharacterGroundConstraintDataProviderInterface.h b/Engine/Plugins/ChaosVD/Source/ChaosVD/Private/ChaosVDCharacterGroundConstraintDataProviderInterface.h
deleted file mode 100644
diff --git a/Engine/Plugins/ChaosVD/Source/ChaosVD/Private/ChaosVDExtractedGeometryDataHandle.h b/Engine/Plugins/ChaosVD/Source/ChaosVD/Private/ChaosVDExtractedGeometryDataHandle.h
deleted file mode 100644
diff --git a/Engine/Plugins/ChaosVD/Source/ChaosVD/Private/ChaosVDGeometryDataComponent.h b/Engine/Plugins/ChaosVD/Source/ChaosVD/Private/ChaosVDGeometryDataComponent.h
deleted file mode 100644
diff --git a/Engine/Plugins/ChaosVD/Source/ChaosVD/Private/ChaosVDObjectDetailsTab.h b/Engine/Plugins/ChaosVD/Source/ChaosVD/Private/ChaosVDObjectDetailsTab.h
deleted file mode 100644
diff --git a/Engine/Plugins/ChaosVD/Source/ChaosVD/Private/ChaosVDScene.cpp b/Engine/Plugins/ChaosVD/Source/ChaosVD/Private/ChaosVDScene.cpp
--- a/Engine/Plugins/ChaosVD/Source/ChaosVD/Private/ChaosVDScene.cpp
+++ b/Engine/Plugins/ChaosVD/Source/ChaosVD/Private/ChaosVDScene.cpp
@@ -35,6 +35,7 @@
 #include "Engine/TextureCube.h"
 #include "TEDS/ChaosVDParticleEditorDataFactory.h"
 #include "TEDS/ChaosVDStructTypedElementData.h"
+#include "TEDS/ChaosVDSelectionInterface.h"
 
 #define LOCTEXT_NAMESPACE "ChaosVisualDebugger"
 
@@ -540,6 +541,11 @@ UWorld* FChaosVDScene::CreatePhysicsVDWorld()
 	return NewWorld;
 }
 
+TArray<FTypedElementHandle> FChaosVDScene::GetSelectedElementHandles()
+{
+	return GetElementSelectionSet()->GetSelectedElementHandles(UChaosVDSelectionInterface::StaticClass());
+}
+
 FTypedElementHandle FChaosVDScene::GetSelectionHandleForObject(const UObject* Object) const
 {
 	FTypedElementHandle Handle;
diff --git a/Engine/Plugins/ChaosVD/Source/ChaosVD/Private/ChaosVDSceneParticle.h b/Engine/Plugins/ChaosVD/Source/ChaosVD/Private/ChaosVDSceneParticle.h
deleted file mode 100644
diff --git a/Engine/Plugins/ChaosVD/Source/ChaosVD/Private/ChaosVDSceneParticleFlags.h b/Engine/Plugins/ChaosVD/Source/ChaosVD/Private/ChaosVDSceneParticleFlags.h
deleted file mode 100644
diff --git a/Engine/Plugins/ChaosVD/Source/ChaosVD/Private/ChaosVDSceneSelectionObserver.h b/Engine/Plugins/ChaosVD/Source/ChaosVD/Private/ChaosVDSceneSelectionObserver.h
deleted file mode 100644
diff --git a/Engine/Plugins/ChaosVD/Source/ChaosVD/Private/ExtensionsSystem/ChaosVDExtensionsManager.cpp b/Engine/Plugins/ChaosVD/Source/ChaosVD/Private/ExtensionsSystem/ChaosVDExtensionsManager.cpp
--- a/Engine/Plugins/ChaosVD/Source/ChaosVD/Private/ExtensionsSystem/ChaosVDExtensionsManager.cpp
+++ b/Engine/Plugins/ChaosVD/Source/ChaosVD/Private/ExtensionsSystem/ChaosVDExtensionsManager.cpp
@@ -15,12 +15,14 @@ FChaosVDExtensionsManager::~FChaosVDExtensionsManager()
 	IModularFeatures::Get().OnModularFeatureRegistered().RemoveAll(this);
 	IModularFeatures::Get().OnModularFeatureUnregistered().RemoveAll(this);
 	
-	EnumerateExtensions([this](const TSharedRef<FChaosVDExtension>& InExtension)
+	// We can't use EnumerateExtensions here as it iterates on AvailableExtensions
+	// which is being modified by UnRegisterExtension, causing a crash
+	TArray<TSharedRef<FChaosVDExtension>> AvailableExtensionsCopy;
+	AvailableExtensions.GenerateValueArray(AvailableExtensionsCopy);
+	for (TSharedRef<FChaosVDExtension> AvailableExtension : AvailableExtensionsCopy)
 	{
-		UnRegisterExtension(InExtension);
-
-		return true;
-	});
+		UnRegisterExtension(AvailableExtension);
+	}
 }
 
 FChaosVDExtensionsManager& FChaosVDExtensionsManager::Get()
diff --git a/Engine/Plugins/ChaosVD/Source/ChaosVD/Private/TEDS/ChaosVDParticleEditorDataFactory.cpp b/Engine/Plugins/ChaosVD/Source/ChaosVD/Private/TEDS/ChaosVDParticleEditorDataFactory.cpp
--- a/Engine/Plugins/ChaosVD/Source/ChaosVD/Private/TEDS/ChaosVDParticleEditorDataFactory.cpp
+++ b/Engine/Plugins/ChaosVD/Source/ChaosVD/Private/TEDS/ChaosVDParticleEditorDataFactory.cpp
@@ -4,8 +4,8 @@
 
 #include "ChaosVDParentDataStorageFactory.h"
 #include "ChaosVDSceneParticle.h"
-#include "ChaosVDSelectionInterface.h"
-#include "ChaosVDStructTypedElementData.h"
+#include "TEDS/ChaosVDSelectionInterface.h"
+#include "TEDS/ChaosVDStructTypedElementData.h"
 #include "Actors/ChaosVDSolverInfoActor.h"
 #include "Elements/Columns/TypedElementCompatibilityColumns.h"
 #include "Elements/Columns/TypedElementLabelColumns.h"
diff --git a/Engine/Plugins/ChaosVD/Source/ChaosVD/Private/TEDS/ChaosVDSelectionInterface.cpp b/Engine/Plugins/ChaosVD/Source/ChaosVD/Private/TEDS/ChaosVDSelectionInterface.cpp
--- a/Engine/Plugins/ChaosVD/Source/ChaosVD/Private/TEDS/ChaosVDSelectionInterface.cpp
+++ b/Engine/Plugins/ChaosVD/Source/ChaosVD/Private/TEDS/ChaosVDSelectionInterface.cpp
@@ -1,4 +1,4 @@
 // Copyright Epic Games, Inc. All Rights Reserved.
 
-#include "ChaosVDSelectionInterface.h"
+#include "TEDS/ChaosVDSelectionInterface.h"
 
diff --git a/Engine/Plugins/ChaosVD/Source/ChaosVD/Private/TEDS/ChaosVDSelectionInterface.h b/Engine/Plugins/ChaosVD/Source/ChaosVD/Private/TEDS/ChaosVDSelectionInterface.h
deleted file mode 100644
diff --git a/Engine/Plugins/ChaosVD/Source/ChaosVD/Private/TEDS/ChaosVDStructTypedElementData.cpp b/Engine/Plugins/ChaosVD/Source/ChaosVD/Private/TEDS/ChaosVDStructTypedElementData.cpp
--- a/Engine/Plugins/ChaosVD/Source/ChaosVD/Private/TEDS/ChaosVDStructTypedElementData.cpp
+++ b/Engine/Plugins/ChaosVD/Source/ChaosVD/Private/TEDS/ChaosVDStructTypedElementData.cpp
@@ -1,6 +1,6 @@
 // Copyright Epic Games, Inc. All Rights Reserved.
 
-#include "ChaosVDStructTypedElementData.h"
+#include "TEDS/ChaosVDStructTypedElementData.h"
 
 #include "Elements/Framework/TypedElementOwnerStore.h"
 #include "Elements/Framework/TypedElementRegistry.h"
diff --git a/Engine/Plugins/ChaosVD/Source/ChaosVD/Private/TEDS/ChaosVDStructTypedElementData.h b/Engine/Plugins/ChaosVD/Source/ChaosVD/Private/TEDS/ChaosVDStructTypedElementData.h
deleted file mode 100644
diff --git a/Engine/Plugins/ChaosVD/Source/ChaosVD/Private/Widgets/SChaosVDDetailsView.h b/Engine/Plugins/ChaosVD/Source/ChaosVD/Private/Widgets/SChaosVDDetailsView.h
deleted file mode 100644
diff --git a/Engine/Plugins/ChaosVD/Source/ChaosVD/Private/Widgets/SChaosVDMainTab.cpp b/Engine/Plugins/ChaosVD/Source/ChaosVD/Private/Widgets/SChaosVDMainTab.cpp
--- a/Engine/Plugins/ChaosVD/Source/ChaosVD/Private/Widgets/SChaosVDMainTab.cpp
+++ b/Engine/Plugins/ChaosVD/Source/ChaosVD/Private/Widgets/SChaosVDMainTab.cpp
@@ -236,6 +236,11 @@ UWorld* SChaosVDMainTab::GetWorld() const
 	return GetChaosVDEngineInstance()->GetCurrentScene()->GetUnderlyingWorld();
 }
 
+TSharedPtr<FChaosVDScene> SChaosVDMainTab::GetScene()
+{
+	return GetChaosVDEngineInstance()->GetCurrentScene();
+}
+
 FEditorModeTools& SChaosVDMainTab::GetEditorModeManager() const
 {
 	check(EditorModeTools.IsValid())
diff --git a/Engine/Plugins/ChaosVD/Source/ChaosVD/Public/ChaosVDBaseSceneObject.h b/Engine/Plugins/ChaosVD/Source/ChaosVD/Public/ChaosVDBaseSceneObject.h
new file mode 100644
--- /dev/null
+++ b/Engine/Plugins/ChaosVD/Source/ChaosVD/Public/ChaosVDBaseSceneObject.h
@@ -0,0 +1,154 @@
+// Copyright Epic Games, Inc. All Rights Reserved.
+
+#pragma once
+#include "Elements/Common/TypedElementHandles.h"
+#include "UObject/NameTypes.h"
+#include "UObject/WeakObjectPtr.h"
+#include "Templates/SharedPointer.h"
+
+#include "ChaosVDBaseSceneObject.generated.h"
+
+class AActor;
+
+/** Base Class for any UStruct based object that has a representation in CVD's World and outliner*/
+USTRUCT()
+struct FChaosVDBaseSceneObject
+{
+	GENERATED_BODY()
+
+	virtual ~FChaosVDBaseSceneObject() = default;
+
+	/**
+	 * Sets the name that will be used for this object in the scene outliner
+	 * @param NewDisplayName New name
+	 */
+	void SetDisplayName(const FString& NewDisplayName)
+	{
+		DisplayName = NewDisplayName;
+	}
+
+	/**
+	 * Sets the name of the icon that will be shown as part of the label for this object in the scene outliner
+	 * @param NewIconName Registered name of the icon in an editor style
+	 */
+	void SetIconName(FName NewIconName)
+	{
+		IconName = NewIconName;
+	}
+
+	/**
+	 * Sets a weak reference to another UStruct based scene object, that is the parent of this object
+	 * @param NewParent Parent Object
+	 */
+	virtual void SetParent(const TSharedPtr<FChaosVDBaseSceneObject>& NewParent)
+	{
+		Parent = NewParent;
+	}
+
+	/**
+	 * Sets a weak reference to an actor that is the parent of this object.
+	 * @param NewParent Parent Actor
+	 */
+	void SetParentParentActor(AActor* NewParent);
+
+
+	/**
+	 *  Returns a copy the name that will be used for this object in the scene outliner
+	 */
+	FString GetDisplayName() const
+	{
+		return DisplayName;
+	}
+
+	/**
+	 *  Returns a reference to the name that will be used for this object in the scene outliner
+	 */
+	const FString& GetDisplayNameRef() const
+	{
+		return DisplayName;
+	}
+
+	/**
+	 * Returns the registered name of the icon that will be shown as part of the label for this object in the scene outliner
+	 */
+	FName GetIconName() const
+	{
+		return IconName;
+	}
+
+	/**
+	 * Returns a weak ptr to this object's parent
+	 */
+	TWeakPtr<FChaosVDBaseSceneObject> GetParent() const
+	{
+		return Parent;
+	}
+
+	/**
+	 * Returns a weak ptr to this object's parent Actor
+	 */
+	AActor* GetParentActor() const
+	{
+		return ParentActor.Get();
+	}
+
+	/**
+	 * Returns the handle for this object in TEDS
+	 */
+	UE::Editor::DataStorage::RowHandle GetTedsRowHandle() const
+	{
+		return CachedRowHandle;
+	}
+
+	/**
+	 * Sets the handle for this object in TEDS
+	 */
+	void SetTedsRowHandle(UE::Editor::DataStorage::RowHandle Handle)
+	{
+		CachedRowHandle = Handle;
+	}
+	
+	enum class EStreamingState
+	{
+		Visible,
+		Hidden
+	};
+
+	void SetStreamingState(EStreamingState NewState)
+	{
+		StreamingState = NewState;
+	}
+	
+	EStreamingState GetStreamingState() const
+	{
+		return StreamingState;
+	}
+
+	virtual void SyncStreamingState()
+	{
+	}
+
+	virtual FBox GetStreamingBounds() const
+	{
+		return FBox(EForceInit::ForceInitToZero);
+	}
+
+	virtual int32 GetStreamingID() const
+	{
+		return INDEX_NONE;
+	}
+
+protected:
+	TWeakPtr<FChaosVDBaseSceneObject> Parent = nullptr;
+
+	TWeakObjectPtr<AActor> ParentActor = nullptr;
+
+	FString DisplayName = TEXT("None");
+
+	UE::Editor::DataStorage::RowHandle CachedRowHandle = UE::Editor::DataStorage::InvalidRowHandle;
+	
+	FName IconName = NAME_None;
+
+	EStreamingState StreamingState = EStreamingState::Hidden;
+};
+
diff --git a/Engine/Plugins/ChaosVD/Source/ChaosVD/Public/ChaosVDCharacterGroundConstraintDataProviderInterface.h b/Engine/Plugins/ChaosVD/Source/ChaosVD/Public/ChaosVDCharacterGroundConstraintDataProviderInterface.h
new file mode 100644
--- /dev/null
+++ b/Engine/Plugins/ChaosVD/Source/ChaosVD/Public/ChaosVDCharacterGroundConstraintDataProviderInterface.h
@@ -0,0 +1,42 @@
+// Copyright Epic Games, Inc. All Rights Reserved.
+
+#pragma once
+
+#include "UObject/Interface.h"
+#include "ChaosVDCharacterGroundConstraintDataProviderInterface.generated.h"
+
+struct FChaosVDCharacterGroundConstraint;
+
+// This class does not need to be modified.
+UINTERFACE()
+class UChaosVDCharacterGroundConstraintDataProviderInterface : public UInterface
+{
+	GENERATED_BODY()
+};
+
+/**
+ * Interface for any object that is able to provide CVD Character Ground Constraint Data
+ */
+class IChaosVDCharacterGroundConstraintDataProviderInterface
+{
+	GENERATED_BODY()
+
+public:
+	
+	/**
+	 * Gathers and populates the provided array with any existing character ground constraint data for this object (if any) 
+	 * @param OutConstraintDataFound Array to populate with constraint data
+	 */
+	virtual void GetCharacterGroundConstraintData(TArray<TSharedPtr<FChaosVDCharacterGroundConstraint>>& OutConstraintsFound) PURE_VIRTUAL(IChaosVDCharacterGroundConstraintDataProviderInterface::GetCharacterGroundConstraintData)
+
+	/**
+	 * Checks if the object implementing the interface has any character ground constraint data available
+	 * @return True if any constraint data is found 
+	 */
+	virtual bool HasCharacterGroundConstraintData() PURE_VIRTUAL(IChaosVDCharacterGroundConstraintDataProviderInterface::HasCharacterGroundConstraintData, return false;)
+
+	/**
+	 * Returns the name of the object providing the constraint data
+	 */
+	virtual FName GetProviderName() PURE_VIRTUAL(IChaosVDCharacterGroundConstraintDataProviderInterface::GetProviderName, return FName();)
+};
diff --git a/Engine/Plugins/ChaosVD/Source/ChaosVD/Public/ChaosVDExtractedGeometryDataHandle.h b/Engine/Plugins/ChaosVD/Source/ChaosVD/Public/ChaosVDExtractedGeometryDataHandle.h
new file mode 100644
--- /dev/null
+++ b/Engine/Plugins/ChaosVD/Source/ChaosVD/Public/ChaosVDExtractedGeometryDataHandle.h
@@ -0,0 +1,117 @@
+﻿// Copyright Epic Games, Inc. All Rights Reserved.
+
+#pragma once
+
+//#include "ChaosVDGeometryBuilder.h"
+#include "Math/Transform.h"
+#include "UObject/NameTypes.h"
+#include "UObject/WeakObjectPtr.h"
+
+class FChaosVDGeometryBuilder;
+class UDynamicMesh;
+class UStaticMesh;
+
+namespace Chaos
+{
+	class FImplicitObject;
+}
+
+/** Structure that acts as a handle to access the generated static mesh of a implicit object, plus other data extracted from it */
+struct FChaosVDExtractedGeometryDataHandle
+{
+	/** Key to be used to match this geometry with recorded Shape Data */
+	uint32 GetDataComponentKey() const
+	{
+		return DataComponentKey;
+	}
+	void SetDataComponentKey(const uint32 Key)
+	{
+		DataComponentKey = Key;
+	}
+
+	/** Key representing the generated static mesh - Handles for different implicit objects can point to the same geometry key, if after being unpacked they result in essentially the same mesh */
+	uint32 GetGeometryKey() const
+	{
+		return GeometryKey;
+	}
+	void SetGeometryKey(const uint32 Key);
+
+	/** Name of the type this underlying implicit object this instance represents */
+	FName GetTypeName() const;
+
+	/** Extracted transform of the generated implicit object in component space */
+	const FTransform& GetRelativeTransform() const
+	{
+		return Transform;
+	}
+	void SetTransform(const FTransform& InTransform)
+	{
+		Transform = InTransform;
+	}
+
+	/** Ptr to the implicit object this handle provides access to and from which the mesh was generated */
+	const Chaos::FImplicitObject* GetImplicitObject() const
+	{
+		return ImplicitObject;
+	}
+
+	/** Returns the Index to correct the shape data instance in the shape instances array of the particle that owns this geometry */
+	int32 GetShapeInstanceIndex() const
+	{
+		return ShapeInstanceIndex;
+	}
+
+	/** If the object was an union and this was a leaf, this getter returns the root union object, otherwise it will be the same as GetImplicitObject */
+	const Chaos::FImplicitObject* GetRootImplicitObject() const
+	{
+		return RootImplicitObject;
+	}
+	
+	/** Sets the Implicit Object Ptr from which the data was extracted and the mesh was generated */
+	void SetImplicitObject(const Chaos::FImplicitObject* InImplicitObject)
+	{
+		ImplicitObject = InImplicitObject;
+	}
+
+	/** Sets the Index on which this Implicit object was in the Hierarchy of the root one */
+	void SetShapeInstanceIndex(int32 InObjectIndex)
+	{
+		ShapeInstanceIndex = InObjectIndex;
+	}
+	
+	/** Sets the Implicit Object Ptr from which the data was extracted and the mesh was generated */
+	void SetRootImplicitObject(const Chaos::FImplicitObject* InImplicitObject)
+	{
+		RootImplicitObject = InImplicitObject;
+	}
+
+	bool operator==(const FChaosVDExtractedGeometryDataHandle& Other) const
+	{
+		return Other.GeometryKey == GeometryKey && Other.DataComponentKey == DataComponentKey && Other.Transform.Equals(Transform);
+	}
+
+	friend uint32 GetTypeHash(const FChaosVDExtractedGeometryDataHandle& Handle)
+	{
+		return HashCombine(HashCombine(Handle.GeometryKey, Handle.DataComponentKey), GetTypeHash(Handle.Transform));
+	}
+
+private:
+
+	/** Key representing the generated static mesh */
+	uint32 GeometryKey = 0;
+
+	/** Key to be used to match this geometry with recorded Shape Data */
+	uint32 DataComponentKey = 0;
+
+	/** Index on which this Implicit object was in the Hierarchy of the root one */
+	int32 ShapeInstanceIndex = INDEX_NONE;
+
+	/** Extracted transform of the generated implicit object in component space */
+	FTransform Transform;
+
+	/** The Implicit Object Ptr from which the data was extracted and the mesh was generated */
+	const Chaos::FImplicitObject* ImplicitObject = nullptr;
+
+	/** If the object was an union and this was a leaf, this will be the root union object, otherwise it will be the same as ImplicitObject */
+	const Chaos::FImplicitObject* RootImplicitObject = nullptr;
+};
diff --git a/Engine/Plugins/ChaosVD/Source/ChaosVD/Public/ChaosVDGeometryDataComponent.h b/Engine/Plugins/ChaosVD/Source/ChaosVD/Public/ChaosVDGeometryDataComponent.h
new file mode 100644
--- /dev/null
+++ b/Engine/Plugins/ChaosVD/Source/ChaosVD/Public/ChaosVDGeometryDataComponent.h
@@ -0,0 +1,426 @@
+// Copyright Epic Games, Inc. All Rights Reserved.
+
+#pragma once
+#include "ChaosVDExtractedGeometryDataHandle.h"
+#include "Chaos/ImplicitObject.h"
+#include "Components/MeshComponent.h"
+#include "DataWrappers/ChaosVDCollisionDataWrappers.h"
+#include "HAL/Platform.h"
+#include "InstancedStaticMeshDelegates.h"
+#include "Templates/SharedPointer.h"
+#include "UObject/Interface.h"
+#include "UObject/StrongObjectPtr.h"
+
+#include "ChaosVDGeometryDataComponent.generated.h"
+
+class UChaosVDParticleVisualizationSettings;
+class UChaosVDParticleVisualizationColorSettings;
+struct FChaosVDSceneParticle;
+class FChaosVDGeometryBuilder;
+class FChaosVDScene;
+class IChaosVDGeometryComponent;
+class UChaosVDInstancedStaticMeshComponent;
+class UChaosVDStaticMeshComponent;
+class UMaterialInstanceDynamic;
+
+struct FChaosVDExtractedGeometryDataHandle;
+struct FChaosVDParticleDataWrapper;
+
+namespace Chaos
+{
+	class FImplicitObject;
+}
+
+DECLARE_MULTICAST_DELEGATE_OneParam(FChaosVDMeshReadyDelegate, IChaosVDGeometryComponent&)
+DECLARE_MULTICAST_DELEGATE_OneParam(FChaosVDMeshComponentEmptyDelegate, UMeshComponent*)
+
+UENUM()
+enum class EChaosVDMaterialType
+{
+	SMOpaque,
+	SMTranslucent,
+	ISMCOpaque,
+	ISMCTranslucent
+};
+
+UENUM()
+enum class EChaosVDMeshComponent
+{
+	Invalid,
+	Static,
+	InstancedStatic,
+	Dynamic
+};
+
+/** Struct holding the a minimum amount of data about a Implicit object to be shown in the details panel */
+USTRUCT()
+struct FChaosVDImplicitObjectBasicView
+{
+	GENERATED_BODY()
+
+	/** Geometry type name*/
+	UPROPERTY(VisibleAnywhere, Category="Recorded GeometryData")
+	FName ImplicitObjectType;
+
+	Chaos::EImplicitObjectType ImplicitObjectTypeEnum = Chaos::ImplicitObjectType::Unknown;
+
+	/** Index of the Shape Instance data for this geometry in the Shape Instance data array */
+	UPROPERTY(VisibleAnywhere, Category="Recorded GeometryData")
+	int32 ShapeInstanceIndex = INDEX_NONE;
+
+	/** True if this is the root implicit object */
+	UPROPERTY(VisibleAnywhere, Category="Recorded GeometryData")
+	bool bIsRootObject = false;
+
+	/** If this is a transformed implicit, this will contain the recorded relative transform */
+	UPROPERTY(VisibleAnywhere, Category="Recorded GeometryData")
+	FTransform RelativeTransform;
+};
+
+/** Struct holding the state of a mesh instance - Is separated from the Mesh instance class so we can show the data in the Details panel */
+USTRUCT()
+struct FChaosVDMeshDataInstanceState
+{
+	GENERATED_BODY()
+
+	/** Recorded Shape instance Data */
+	UPROPERTY(VisibleAnywhere, Category="Recorded GeometryData")
+	FChaosVDShapeCollisionData CollisionData;
+
+	/** Minimum set of data about the recorded implicit object */
+	UPROPERTY(VisibleAnywhere, Category="Recorded GeometryData")
+	FChaosVDImplicitObjectBasicView ImplicitObjectInfo;
+
+	/* CVD Debug - Current world transform used to render this Mesh */
+	UPROPERTY(VisibleAnywhere, Category="CVD GeometryData", meta=(EditCondition=bShowCVDDebugData, EditConditionHides))
+	FTransform CurrentWorldTransform;
+
+	/* CVD Debug - Current mesh component type to render this Mesh */
+	UPROPERTY(VisibleAnywhere, Category="CVD GeometryData", meta=(EditCondition=bShowCVDDebugData, EditConditionHides))
+	EChaosVDMeshComponent MeshComponentType = EChaosVDMeshComponent::Invalid;
+
+	/* CVD Debug - Pointer to the mesh component used to render this Mesh */
+	UPROPERTY()
+	TObjectPtr<UMeshComponent> MeshComponent;
+
+	/* CVD Debug - Instance index of mesh component used to render this Mesh */
+	UPROPERTY(VisibleAnywhere, Category="CVD GeometryData", meta=(EditCondition=bShowCVDDebugData, EditConditionHides))
+	int32 MeshInstanceIndex = INDEX_NONE;
+
+	/* CVD Debug - Color used to render this mesh */
+	UPROPERTY(VisibleAnywhere, Category="CVD GeometryData", meta=(EditCondition=bShowCVDDebugData, EditConditionHides))
+	FLinearColor CurrentGeometryColor = FLinearColor(ForceInitToZero);
+
+	/* CVD Debug - Id of the particle this geometry belongs */
+	UPROPERTY(VisibleAnywhere, Category="CVD GeometryData", meta=(EditCondition=bShowCVDDebugData, EditConditionHides))
+	int32 OwningParticleID = INDEX_NONE;
+
+	/* CVD Debug - Id of the solver this geometry belongs */
+	UPROPERTY(VisibleAnywhere, Category="CVD GeometryData", meta=(EditCondition=bShowCVDDebugData, EditConditionHides))
+	int32 OwningSolverID = INDEX_NONE;
+
+	UPROPERTY(VisibleAnywhere, Category="CVD GeometryData", meta=(EditCondition=bShowCVDDebugData, EditConditionHides))
+	bool bIsVisible = true;
+
+	UPROPERTY(VisibleAnywhere, Category="CVD GeometryData", meta=(EditCondition=bShowCVDDebugData, EditConditionHides))
+	bool bIsSelected = false;
+
+	UPROPERTY(EditAnywhere, Category="CVD GeometryData")
+	bool bShowCVDDebugData = false;
+};
+
+/** Structure that represents a specific mesh instance on a CVD Mesh component (instanced or static) */
+class FChaosVDInstancedMeshData : public TSharedFromThis<FChaosVDInstancedMeshData>
+{
+public:
+	explicit FChaosVDInstancedMeshData(int32 InInstanceIndex, UMeshComponent* InMeshComponent, int32 InParticleID, int32 InSolverID, const TSharedRef<FChaosVDExtractedGeometryDataHandle>& InSourceGeometryHandle);
+
+	/** Returns the Particle ID of the particle owning this mesh instance */
+	int32 GetOwningParticleID() const
+	{
+		return InstanceState.OwningParticleID;
+	}
+
+	/** Returns the Solver ID of the particle owning this mesh instance */
+	int32 GetOwningSolverID() const
+	{
+		return InstanceState.OwningSolverID;
+	}
+
+	/** Applies the provided world transform to the mesh instance this handle represents
+	 * @param InTransform Transform to apply
+	 */
+	void SetWorldTransform(const FTransform& InTransform);
+
+	/** Returns the world transform of the mesh instance this handle represents */
+	const FTransform& GetWorldTransform() const
+	{
+		return InstanceState.CurrentWorldTransform;
+	}
+
+	/** Returns the geometry handle used to create the mesh instance this handle represents */
+	const TSharedRef<FChaosVDExtractedGeometryDataHandle>& GetGeometryHandle() const
+	{
+		return ExtractedGeometryHandle;
+	}
+
+	/** Applies to the provided color to the mesh instance this handle represents */
+	void SetInstanceColor(const FLinearColor& NewColor);
+
+	/** Returns the current color of the mesh instance this handle represents */
+	FLinearColor GetInstanceColor() const
+	{
+		return InstanceState.CurrentGeometryColor;
+	}
+	
+	/** Applies to the provided shape collision data to the mesh instance this handle represents */
+	void UpdateMeshComponentForCollisionData(const FChaosVDShapeCollisionData& InCollisionData);
+
+	/** Returns the mesh component used to render the mesh instance this handle represents */
+	UMeshComponent* GetMeshComponent() const
+	{
+		return InstanceState.MeshComponent;
+	}
+
+	/** Returns the instance index of the mesh instance this handle represents */
+	int32 GetMeshInstanceIndex() const
+	{
+		return InstanceState.MeshInstanceIndex;
+	}
+
+	/** Returns the type of the component used to render the mesh instance this handle represents */
+	EChaosVDMeshComponent GetMeshComponentType() const
+	{
+		return InstanceState.MeshComponentType;
+	}
+
+	/** Sets a Ptr to the geometry builder used to generate and manage the geometry/mesh components this handle represents */
+	void SetGeometryBuilder(const TWeakPtr<FChaosVDGeometryBuilder>& InGeometryBuilder)
+	{
+		GeometryBuilderInstance = InGeometryBuilder;
+	}
+	TWeakPtr<FChaosVDGeometryBuilder> GetGeometryBuilder()
+	{
+		return GeometryBuilderInstance;
+	}
+
+	/** Marks this mesh instance as selected. Used to handle Selection in Editor */
+	void SetIsSelected(bool bInIsSelected);
+
+	/** Sets the visibility of this mesh instance */
+	void SetVisibility(bool bInIsVisible);
+
+	/** Returns the current visibility state this mesh instance */
+	bool GetVisibility() const
+	{
+		return InstanceState.bIsVisible;
+	}
+
+	/** Applies a new shape collision data to this mesh instance */
+	void SetGeometryCollisionData(const FChaosVDShapeCollisionData&& InCollisionData);
+
+	/** Returns the current shape collision data to this mesh instance */
+	FChaosVDShapeCollisionData& GetGeometryCollisionData()
+	{
+		return InstanceState.CollisionData;
+	}
+
+	FChaosVDMeshDataInstanceState& GetState()
+	{
+		return InstanceState;
+	}
+
+	/** Used only for debugging purposes - It will be set to true if we received new Shape Instance data but the Shape Index for the implicit object we represent is not valid */
+	bool bFailedToUpdateShapeInstanceData = false;
+	
+	/** Sets the mesh instance index of the mesh instance this handle represents */
+	void SetMeshInstanceIndex(int32 NewIndex)
+	{
+		InstanceState.MeshInstanceIndex = NewIndex;
+	}
+
+	/** Returns true if this instance is queued to be destroyed */
+	bool IsPendingDestroy() const
+	{
+		return bIsPendingDestroy;
+	}
+
+	/** Marks a instance that is queued to be destroyed at the end of the frame */
+	void MarkPendingDestroy()
+	{
+		bIsPendingDestroy = true;
+	}
+
+private:
+
+	/** Sets the mesh component used to render the mesh instance this handle represents */
+	void SetMeshComponent(UMeshComponent* NewComponent)
+	{
+		InstanceState.MeshComponent = NewComponent;
+	}
+
+	FChaosVDMeshDataInstanceState InstanceState;
+
+	bool bIsPendingDestroy = false;
+
+	TSharedRef<FChaosVDExtractedGeometryDataHandle> ExtractedGeometryHandle;
+
+	TWeakPtr<FChaosVDGeometryBuilder> GeometryBuilderInstance = nullptr;
+
+	friend UChaosVDInstancedStaticMeshComponent;
+	friend UChaosVDStaticMeshComponent;
+};
+
+enum class EChaosVDMeshAttributesFlags : uint8
+{
+	None = 0,
+	MirroredGeometry = 1 << 0,
+	TranslucentGeometry = 1 << 1
+};
+
+ENUM_CLASS_FLAGS(EChaosVDMeshAttributesFlags)
+
+UINTERFACE()
+class UChaosVDGeometryComponent : public UInterface
+{
+	GENERATED_BODY()
+};
+
+/**
+ * Interface with a default implementation for any Geometry component that
+ * contains CVD data
+ */
+class IChaosVDGeometryComponent
+{
+	GENERATED_BODY()
+
+public:
+
+	/** Returns the Geometry Handle used to identify the geometry data this component represents */
+	virtual uint32 GetGeometryKey() const PURE_VIRTUAL(IChaosVDGeometryDataComponent::GetGeometryKey, return 0;);
+
+	/** Returns the CVD Mesh Data Instance handle for the provided Instance index */
+	virtual TSharedPtr<FChaosVDInstancedMeshData> GetMeshDataInstanceHandle(int32 InstanceIndex) const PURE_VIRTUAL(IChaosVDGeometryDataComponent::GetMeshDataInstanceHandle, return nullptr;);
+
+	/**
+	 * Add a new instance to this mesh component
+	 * @param InstanceTransform The transform the new instance will use
+	 * @param bIsWorldSpace True if the provide transform is in world space
+	 * @param InGeometryHandle The CVD Geometry Handle with data about the geometry instantiate
+	 * @param ParticleID Particle ID owning this geometry
+	 * @param SolverID Solver ID of the particle owning this geometry
+	 * @return CVD Mesh instance handle that provides access to this component and specific instance, allowing manipulation of it
+	 */
+	virtual TSharedPtr<FChaosVDInstancedMeshData> AddMeshInstance(const FTransform InstanceTransform, bool bIsWorldSpace, const TSharedRef<FChaosVDExtractedGeometryDataHandle>& InGeometryHandle, int32 ParticleID, int32 SolverID) PURE_VIRTUAL(IChaosVDGeometryDataComponent::AddMeshInstance, return nullptr;);
+
+	/**
+	 * Adds a new instance to this mesh component, but using an existing Mesh Data Handle instead of creating a new one
+	 * @param InMeshDataHandle Handle that will provide access to this mesh instance
+	 */
+	virtual void AddExistingMeshInstance(const TSharedRef<FChaosVDInstancedMeshData>& InMeshDataHandle) PURE_VIRTUAL(IChaosVDGeometryDataComponent::AddExistingMeshInstance);
+
+	enum class ERemovalMode
+	{
+		Instant,
+		Deferred
+	};
+
+	/**
+	 * Removes the instance the provided handle represents
+	 * @param InHandleToRemove 
+	 */
+	virtual void RemoveMeshInstance(const TSharedRef<FChaosVDInstancedMeshData>& InHandleToRemove, ERemovalMode Mode = ERemovalMode::Deferred) PURE_VIRTUAL(IChaosVDGeometryDataComponent::RemoveMeshInstance);
+
+	/** True if the mesh this component represents is ready for use */
+	virtual bool IsMeshReady() const  PURE_VIRTUAL(IChaosVDGeometryDataComponent::IsMeshReady, return false;);
+
+	/** Sets if the mesh this component represents is ready for use or not */
+	virtual void SetIsMeshReady(bool bIsReady) PURE_VIRTUAL(IChaosVDGeometryDataComponent::SetIsMeshReady);
+	
+	/** Triggers when the mesh this component represents is ready */
+	virtual FChaosVDMeshReadyDelegate* OnMeshReady() PURE_VIRTUAL(IChaosVDGeometryDataComponent::OnMeshReady, return nullptr;);
+
+	/** Triggers when the component does not have any instance to render. Used to allow it to return to the mesh component tool for future re-use*/
+	virtual FChaosVDMeshComponentEmptyDelegate* OnComponentEmpty() PURE_VIRTUAL(IChaosVDGeometryDataComponent::OnComponentEmpty, return nullptr;);
+
+	/** Updates the visibility of this component based on the stored CVD data */
+	virtual void UpdateVisibilityForInstance(const TSharedRef<FChaosVDInstancedMeshData>& InInstanceHandle) PURE_VIRTUAL(IChaosVDGeometryDataComponent::UpdateVisibilityForInstance);
+
+	/** Changes the selection state of the provided instance - Used for Selection in Editor */
+	virtual void UpdateSelectionStateForInstance(const TSharedRef<FChaosVDInstancedMeshData>& InInstanceHandle) PURE_VIRTUAL(IChaosVDGeometryDataComponent::UpdateSelectionStateForInstance);
+
+	/** Updates the colors of this component based on the stored CVD data */
+	virtual void UpdateColorForInstance(const TSharedRef<FChaosVDInstancedMeshData>& InInstanceHandle) PURE_VIRTUAL(IChaosVDGeometryDataComponent::UpdateColorForInstance);
+
+	/** Updates the colors of this component based on the stored CVD data */
+	virtual void UpdateWorldTransformForInstance(const TSharedRef<FChaosVDInstancedMeshData>& InInstanceHandle) PURE_VIRTUAL(IChaosVDGeometryDataComponent::UpdateWorldTransformForInstance);
+
+	/** Sets the CVD Mesh Attribute flags this component is compatible with*/
+	virtual void SetMeshComponentAttributeFlags(EChaosVDMeshAttributesFlags Flags) PURE_VIRTUAL(IChaosVDGeometryDataComponent::SetMeshComponentAttributeFlags);
+	
+	/** Returns the CVD Mesh Attribute flags this component is compatible with*/
+	virtual EChaosVDMeshAttributesFlags GetMeshComponentAttributeFlags() const PURE_VIRTUAL(IChaosVDGeometryDataComponent::GetMeshComponentAttributeFlags, return EChaosVDMeshAttributesFlags::None;);
+
+	/** Resets the state of this mesh component, so it can be re-used later on */
+	virtual void Reset() PURE_VIRTUAL(IChaosVDGeometryDataComponent::Reset);
+
+	virtual void Initialize() PURE_VIRTUAL(IChaosVDGeometryDataComponent::Initialize);
+	
+	/** Sets a Ptr to the geometry builder used to generate and manage the geometry/mesh components */
+	virtual void SetGeometryBuilder(TWeakPtr<FChaosVDGeometryBuilder> GeometryBuilder) PURE_VIRTUAL(IChaosVDGeometryDataComponent::SetGeometryBuilder);
+
+	virtual EChaosVDMaterialType GetMaterialType() const PURE_VIRTUAL(IChaosVDGeometryDataComponent::GetMaterialType,  return EChaosVDMaterialType::SMOpaque;);
+
+	virtual bool GetIsDestroyed() const PURE_VIRTUAL(IChaosVDGeometryDataComponent::GetIsDestroyed,  return false;);
+
+	virtual void SetIsDestroyed(bool bNewIsDestroyed) PURE_VIRTUAL(IChaosVDGeometryDataComponent::SetIsDestroyed);
+};
+
+class FChaosVDGeometryComponentUtils
+{
+public:
+	/** Finds and updates the Shape data using the provided array as source*/
+	static void UpdateCollisionDataFromShapeArray(TArrayView<FChaosVDShapeCollisionData> InShapeArray, const TSharedRef<FChaosVDInstancedMeshData>& InInstanceHandle);
+
+	/** Calculates and updates the color used to render the mesh represented by the provided handle, based on the particle state */
+	static void UpdateMeshColor(const TSharedRef<FChaosVDInstancedMeshData>& InInstanceHandle, const FChaosVDParticleDataWrapper& InParticleData, bool bIsServer);
+
+	/** Calculates the correct visibility state based on the particle state, and applies it to the mesh instance the provided handle represents */
+	static void UpdateMeshVisibility(const TSharedRef<FChaosVDInstancedMeshData>& InInstanceHandle, const FChaosVDParticleDataWrapper& InParticleData, bool bIsActive);
+
+	/** Returns the material to use as a base to create material instances for the provided type */
+	static UMaterialInterface* GetBaseMaterialForType(EChaosVDMaterialType Type);
+
+	/** Returns the correct material type to use based on the provided Component type and Mesh Attributes */
+	template<typename TComponent>
+	static EChaosVDMaterialType GetMaterialTypeForComponent(EChaosVDMeshAttributesFlags MeshAttributes);
+
+private:
+	/** Returns the color that needs to be used to present the provided particle data based on its state and current selected options */
+	static FLinearColor GetGeometryParticleColor(const TSharedRef<FChaosVDExtractedGeometryDataHandle>& InGeometryHandle, const FChaosVDParticleDataWrapper& InParticleData, bool bIsServer);
+
+	static const UChaosVDParticleVisualizationColorSettings* GetParticleColorSettings();
+	static const UChaosVDParticleVisualizationSettings* GetParticleVisualizationSettings();
+
+	static TObjectPtr<UChaosVDParticleVisualizationColorSettings> CachedParticleColorSettings;
+	static TObjectPtr<UChaosVDParticleVisualizationSettings> CachedParticleVisualizationSettings;
+};
+
+template <typename TComponent>
+EChaosVDMaterialType FChaosVDGeometryComponentUtils::GetMaterialTypeForComponent(EChaosVDMeshAttributesFlags MeshAttributes)
+{
+	constexpr bool bIsInstancedMeshComponent = std::is_base_of_v<UInstancedStaticMeshComponent, TComponent>;
+	if (EnumHasAnyFlags(MeshAttributes, EChaosVDMeshAttributesFlags::TranslucentGeometry))
+	{
+		return bIsInstancedMeshComponent ? EChaosVDMaterialType::ISMCTranslucent : EChaosVDMaterialType::SMTranslucent;
+	}
+	else
+	{
+		return bIsInstancedMeshComponent ? EChaosVDMaterialType::ISMCOpaque : EChaosVDMaterialType::SMOpaque;
+	}
+}
+
+namespace Chaos::VisualDebugger
+{
+	void SelectParticleWithGeometryInstance(const TSharedRef<FChaosVDScene>& InScene, FChaosVDSceneParticle* Particle, const TSharedPtr<FChaosVDInstancedMeshData>& InMeshDataHandle);
+}
diff --git a/Engine/Plugins/ChaosVD/Source/ChaosVD/Public/ChaosVDObjectDetailsTab.h b/Engine/Plugins/ChaosVD/Source/ChaosVD/Public/ChaosVDObjectDetailsTab.h
new file mode 100644
--- /dev/null
+++ b/Engine/Plugins/ChaosVD/Source/ChaosVD/Public/ChaosVDObjectDetailsTab.h
@@ -0,0 +1,81 @@
+// Copyright Epic Games, Inc. All Rights Reserved.
+
+#pragma once
+
+#include "ChaosVDSceneSelectionObserver.h"
+#include "Templates/SharedPointer.h"
+#include "ChaosVDTabSpawnerBase.h"
+#include "Delegates/IDelegateInstance.h"
+#include "Widgets/SChaosVDDetailsView.h"
+
+struct FChaosVDSolverDataSelectionHandle;
+struct FChaosVDParticleDebugData;
+
+class AActor;
+class FName;
+class FSpawnTabArgs;
+class FTabManager;
+class SChaosVDDetailsView;
+class SDockTab;
+
+/** Spawns and handles an instance for a selection independent details panel */
+class FChaosVDStandAloneObjectDetailsTab : public FChaosVDTabSpawnerBase, public TSharedFromThis<FChaosVDStandAloneObjectDetailsTab>
+{
+public:
+
+	FChaosVDStandAloneObjectDetailsTab(const FName& InTabID, TSharedPtr<FTabManager> InTabManager, TWeakPtr<SChaosVDMainTab> InOwningTabWidget) : FChaosVDTabSpawnerBase(InTabID, InTabManager, InOwningTabWidget)
+	{
+	}
+	
+	TSharedPtr<SChaosVDDetailsView> GetDetailsPanel()
+	{
+		return DetailsPanelView;
+	}
+
+	void AddUnsupportedStruct(const UStruct* Struct);
+	
+	/** Updates the current object this view details is viewing */
+	template<typename TStruct>
+	void SetStructToInspect(TStruct* NewStruct);
+
+protected:
+
+	virtual TSharedRef<SDockTab> HandleTabSpawnRequest(const FSpawnTabArgs& Args) override;
+	virtual void HandleTabClosed(TSharedRef<SDockTab> InTabClosed) override;
+
+	bool IsSupportedStruct(const TWeakObjectPtr<const UStruct>& InWeakStructPtr);
+
+	TSharedPtr<SChaosVDDetailsView> DetailsPanelView;
+
+	TSet<TWeakObjectPtr<const UStruct>> UnsupportedStructs;
+};
+
+template <typename TStruct>
+void FChaosVDStandAloneObjectDetailsTab::SetStructToInspect(TStruct* NewStruct)
+{
+	DetailsPanelView->SetSelectedStruct(NewStruct);
+}
+
+/** Spawns and handles an instance for the visual debugger details panel */
+class CHAOSVD_API FChaosVDObjectDetailsTab : public FChaosVDStandAloneObjectDetailsTab, public FChaosVDSceneSelectionObserver
+{
+public:
+
+	FChaosVDObjectDetailsTab(const FName& InTabID, TSharedPtr<FTabManager> InTabManager, TWeakPtr<SChaosVDMainTab> InOwningTabWidget) : FChaosVDStandAloneObjectDetailsTab(InTabID, InTabManager, InOwningTabWidget)
+	{
+	}
+
+protected:
+
+	virtual TSharedRef<SDockTab> HandleTabSpawnRequest(const FSpawnTabArgs& Args) override;
+	virtual void HandleTabClosed(TSharedRef<SDockTab> InTabClosed) override;
+
+	virtual void HandlePostSelectionChange(const UTypedElementSelectionSet* ChangedSelectionSet) override;
+
+	void HandleActorsSelection(TArrayView<AActor*> SelectedActors);
+
+	void HandleSolverDataSelectionChange(const TSharedPtr<FChaosVDSolverDataSelectionHandle>& SelectionHandle);
+
+	FDelegateHandle SelectionDelegateHandle;
+	TWeakObjectPtr<UObject> CurrentSelectedObject = nullptr;
+};
\ No newline at end of file
diff --git a/Engine/Plugins/ChaosVD/Source/ChaosVD/Public/ChaosVDScene.h b/Engine/Plugins/ChaosVD/Source/ChaosVD/Public/ChaosVDScene.h
--- a/Engine/Plugins/ChaosVD/Source/ChaosVD/Public/ChaosVDScene.h
+++ b/Engine/Plugins/ChaosVD/Source/ChaosVD/Public/ChaosVDScene.h
@@ -216,6 +216,9 @@ public:
 		return LoadedRecording;
 	}
 
+	// Returns the currently selected elements in the scene, which can then be further inspected with GetStructDataFromTypedElementHandle
+	CHAOSVD_API TArray<FTypedElementHandle> GetSelectedElementHandles();
+
 	FName GetTEDSSelectionSetName();
 
 	void ClearSelectionAndNotify();
diff --git a/Engine/Plugins/ChaosVD/Source/ChaosVD/Public/ChaosVDSceneParticle.h b/Engine/Plugins/ChaosVD/Source/ChaosVD/Public/ChaosVDSceneParticle.h
new file mode 100644
--- /dev/null
+++ b/Engine/Plugins/ChaosVD/Source/ChaosVD/Public/ChaosVDSceneParticle.h
@@ -0,0 +1,212 @@
+// Copyright Epic Games, Inc. All Rights Reserved.
+
+#pragma once
+
+#include "ChaosVDBaseSceneObject.h"
+#include "Chaos/Core.h"
+#include "ChaosVDCharacterGroundConstraintDataProviderInterface.h"
+#include "ChaosVDGeometryDataComponent.h"
+#include "ChaosVDSceneParticleFlags.h"
+#include "DataWrappers/ChaosVDParticleDataWrapper.h"
+#include "DataWrappers/ChaosVDCollisionDataWrappers.h"
+
+#include "ChaosVDSceneParticle.generated.h"
+
+class IEditorDataStorageCompatibilityProvider;
+class IEditorDataStorageProvider;
+
+UENUM()
+enum class EChaosVDParticleVisibilityUpdateFlags
+{
+	None = 0,
+	DirtyScene = 1 << 0,
+};
+ENUM_CLASS_FLAGS(EChaosVDParticleVisibilityUpdateFlags)
+
+USTRUCT()
+struct FChaosVDSceneParticle : public FChaosVDBaseSceneObject 
+{
+
+GENERATED_BODY()
+
+public:
+	virtual ~FChaosVDSceneParticle() override;
+
+	FChaosVDSceneParticle();
+
+	void PreUpdateFromRecordedParticleData(const TSharedPtr<FChaosVDParticleDataWrapper>& InRecordedData, const Chaos::FRigidTransform3& SimulationTransform);
+	void ProcessPendingParticleDataUpdates();
+
+	TSharedPtr<const FChaosVDParticleDataWrapper> GetParticleData() const
+	{
+		return ParticleDataPtr;
+	}
+
+	virtual void SetParent(const TSharedPtr<FChaosVDBaseSceneObject>& NewParent) override;
+
+	void UpdateGeometryComponentsVisibility(EChaosVDParticleVisibilityUpdateFlags Flags = EChaosVDParticleVisibilityUpdateFlags::None);
+
+	void UpdateGeometryColors();
+
+	/** Changes the active state of this CVD Particle Actor */
+	void SetIsActive(bool bNewActive);
+
+	void AddHiddenFlag(EChaosVDHideParticleFlags Flag);
+
+	void RemoveHiddenFlag(EChaosVDHideParticleFlags Flag);
+
+	EChaosVDHideParticleFlags GetHideFlags() const
+	{
+		return HideParticleFlags;
+	}
+
+	bool IsVisible() const
+	{
+		return HideParticleFlags == EChaosVDHideParticleFlags::None;
+	}
+
+	/** Returns true if this particle actor is active - Inactive Particle actors are still in the world but with outdated data
+	 * and hidden from the viewport and outliner. They represent particles that were destroyed.
+	 */
+	bool IsActive() const
+	{
+		return bIsActive;
+	}
+
+	FBox GetBoundingBox() const;
+	Chaos::TAABB<double, 3> GetChaosBoundingBox() const;
+
+	// BEGIN IChaosVDCollisionDataProvider Interface
+	TConstArrayView<TSharedPtr<FChaosVDParticlePairMidPhase>> GetCollisionData();
+	bool HasCollisionData();
+	FName GetProviderName();
+	// END IChaosVDCollisionDataProvider Interface
+
+	// BEGIN IChaosVDCharacterGroundConstraintDataProviderInterface
+	void GetCharacterGroundConstraintData(TArray<TSharedPtr<FChaosVDCharacterGroundConstraint>>& OutConstraintsFound);
+	bool HasCharacterGroundConstraintData();
+	// END IChaosVDCharacterGroundConstraintDataProviderInterface
+
+	void SetIsServerParticle(bool bNewIsServer)
+	{
+		bIsServer = bNewIsServer;
+	}
+
+	bool GetIsServerParticle() const
+	{
+		return bIsServer;
+	}
+
+	void UpdateMeshInstancesSelectionState();
+
+	// BEGIN IChaosVDGeometryOwner Interface
+	TConstArrayView<TSharedRef<FChaosVDInstancedMeshData>> GetMeshInstances() const
+	{
+		return MeshDataHandles;
+	}
+
+	virtual void SetSelectedMeshInstance(const TWeakPtr<FChaosVDInstancedMeshData>& GeometryInstanceToSelect);
+
+	virtual TWeakPtr<FChaosVDInstancedMeshData> GetSelectedMeshInstance() const
+	{
+		return CurrentSelectedGeometryInstance;
+	}
+
+	// END IChaosVDGeometryOwner Interface
+
+	void HandleDeSelected();
+	void HandleSelected();
+
+	bool IsSelected();
+
+	void SetScene(const TWeakPtr<FChaosVDScene>& NewScene)
+	{
+		SceneWeakPtr = NewScene;
+	}
+	
+	TWeakPtr<FChaosVDScene> GetScene() const
+	{
+		return SceneWeakPtr;
+	}
+
+	TWeakPtr<FChaosVDSceneParticle> GetParentParticle()
+	{
+		return ParentParticleInstance;
+	}
+
+	EChaosVDSceneParticleDirtyFlags GetDirtyFlags() const
+	{
+		return DirtyFlags;
+	}
+
+	void RemoveAllGeometry();
+
+	virtual FBox GetStreamingBounds() const override;
+
+	virtual void SyncStreamingState() override;
+
+	virtual int32 GetStreamingID() const override;
+
+protected:
+	
+	void UpdateGeometry(const Chaos::FConstImplicitObjectPtr& InImplicitObject, EChaosVDActorGeometryUpdateFlags OptionsFlags = EChaosVDActorGeometryUpdateFlags::None);
+
+	void CalculateAndCacheBounds() const;
+
+	void UpdateParent(const TSharedPtr<FChaosVDParticleDataWrapper>& InRecordedData);
+
+	void ProcessUpdatedAndRemovedHandles(TArray<TSharedPtr<FChaosVDExtractedGeometryDataHandle>>& OutExtractedGeometryDataHandles);
+
+	const TArray<TSharedPtr<FChaosVDParticlePairMidPhase>>* GetCollisionMidPhasesArray() const;
+
+	const TArray<TSharedPtr<FChaosVDConstraintDataWrapperBase>>* GetCharacterGroundConstraintArray() const;
+
+	void UpdateShapeDataComponents();
+
+	void ApplyPendingTransformData();
+
+	template<typename TTaskCallback>
+	void VisitGeometryInstances(const TTaskCallback& VisitorCallback);
+
+	uint8 bIsGeometryDataGenerationStarted : 1 = false;
+	uint8 bIsActive: 1 = false;
+	uint8 bIsServer: 1 = false;
+	EChaosVDHideParticleFlags HideParticleFlags = EChaosVDHideParticleFlags::None;
+	EChaosVDSceneParticleDirtyFlags DirtyFlags = EChaosVDSceneParticleDirtyFlags::None;
+
+	Chaos::FConstImplicitObjectPtr CurrentRootGeometry = nullptr;
+
+	TWeakPtr<FChaosVDSceneParticle> ParentParticleInstance;
+
+	TSharedPtr<FChaosVDParticleDataWrapper> ParticleDataPtr;
+
+	FDelegateHandle GeometryUpdatedDelegate;
+
+	TArray<TSharedRef<FChaosVDInstancedMeshData>> MeshDataHandles;
+
+	TWeakPtr<FChaosVDInstancedMeshData> CurrentSelectedGeometryInstance;
+
+	TWeakPtr<FChaosVDScene> SceneWeakPtr;
+
+	FTransform PendingParticleTransform;
+	FTransform CachedSimulationTransform;
+
+	mutable FBox CachedBounds = FBox(ForceInitToZero);
+
+	/** Called when this particle is destroyed. This is intentionally private and only intended to be used by
+	 * the customization class for this object
+	 */
+	FSimpleDelegate ParticleDestroyedDelegate;
+
+	friend class FChaosVDSceneParticleCustomization;
+	friend class UChaosVDParticleDataComponent;
+};
+
+template <typename TVisitorCallback>
+void FChaosVDSceneParticle::VisitGeometryInstances(const TVisitorCallback& VisitorCallback)
+{
+	for (TSharedRef<FChaosVDInstancedMeshData>& MeshDataHandle : MeshDataHandles)
+	{
+		VisitorCallback(MeshDataHandle);
+	}
+}
diff --git a/Engine/Plugins/ChaosVD/Source/ChaosVD/Public/ChaosVDSceneParticleFlags.h b/Engine/Plugins/ChaosVD/Source/ChaosVD/Public/ChaosVDSceneParticleFlags.h
new file mode 100644
--- /dev/null
+++ b/Engine/Plugins/ChaosVD/Source/ChaosVD/Public/ChaosVDSceneParticleFlags.h
@@ -0,0 +1,46 @@
+// Copyright Epic Games, Inc. All Rights Reserved.
+
+#pragma once
+
+#include "HAL/Platform.h"
+#include "Misc/EnumClassFlags.h"
+
+#include "ChaosVDSceneParticleFlags.generated.h"
+
+/** Options flags to control how geometry is updated in a ChaosVDActor */
+UENUM()
+enum class EChaosVDActorGeometryUpdateFlags : uint8
+{
+	None = 0,
+	ForceUpdate = 1 << 0
+};
+ENUM_CLASS_FLAGS(EChaosVDActorGeometryUpdateFlags)
+
+UENUM()
+enum class EChaosVDHideParticleFlags : uint8
+{
+	None = 0,
+	HiddenByVisualizationFlags = 1 << 0,
+	HiddenBySceneOutliner = 1 << 1,
+	HiddenByActiveState = 1 << 2,
+	HiddenBySolverVisibility = 1 << 3,
+};
+ENUM_CLASS_FLAGS(EChaosVDHideParticleFlags)
+
+/** Set of flags used to indicate which data was modified and needs to be evaluated again by other systems */
+UENUM()
+enum class EChaosVDSceneParticleDirtyFlags : uint16
+{
+	None =			0,
+	Visibility =	1 << 0,
+	Coloring =		1 << 1,
+	Active =		1 << 2,
+	Transform =		1 << 3,
+	Parent =		1 << 4,
+	Geometry =		1 << 5,
+	CollisionData =	1 << 6,
+	PreUpdatePass =	1 << 7,
+	TEDS =	1 << 8,
+	StreamingBounds = 1 << 9,
+};
+ENUM_CLASS_FLAGS(EChaosVDSceneParticleDirtyFlags)
diff --git a/Engine/Plugins/ChaosVD/Source/ChaosVD/Public/ChaosVDSceneSelectionObserver.h b/Engine/Plugins/ChaosVD/Source/ChaosVD/Public/ChaosVDSceneSelectionObserver.h
new file mode 100644
--- /dev/null
+++ b/Engine/Plugins/ChaosVD/Source/ChaosVD/Public/ChaosVDSceneSelectionObserver.h
@@ -0,0 +1,39 @@
+﻿// Copyright Epic Games, Inc. All Rights Reserved.
+
+#pragma once
+#include "UObject/WeakObjectPtrTemplates.h"
+#include "UObject/WeakObjectPtr.h"
+
+class FChaosVDScene;
+class UTypedElementSelectionSet;
+
+/** Base class that provided access to selection events from the Chaos visual debugger
+ * local selection system.
+ */
+class CHAOSVD_API FChaosVDSceneSelectionObserver
+{
+public:
+
+	FChaosVDSceneSelectionObserver();
+	virtual ~FChaosVDSceneSelectionObserver();
+
+protected:
+
+	/** Stores a weak ptr to the selection set object containing the current selection state */
+	void RegisterSelectionSetObject(UTypedElementSelectionSet* SelectionSetObject);
+	
+	const UTypedElementSelectionSet* GetSelectionSetObject() const { return ObservedSelection.IsValid() ? ObservedSelection.Get() : nullptr; };
+
+	/** Called when the current selection might change
+	 * @param SelectionSetPreChange the pointer to the selection set object that might be about to change
+	 */
+	virtual void HandlePreSelectionChange(const UTypedElementSelectionSet* SelectionSetPreChange){}
+	
+	/** Called when the current selection has changed
+	 * @param ChangesSelectionSet the pointer to the selection set object that has changed
+	 */
+	virtual void HandlePostSelectionChange(const UTypedElementSelectionSet* ChangesSelectionSet){}
+
+	/** Weak Ptr to the selection set object we are observing */
+	TWeakObjectPtr<UTypedElementSelectionSet> ObservedSelection;
+};
diff --git a/Engine/Plugins/ChaosVD/Source/ChaosVD/Public/ChaosVDSolverDataSelection.h b/Engine/Plugins/ChaosVD/Source/ChaosVD/Public/ChaosVDSolverDataSelection.h
--- a/Engine/Plugins/ChaosVD/Source/ChaosVD/Public/ChaosVDSolverDataSelection.h
+++ b/Engine/Plugins/ChaosVD/Source/ChaosVD/Public/ChaosVDSolverDataSelection.h
@@ -24,6 +24,16 @@ struct FChaosVDSelectionMultipleView
 
 	template<typename StructType>
 	void AddData(StructType* Struct);
+	
+	void AddData(TSharedPtr<FStructOnScope> StructOnScope)
+	{
+		DataInstances.Add(StructOnScope);
+	}
+
+	void Clear()
+	{
+		DataInstances.Empty();
+	}
 
 protected:
 	TArray<TSharedPtr<FStructOnScope>> DataInstances;
diff --git a/Engine/Plugins/ChaosVD/Source/ChaosVD/Public/TEDS/ChaosVDSelectionInterface.h b/Engine/Plugins/ChaosVD/Source/ChaosVD/Public/TEDS/ChaosVDSelectionInterface.h
new file mode 100644
--- /dev/null
+++ b/Engine/Plugins/ChaosVD/Source/ChaosVD/Public/TEDS/ChaosVDSelectionInterface.h
@@ -0,0 +1,16 @@
+// Copyright Epic Games, Inc. All Rights Reserved.
+
+#pragma once
+
+#include "Elements/Interfaces/TypedElementSelectionInterface.h"
+#include "UObject/Object.h"
+#include "ChaosVDSelectionInterface.generated.h"
+
+/**
+ * CVD Custom Selection Interface. We don't have any specific implementation overrides yet
+ */
+UCLASS()
+class CHAOSVD_API UChaosVDSelectionInterface : public UObject, public ITypedElementSelectionInterface
+{
+	GENERATED_BODY()
+};
diff --git a/Engine/Plugins/ChaosVD/Source/ChaosVD/Public/TEDS/ChaosVDStructTypedElementData.h b/Engine/Plugins/ChaosVD/Source/ChaosVD/Public/TEDS/ChaosVDStructTypedElementData.h
new file mode 100644
--- /dev/null
+++ b/Engine/Plugins/ChaosVD/Source/ChaosVD/Public/TEDS/ChaosVDStructTypedElementData.h
@@ -0,0 +1,104 @@
+// Copyright Epic Games, Inc. All Rights Reserved.
+
+#pragma once
+
+#include "Elements/Framework/TypedElementData.h"
+#include "Elements/Framework/TypedElementHandle.h"
+#include "Elements/Framework/TypedElementOwnerStore.h"
+#include "Elements/Framework/TypedElementRegistry.h"
+
+class FStructOnScope;
+
+namespace Chaos::VD::TypedElementDataUtil
+{
+	static const FName NAME_CVD_StructDataElement = FName("CVD_StructDataElement");
+
+	struct CHAOSVD_API FStructTypedElementData
+	{
+		UE_DECLARE_TYPED_ELEMENT_DATA_RTTI(FStructTypedElementData);
+
+		UStruct* TypeInfo = nullptr;
+		void* RawData = nullptr;
+
+		template<typename StructDataType>
+		StructDataType* GetData() const;
+
+		TSharedPtr<FStructOnScope> GetDataAsStructScope() const;
+	};
+
+	extern TTypedElementOwnerStore<FStructTypedElementData, void*> GCVDTypedStructDataElementOwnerStore;
+
+	template <typename StructDataType>
+	StructDataType* FStructTypedElementData::GetData() const
+	{
+		if (TypeInfo && StructDataType::StaticStruct()->IsChildOf(TypeInfo))
+		{
+			return static_cast<StructDataType*>(RawData);
+		}
+
+		return nullptr;
+	}
+
+	template<typename StructDataType>
+	TTypedElementOwner<FStructTypedElementData> CreateTypedElementDataForStructData(StructDataType* InElementData)
+	{
+		UTypedElementRegistry* Registry = UTypedElementRegistry::GetInstance();
+
+		TTypedElementOwner<FStructTypedElementData> TypedElement;
+		if (ensureMsgf(Registry, TEXT("Typed element was requested for '%s' before the registry was available! This usually means that NewObject was used instead of CreateDefaultSubobject during CDO construction."), *InElementData->GetDisplayName()))
+		{
+			TypedElement = Registry->CreateElement<FStructTypedElementData>(NAME_CVD_StructDataElement);
+		}
+
+		if (TypedElement)
+		{
+			TypedElement.GetDataChecked().TypeInfo = StructDataType::StaticStruct();
+			TypedElement.GetDataChecked().RawData = InElementData;
+		}
+		return TypedElement;	
+	}
+
+	template<typename StructDataType>
+	FTypedElementHandle AcquireTypedElementHandleForStruct(StructDataType* ElementInstance, const bool bAllowCreate)
+	{
+		if (!ElementInstance)
+		{
+			return FTypedElementHandle();
+		}
+
+		TTypedElementOwnerScopedAccess<FStructTypedElementData> EditorElement = bAllowCreate
+																			   ? GCVDTypedStructDataElementOwnerStore.FindOrRegisterElementOwner(ElementInstance, [ElementInstance]()
+																			   {
+																				   return CreateTypedElementDataForStructData(ElementInstance);
+																			   })
+																			   : GCVDTypedStructDataElementOwnerStore.FindElementOwner(ElementInstance);
+		
+		if (EditorElement)
+		{
+			return EditorElement->AcquireHandle();
+		}
+
+		return FTypedElementHandle();
+	}
+
+	void DestroyTypedElementHandleForStruct(void* InElementData);
+
+	template <typename StructDataType>
+	StructDataType* GetStructDataFromTypedElementHandle(const FTypedElementHandle& InHandle, const bool bSilent = true)
+	{
+		const FStructTypedElementData* StructElement = InHandle.GetData<FStructTypedElementData>(bSilent);
+		return StructElement ? StructElement->GetData<StructDataType>() : nullptr;
+	}
+
+	CHAOSVD_API TSharedPtr<FStructOnScope> GetStructOnScopeDataFromTypedElementHandle(const FTypedElementHandle& InHandle, const bool bSilent = false);
+
+	void CleanUpTypedElementStore();
+}
+
+template <>
+inline FString GetTypedElementDebugId< Chaos::VD::TypedElementDataUtil::FStructTypedElementData>(const  Chaos::VD::TypedElementDataUtil::FStructTypedElementData& InElementData)
+{
+	return FString::Printf(TEXT("[%s] - [%p]"), InElementData.TypeInfo ? *InElementData.TypeInfo->GetName() : TEXT("Invalid Struct info"), InElementData.RawData);
+}
+
+
diff --git a/Engine/Plugins/ChaosVD/Source/ChaosVD/Public/Widgets/SChaosVDDetailsView.h b/Engine/Plugins/ChaosVD/Source/ChaosVD/Public/Widgets/SChaosVDDetailsView.h
new file mode 100644
--- /dev/null
+++ b/Engine/Plugins/ChaosVD/Source/ChaosVD/Public/Widgets/SChaosVDDetailsView.h
@@ -0,0 +1,74 @@
+// Copyright Epic Games, Inc. All Rights Reserved.
+
+#pragma once
+
+#include "IStructureDetailsView.h"
+#include "Templates/SharedPointer.h"
+#include "UObject/StructOnScope.h"
+#include "UObject/WeakObjectPtrTemplates.h"
+#include "Widgets/SCompoundWidget.h"
+
+class SChaosVDMainTab;
+class IStructureDetailsView;
+class FSubobjectEditorTreeNode;
+class AActor;
+class SBox;
+class SSplitter;
+class IDetailsView;
+class SSubobjectEditor;
+class UObject;
+
+/**
+ * Simple details for CVD objects and structures
+ */
+class CHAOSVD_API SChaosVDDetailsView : public SCompoundWidget
+{
+public:
+	SLATE_BEGIN_ARGS(SChaosVDDetailsView)
+		{
+		}
+	SLATE_END_ARGS()
+
+	/** Constructs this widget with InArgs */
+	void Construct(const FArguments& InArgs, const TSharedRef<SChaosVDMainTab>& InMainTab);
+
+	/** Updates the current object this view details is viewing */
+	void SetSelectedObject(UObject* NewObject);
+
+	/** Updates the current object this view details is viewing */
+	template<typename TStruct>
+	void SetSelectedStruct(TStruct* NewStruct);
+
+	void SetSelectedStruct(const TSharedPtr<FStructOnScope>& NewStruct);
+
+protected:
+
+	TSharedPtr<IDetailsView> CreateObjectDetailsView();
+	TSharedPtr<IStructureDetailsView> CreateStructureDataDetailsView() const;
+
+	EVisibility GetStructDetailsVisibility() const;
+	EVisibility GetObjectDetailsVisibility() const;
+
+	TWeakObjectPtr<UObject> CurrentObjectInView;
+
+	TWeakPtr<FStructOnScope> CurrentStructInView;
+	
+	TSharedPtr<IDetailsView> DetailsView;
+
+	TSharedPtr<IStructureDetailsView> StructDetailsView;
+
+	TWeakPtr<SChaosVDMainTab> MainTabWeakPtr;
+};
+
+template <typename TStruct>
+void SChaosVDDetailsView::SetSelectedStruct(TStruct* NewStruct)
+{
+	TSharedPtr<FStructOnScope> StructDataView = nullptr;
+
+	if (NewStruct)
+	{
+		StructDataView = MakeShared<FStructOnScope>(TStruct::StaticStruct(), reinterpret_cast<uint8*>(NewStruct));
+	}
+
+	SetSelectedStruct(StructDataView);
+}
diff --git a/Engine/Plugins/ChaosVD/Source/ChaosVD/Public/Widgets/SChaosVDMainTab.h b/Engine/Plugins/ChaosVD/Source/ChaosVD/Public/Widgets/SChaosVDMainTab.h
--- a/Engine/Plugins/ChaosVD/Source/ChaosVD/Public/Widgets/SChaosVDMainTab.h
+++ b/Engine/Plugins/ChaosVD/Source/ChaosVD/Public/Widgets/SChaosVDMainTab.h
@@ -29,6 +29,7 @@ class FStructOnScope;
 class FUICommandList;
 class SButton;
 class SDockTab;
+class FChaosVDScene;
 
 struct FDetailsViewArgs;
 struct FStructureDetailsViewArgs;
@@ -63,6 +64,8 @@ public:
 		return ChaosVDEngine.ToSharedRef();
 	}
 
+	CHAOSVD_API TSharedPtr<FChaosVDScene> GetScene();
+
 	template<typename TabType>
 	TWeakPtr<TabType> GetTabSpawnerInstance(FName TabID);
 
@@ -144,8 +147,8 @@ private:
 	void RegisterMainTabMenu();
 	void RegisterSettingsMenu();
 
-	void HandleTabSpawned(TSharedRef<SDockTab> Tab, FName TabID);
-	void HandleTabDestroyed(TSharedRef<SDockTab> Tab, FName TabID);
+	CHAOSVD_API void HandleTabSpawned(TSharedRef<SDockTab> Tab, FName TabID);
+	CHAOSVD_API void HandleTabDestroyed(TSharedRef<SDockTab> Tab, FName TabID);
 
 	TSharedRef<FTabManager::FLayout> GenerateMainLayout();
 
diff --git a/Engine/Plugins/Experimental/Mover/Mover.uplugin b/Engine/Plugins/Experimental/Mover/Mover.uplugin
--- a/Engine/Plugins/Experimental/Mover/Mover.uplugin
+++ b/Engine/Plugins/Experimental/Mover/Mover.uplugin
@@ -20,10 +20,21 @@
 			"Type": "Runtime",
 			"LoadingPhase": "Default"
 		},
+		{
+			"Name": "MoverCVDData",
+			"Type": "Runtime",
+			"LoadingPhase": "Default",
+			"TargetConfigurationDenyList": [ "Shipping" ]
+		},
 		{
 			"Name": "MoverEditor",
 			"Type": "UncookedOnly",
 			"LoadingPhase": "PreDefault"
+		},
+		{
+			"Name": "MoverCVDEditor",
+			"Type": "UncookedOnly",
+			"LoadingPhase": "PreDefault"
 		}
 	],
 	"Plugins": [
@@ -42,6 +53,10 @@
 		{
 			"Name": "Water",
 			"Enabled": true
+		},
+		{
+			"Name": "ChaosVD",
+			"Enabled": true
 		}
 	]
 }
\ No newline at end of file
diff --git a/Engine/Plugins/Experimental/Mover/Source/Mover/Mover.Build.cs b/Engine/Plugins/Experimental/Mover/Source/Mover/Mover.Build.cs
--- a/Engine/Plugins/Experimental/Mover/Source/Mover/Mover.Build.cs
+++ b/Engine/Plugins/Experimental/Mover/Source/Mover/Mover.Build.cs
@@ -36,6 +36,12 @@ public class Mover : ModuleRules
 			}
 			);
 
+		if (IsChaosVisualDebuggerSupported(Target))
+		{
+			PublicDependencyModuleNames.Add("MoverCVDData");
+		}
+		SetupModuleChaosVisualDebuggerSupport(Target);
+
 		SetupGameplayDebuggerSupport(Target);
 	}
 }
diff --git a/Engine/Plugins/Experimental/Mover/Source/Mover/Private/Backends/MoverNetworkPhysicsLiaison.cpp b/Engine/Plugins/Experimental/Mover/Source/Mover/Private/Backends/MoverNetworkPhysicsLiaison.cpp
--- a/Engine/Plugins/Experimental/Mover/Source/Mover/Private/Backends/MoverNetworkPhysicsLiaison.cpp
+++ b/Engine/Plugins/Experimental/Mover/Source/Mover/Private/Backends/MoverNetworkPhysicsLiaison.cpp
@@ -16,6 +16,7 @@
 #include "PhysicsMover/PhysicsMoverSimulationTypes.h"
 #include "DefaultMovementSet/Modes/FallingMode.h"
 #include "Misc/DataValidation.h"
+#include "ChaosVisualDebugger/MoverCVDRuntimeTrace.h"
 
 #define LOCTEXT_NAMESPACE "Mover"
 
@@ -298,6 +299,9 @@ void UMoverNetworkPhysicsLiaisonComponent::PerformProcessInputs_Internal(int32 P
 		CommonMovementSettings->MaxSpeed = MovementSettings->MaxSpeed;
 		CommonMovementSettings->Acceleration = MovementSettings->Acceleration;
 	}
+	
+	// This will only do something if CVD is actively tracing and the mover info CVD data channel checked
+	FMoverCVDRuntimeTrace::TraceMoverData(&GetMoverComponent(), &Input.InputCmd, &Input.SyncState);
 }
 
 bool UMoverNetworkPhysicsLiaisonComponent::CanSimulate_Internal(const FPhysicsMoverSimulationTickParams& TickParams, const FPhysicsMoverAsyncInput& Input) const
diff --git a/Engine/Plugins/Experimental/Mover/Source/Mover/Private/ChaosVisualDebugger/MoverCVDRuntimeTrace.cpp b/Engine/Plugins/Experimental/Mover/Source/Mover/Private/ChaosVisualDebugger/MoverCVDRuntimeTrace.cpp
new file mode 100644
--- /dev/null
+++ b/Engine/Plugins/Experimental/Mover/Source/Mover/Private/ChaosVisualDebugger/MoverCVDRuntimeTrace.cpp
@@ -0,0 +1,153 @@
+// Copyright Epic Games, Inc. All Rights Reserved.
+
+#include "ChaosVisualDebugger/MoverCVDRuntimeTrace.h"
+
+#if WITH_CHAOS_VISUAL_DEBUGGER
+
+#include "MoverCVDDataWrappers.h"
+#include "MoverSimulationTypes.h"
+#include "Serialization/MemoryWriter.h"
+#include "Serialization/MemoryReader.h"
+#include "ChaosVisualDebugger/ChaosVisualDebuggerTrace.h"
+#include "MoverComponent.h"
+#include "Chaos/PhysicsObjectInterface.h"
+#include "Chaos/PhysicsObjectInternalInterface.h"
+#include "PhysicsProxy/SingleParticlePhysicsProxy.h"
+#include "ChaosVisualDebugger/ChaosVDTraceMacros.h"
+#include "PBDRigidsSolver.h"
+#include "Physics/Experimental/PhysScene_Chaos.h"
+#include "Interfaces/IPhysicsComponent.h"
+
+CVD_DEFINE_OPTIONAL_DATA_CHANNEL(MoverData, EChaosVDDataChannelInitializationFlags::CanChangeEnabledState)
+
+// FSkipObjectRefsMemoryWriter and FSkipObjectRefsMemoryReader are a 
+// workaround for serializing mover info structs with object references in them, such as the mover base
+// Ultimately we might want to attempt to translate the object to a particle ID and resolve that on the CVD side
+// by customizing the display and show the particle ID instead
+class FSkipObjectRefsMemoryWriter : public FMemoryWriter
+{
+public:
+	FSkipObjectRefsMemoryWriter(TArray<uint8, TSizedDefaultAllocator<32>>& InBytes, bool bIsPersistent = false, bool bSetOffset = false, const FName InArchiveName = NAME_None)
+		: FMemoryWriter(InBytes, bIsPersistent, bSetOffset, InArchiveName)
+	{
+	}
+
+	virtual FArchive& operator<<(struct FObjectPtr& Value)
+	{
+		return *this;
+	}
+};
+class FSkipObjectRefsMemoryReader : public FMemoryReader
+{
+public:
+	FSkipObjectRefsMemoryReader(const TArray<uint8, TSizedDefaultAllocator<32>>& InBytes, bool bIsPersistent = false)
+		: FMemoryReader(InBytes, bIsPersistent)
+	{
+	}
+
+	virtual FArchive& operator<<(struct FObjectPtr& Value)
+	{
+		return *this;
+	}
+};
+
+void FMoverCVDRuntimeTrace::UnwrapSimData(const FMoverCVDSimDataWrapper& InSimDataWrapper, TSharedPtr<FMoverInputCmdContext>& OutInputCmd, TSharedPtr<FMoverSyncState>& OutSyncState)
+{
+	// Consider using Chaos::VisualDebugger::FChaosVDMemoryReader
+
+	// Input Cmd
+	{
+		// Deserialize into a new struct allocated dynamically
+		FMemoryReader ArReader(InSimDataWrapper.InputCmdBytes, true);
+		FMoverInputCmdContext* NewStruct = new FMoverInputCmdContext;
+		FMoverInputCmdContext::StaticStruct()->SerializeBin(ArReader, NewStruct);
+		// Input cmd's Collection of Custom Input Structs
+		FSkipObjectRefsMemoryReader ArInputCollectionReader(InSimDataWrapper.InputMoverDataCollectionBytes, true);
+		NewStruct->InputCollection.SerializeDebugData(ArInputCollectionReader);
+		OutInputCmd = TSharedPtr<FMoverInputCmdContext>(NewStruct);
+	}
+
+	// Sync State
+	{
+		// Deserialize into a new struct allocated dynamically
+		FMemoryReader ArReader(InSimDataWrapper.SyncStateBytes, true);
+		FMoverSyncState* NewStruct = new FMoverSyncState;
+		FMoverSyncState::StaticStruct()->SerializeBin(ArReader, NewStruct);
+		OutSyncState = TSharedPtr<FMoverSyncState>(NewStruct);
+	}
+}
+
+void FMoverCVDRuntimeTrace::WrapSimData(uint32 SolverID, uint32 ParticleID, const FMoverInputCmdContext& InInputCmd, const FMoverSyncState& InSyncState, FMoverCVDSimDataWrapper& OutSimDataWrapper)
+{
+	OutSimDataWrapper.SolverID = SolverID;
+	OutSimDataWrapper.ParticleID = ParticleID;
+
+	// Consider using Chaos::VisualDebugger::FChaosVDMemoryWriter
+
+	// Input Cmd
+	{
+		FMemoryWriter ArWriter(OutSimDataWrapper.InputCmdBytes, true);
+		// This is not version friendly, we need to instead use SerializeTagParticles. Slower and Sergio is working on a faster version, but it's not available yet.
+		InInputCmd.StaticStruct()->SerializeBin(ArWriter, &const_cast<FMoverInputCmdContext&>(InInputCmd));
+	}
+
+	// Input cmd's Collection of Custom Input Structs
+	{
+		FSkipObjectRefsMemoryWriter ArWriter(OutSimDataWrapper.InputMoverDataCollectionBytes, true);
+		const_cast<FMoverInputCmdContext&>(InInputCmd).InputCollection.SerializeDebugData(ArWriter);
+	}
+
+	// Sync State
+	{
+		FMemoryWriter ArWriter(OutSimDataWrapper.SyncStateBytes, true);
+		// This is not version friendly, we need to instead use SerializeTagParticles. Slower and Sergio is working on a faster version, but it's not available yet.
+		InSyncState.StaticStruct()->SerializeBin(ArWriter, &const_cast<FMoverSyncState&>(InSyncState));
+	}
+}
+
+void FMoverCVDRuntimeTrace::TraceMoverData(UMoverComponent* MoverComponent, const FMoverInputCmdContext* InputCmd, const FMoverSyncState* SyncState)
+{
+	if (!FChaosVisualDebuggerTrace::IsTracing())
+	{
+		return;
+	}
+
+	if (!CVDDC_MoverData->IsChannelEnabled())
+	{
+		return;
+	}
+	
+	if (!InputCmd || !SyncState)
+	{
+		return;
+	}
+
+	if (UWorld* World = MoverComponent->GetWorld())
+	{
+		int32 ParticleID = INDEX_NONE;
+		if (const IPhysicsComponent* PhysicsComponent = Cast<IPhysicsComponent>(MoverComponent->GetUpdatedComponent()))
+		{
+			Chaos::FReadPhysicsObjectInterface_Internal Interface = Chaos::FPhysicsObjectInternalInterface::GetRead();
+			const Chaos::FPhysicsObject* PhysicsObject = PhysicsComponent ? PhysicsComponent->GetPhysicsObjectById(0) : nullptr; // get the root physics object
+			const Chaos::FGeometryParticleHandle* ParticleHandle = PhysicsObject? Interface.GetParticle(PhysicsObject) : nullptr;
+			ParticleID = ParticleHandle ? ParticleHandle->UniqueIdx().Idx : INDEX_NONE;
+		}
+
+		int32 SolverID = CVD_TRACE_GET_SOLVER_ID_FROM_WORLD(World);
+		TraceMoverDataPrivate(SolverID, ParticleID, InputCmd, SyncState);
+	}
+}
+
+void FMoverCVDRuntimeTrace::TraceMoverDataPrivate(uint32 SolverID, uint32 ParticleID, const FMoverInputCmdContext* InputCmd, const FMoverSyncState* SyncState)
+{
+	FMoverCVDSimDataWrapper SimDataWrapper;
+	WrapSimData(SolverID, ParticleID, *InputCmd, *SyncState, SimDataWrapper);
+	SimDataWrapper.MarkAsValid();
+
+	FChaosVDScopedTLSBufferAccessor TLSDataBuffer;
+	Chaos::VisualDebugger::WriteDataToBuffer(TLSDataBuffer.BufferRef, SimDataWrapper);
+
+	FChaosVisualDebuggerTrace::TraceBinaryData(TLSDataBuffer.BufferRef, FMoverCVDSimDataWrapper::WrapperTypeName);
+}
+
+#endif // WITH_CHAOS_VISUAL_DEBUGGER
diff --git a/Engine/Plugins/Experimental/Mover/Source/Mover/Private/MoverComponent.cpp b/Engine/Plugins/Experimental/Mover/Source/Mover/Private/MoverComponent.cpp
--- a/Engine/Plugins/Experimental/Mover/Source/Mover/Private/MoverComponent.cpp
+++ b/Engine/Plugins/Experimental/Mover/Source/Mover/Private/MoverComponent.cpp
@@ -25,7 +25,7 @@
 #include "UObject/ObjectSaveContext.h"
 #include "Components/SkeletalMeshComponent.h"
 #include "MotionWarpingComponent.h"
-
+#include "ChaosVisualDebugger/MoverCVDRuntimeTrace.h"
 
 #if WITH_EDITOR
 #include "Misc/DataValidation.h"
@@ -421,6 +421,9 @@ void UMoverComponent::InitializeSimulationState(FMoverSyncState* OutSync, FMover
 
 void UMoverComponent::SimulationTick(const FMoverTimeStep& InTimeStep, const FMoverTickStartData& SimInput, OUT FMoverTickEndData& SimOutput)
 {
+	// Send mover info to the Chaos Visual Debugger (this will do nothing if CVD is not recording, or the mover info data channel not enabled)
+	FMoverCVDRuntimeTrace::TraceMoverData(this, &SimInput.InputCmd, &SimInput.SyncState);
+
 	const bool bIsResimulating = InTimeStep.BaseSimTimeMs <= CachedNewestSimTickTimeStep.BaseSimTimeMs;
 
 	FMoverTimeStep MoverTimeStep(InTimeStep);
diff --git a/Engine/Plugins/Experimental/Mover/Source/Mover/Private/MoverTypes.cpp b/Engine/Plugins/Experimental/Mover/Source/Mover/Private/MoverTypes.cpp
--- a/Engine/Plugins/Experimental/Mover/Source/Mover/Private/MoverTypes.cpp
+++ b/Engine/Plugins/Experimental/Mover/Source/Mover/Private/MoverTypes.cpp
@@ -84,7 +84,68 @@ bool FMoverDataCollection::NetSerialize(FArchive& Ar, UPackageMap* Map, bool& bO
 	return true;
 }
 
+struct FMoverDataDeleter
+{
+	FORCEINLINE void operator()(FMoverDataStructBase* Object) const
+	{
+		check(Object);
+		UScriptStruct* ScriptStruct = Object->GetScriptStruct();
+		check(ScriptStruct);
+		ScriptStruct->DestroyStruct(Object);
+		FMemory::Free(Object);
+	}
+};
 
+bool FMoverDataCollection::SerializeDebugData(FArchive& Ar)
+{
+	// DISCLAIMER: This serialization is not version independent, so it might not be good enough to be used for the Chaos Visual Debugger in the long run
+
+	// First serialize the number of structs in the collection
+	int32 NumDataStructs;
+	if (Ar.IsLoading())
+	{
+		Ar << NumDataStructs;
+		DataArray.SetNumZeroed(NumDataStructs);
+	}
+	else
+	{
+		NumDataStructs = DataArray.Num();
+		Ar << NumDataStructs;
+	}
+
+	if (Ar.IsLoading())
+	{
+		DataArray.Empty();
+		for (int32 i = 0; i < NumDataStructs && !Ar.IsError(); ++i)
+		{
+			FString StructName;
+			Ar << StructName;
+			if (UScriptStruct* MoveDataStruct = Cast<UScriptStruct>(FindObject<UStruct>(nullptr, *StructName)))
+			{
+				FMoverDataStructBase* NewMoverData = AddDataByType(MoveDataStruct);
+				MoveDataStruct->SerializeBin(Ar, NewMoverData);
+			}
+		}
+	}
+	else
+	{
+		for (int32 i = 0; i < DataArray.Num() && !Ar.IsError(); ++i)
+		{
+			FMoverDataStructBase* MoveDataStruct = DataArray[i].Get();
+			if (MoveDataStruct)
+			{
+				// The FullName of the script struct will be something like "ScriptStruct /Script/Mover.FCharacterDefaultInputs"
+				FString FullStructName = MoveDataStruct->GetScriptStruct()->GetFullName(nullptr);
+				// We don't need to save the first part since we only ever save UScriptStructs (C++ structs)
+				FString StructName = FullStructName.RightChop(13); // So we chop the "ScriptStruct " part (hence 13 characters)
+				Ar << StructName;
+				MoveDataStruct->GetScriptStruct()->SerializeBin(Ar, MoveDataStruct);
+			}
+		}
+	}
+
+	return true;
+}
 
 FMoverDataCollection& FMoverDataCollection::operator=(const FMoverDataCollection& Other)
 {
@@ -252,18 +313,6 @@ TArray<TSharedPtr<FMoverDataStructBase>>::TConstIterator FMoverDataCollection::G
 	return DataArray.CreateConstIterator();
 }
 
-struct FMoverDataDeleter
-{
-	FORCEINLINE void operator()(FMoverDataStructBase* Object) const
-	{
-		check(Object);
-		UScriptStruct* ScriptStruct = Object->GetScriptStruct();
-		check(ScriptStruct);
-		ScriptStruct->DestroyStruct(Object);
-		FMemory::Free(Object);
-	}
-};
-
 //static 
 TSharedPtr<FMoverDataStructBase> FMoverDataCollection::CreateDataByType(const UScriptStruct* DataStructType)
 {
diff --git a/Engine/Plugins/Experimental/Mover/Source/Mover/Public/ChaosVisualDebugger/MoverCVDRuntimeTrace.h b/Engine/Plugins/Experimental/Mover/Source/Mover/Public/ChaosVisualDebugger/MoverCVDRuntimeTrace.h
new file mode 100644
--- /dev/null
+++ b/Engine/Plugins/Experimental/Mover/Source/Mover/Public/ChaosVisualDebugger/MoverCVDRuntimeTrace.h
@@ -0,0 +1,41 @@
+// Copyright Epic Games, Inc. All Rights Reserved.
+
+#pragma once
+
+#if WITH_CHAOS_VISUAL_DEBUGGER
+
+#include "Containers/Array.h"
+#include "HAL/Platform.h"
+
+#include "ChaosVisualDebugger/ChaosVDOptionalDataChannel.h"
+
+struct FMoverCVDSimDataWrapper;
+struct FMoverInputCmdContext;
+struct FMoverSyncState;
+class UMoverComponent;
+
+/** Utility functions used to trace Mover data into the Chaos Visual Debugger */
+class FMoverCVDRuntimeTrace
+{
+public:
+	static void TraceMoverData(UMoverComponent* MoverComponent, const FMoverInputCmdContext* InputCmd, const FMoverSyncState* SyncState);
+
+	MOVER_API static void UnwrapSimData(const FMoverCVDSimDataWrapper& InSimDataWrapper, TSharedPtr<FMoverInputCmdContext>& OutInputCmd, TSharedPtr<FMoverSyncState>& OutSyncState);
+	static void WrapSimData(uint32 SolverID, uint32 ParticleID, const FMoverInputCmdContext& InInputCmd, const FMoverSyncState& InSyncState, FMoverCVDSimDataWrapper& OutSimDataWrapper);
+
+private:
+	static void TraceMoverDataPrivate(uint32 SolverID, uint32 ParticleID, const FMoverInputCmdContext* InputCmd, const FMoverSyncState* SyncState);
+};
+
+CVD_DECLARE_OPTIONAL_DATA_CHANNEL_EXTERN(MoverData, MOVER_API)
+
+#else
+
+// Noop implementation in case this is compiled without Chaos Visual Debugger support (e.g. shipping)
+class FMoverCVDRuntimeTrace
+{
+public:
+	static void TraceMoverData(UMoverComponent* MoverComponent, const FMoverInputCmdContext* InputCmd, const FMoverSyncState* SyncState) {}
+};
+
+#endif // WITH_CHAOS_VISUAL_DEBUGGER
diff --git a/Engine/Plugins/Experimental/Mover/Source/Mover/Public/MoverDataModelTypes.h b/Engine/Plugins/Experimental/Mover/Source/Mover/Public/MoverDataModelTypes.h
--- a/Engine/Plugins/Experimental/Mover/Source/Mover/Public/MoverDataModelTypes.h
+++ b/Engine/Plugins/Experimental/Mover/Source/Mover/Public/MoverDataModelTypes.h
@@ -37,28 +37,28 @@ struct MOVER_API FCharacterDefaultInputs : public FMoverDataStructBase
 	EMoveInputType GetMoveInputType() const { return MoveInputType; }
 
 protected:
-	UPROPERTY(BlueprintReadWrite, Category = Mover)
+	UPROPERTY(VisibleAnywhere, BlueprintReadWrite, Category = Mover)
 	EMoveInputType MoveInputType;
 
 	// Representing the directional move input for this frame. Must be interpreted according to MoveInputType. Relative to MovementBase if set, world space otherwise. Will be truncated to match network serialization precision.
-	UPROPERTY(BlueprintReadWrite, Category = Mover)
+	UPROPERTY(VisibleAnywhere, BlueprintReadWrite, Category = Mover)
 	FVector MoveInput;
 
 public:
 	// Facing direction intent, as a normalized forward-facing direction. A zero vector indicates no intent to change facing direction. Relative to MovementBase if set, world space otherwise.
-	UPROPERTY(BlueprintReadWrite, Category = Mover)
+	UPROPERTY(VisibleAnywhere, BlueprintReadWrite, Category = Mover)
 	FVector OrientationIntent;
 
 	// World space orientation that the controls were based on. This is commonly a player's camera rotation.
-	UPROPERTY(BlueprintReadWrite, Category = Mover)
+	UPROPERTY(VisibleAnywhere, BlueprintReadWrite, Category = Mover)
 	FRotator ControlRotation;
 
 	// Used to force the Mover actor into a different movement mode
-	UPROPERTY(BlueprintReadWrite, Category = Mover)
+	UPROPERTY(VisibleAnywhere, BlueprintReadWrite, Category = Mover)
 	FName SuggestedMovementMode;
 
 	// Specifies whether we are using a movement base, which will affect how move inputs are interpreted
-	UPROPERTY(BlueprintReadWrite, Category = Mover)
+	UPROPERTY(VisibleAnywhere, BlueprintReadWrite, Category = Mover)
 	bool bUsingMovementBase;
 
 	// Optional: when moving on a base, input may be relative to this object
@@ -66,12 +66,12 @@ public:
 	TObjectPtr<UPrimitiveComponent> MovementBase;
 
 	// Optional: for movement bases that are skeletal meshes, this is the bone we're based on. Only valid if MovementBase is set.
-	UPROPERTY(BlueprintReadWrite, Category = Mover)
+	UPROPERTY(VisibleAnywhere, BlueprintReadWrite, Category = Mover)
 	FName MovementBaseBoneName;
 
-	UPROPERTY(BlueprintReadWrite, Category = Mover)
+	UPROPERTY(VisibleAnywhere, BlueprintReadWrite, Category = Mover)
 	bool bIsJumpJustPressed;
-	UPROPERTY(BlueprintReadWrite, Category = Mover)
+	UPROPERTY(VisibleAnywhere, BlueprintReadWrite, Category = Mover)
 	bool bIsJumpPressed;
 
 	FVector GetMoveInput_WorldSpace() const;
diff --git a/Engine/Plugins/Experimental/Mover/Source/Mover/Public/MoverTypes.h b/Engine/Plugins/Experimental/Mover/Source/Mover/Public/MoverTypes.h
--- a/Engine/Plugins/Experimental/Mover/Source/Mover/Public/MoverTypes.h
+++ b/Engine/Plugins/Experimental/Mover/Source/Mover/Public/MoverTypes.h
@@ -207,6 +207,11 @@ struct MOVER_API FMoverDataCollection
 	/** Serialize all data in this collection */
 	bool NetSerialize(FArchive& Ar, UPackageMap* Map, bool& bOutSuccess);
 
+	/** Serializes data in this collection for debug purposes.
+	*   This is currently only usable in the context of sending mover info to the Chaos Visual Debugger
+	*/
+	bool SerializeDebugData(FArchive& Ar);
+
 	/** Copy operator - deep copy so it can be used for archiving/saving off data */
 	FMoverDataCollection& operator=(const FMoverDataCollection& Other);
 
@@ -288,6 +293,11 @@ struct MOVER_API FMoverDataCollection
 	/** Adds data to the collection. If an existing data struct of the same type is already there, it will be removed first. */
 	void AddOrOverwriteData(const TSharedPtr<FMoverDataStructBase> DataInstance);
 
+	const TArray< TSharedPtr<FMoverDataStructBase> >& GetDataArray() const
+	{
+		return DataArray;
+	}
+
 protected:
 	static TSharedPtr<FMoverDataStructBase> CreateDataByType(const UScriptStruct* DataStructType);
 	FMoverDataStructBase* AddDataByType(const UScriptStruct* DataStructType);
diff --git a/Engine/Plugins/Experimental/Mover/Source/MoverCVDData/MoverCVDData.Build.cs b/Engine/Plugins/Experimental/Mover/Source/MoverCVDData/MoverCVDData.Build.cs
new file mode 100644
--- /dev/null
+++ b/Engine/Plugins/Experimental/Mover/Source/MoverCVDData/MoverCVDData.Build.cs
@@ -0,0 +1,20 @@
+using UnrealBuildTool;
+
+public class MoverCVDData : ModuleRules
+{
+    public MoverCVDData(ReadOnlyTargetRules Target) : base(Target)
+    {
+	    PublicDependencyModuleNames.AddRange(
+			new string[] {
+				"Core",
+				"CoreUObject",
+			}
+		);
+
+		PrivateDependencyModuleNames.AddRange(
+			new string[] {
+				"ChaosVDRuntime",
+			}
+		);
+	}
+}
\ No newline at end of file
diff --git a/Engine/Plugins/Experimental/Mover/Source/MoverCVDData/Private/MoverCVDData.cpp b/Engine/Plugins/Experimental/Mover/Source/MoverCVDData/Private/MoverCVDData.cpp
new file mode 100644
--- /dev/null
+++ b/Engine/Plugins/Experimental/Mover/Source/MoverCVDData/Private/MoverCVDData.cpp
@@ -0,0 +1,5 @@
+// Copyright Epic Games, Inc. All Rights Reserved.
+
+#include "Modules/ModuleManager.h"
+    
+IMPLEMENT_MODULE(FDefaultModuleImpl, MoverCVDData)
\ No newline at end of file
diff --git a/Engine/Plugins/Experimental/Mover/Source/MoverCVDData/Private/MoverCVDDataWrappers.cpp b/Engine/Plugins/Experimental/Mover/Source/MoverCVDData/Private/MoverCVDDataWrappers.cpp
new file mode 100644
--- /dev/null
+++ b/Engine/Plugins/Experimental/Mover/Source/MoverCVDData/Private/MoverCVDDataWrappers.cpp
@@ -0,0 +1,21 @@
+// Copyright Epic Games, Inc. All Rights Reserved.
+
+#include "MoverCVDDataWrappers.h"
+
+bool FMoverCVDSimDataWrapper::Serialize(FArchive& Ar)
+{
+	Ar << bHasValidData;
+
+	if (!bHasValidData)
+	{
+		return !Ar.IsError();
+	}
+
+	Ar << SolverID;
+	Ar << ParticleID;
+	Ar << SyncStateBytes;
+	Ar << InputCmdBytes;
+	Ar << InputMoverDataCollectionBytes;
+
+	return !Ar.IsError();
+}
diff --git a/Engine/Plugins/Experimental/Mover/Source/MoverCVDData/Public/MoverCVDDataWrappers.h b/Engine/Plugins/Experimental/Mover/Source/MoverCVDData/Public/MoverCVDDataWrappers.h
new file mode 100644
--- /dev/null
+++ b/Engine/Plugins/Experimental/Mover/Source/MoverCVDData/Public/MoverCVDDataWrappers.h
@@ -0,0 +1,47 @@
+// Copyright Epic Games, Inc. All Rights Reserved.
+
+#pragma once
+
+#include "DataWrappers/ChaosVDDataSerializationMacros.h"
+#include "DataWrappers/ChaosVDParticleDataWrapper.h"
+#include "UObject/ObjectMacros.h"
+
+#include "MoverCVDDataWrappers.generated.h"
+
+USTRUCT(DisplayName="Mover Particle Info")
+struct FMoverCVDSimDataParticleInfo
+{
+	GENERATED_BODY()
+
+	UPROPERTY(VisibleAnywhere, Category="Mover Info")
+	int32 SolverID = INDEX_NONE;
+
+	UPROPERTY(VisibleAnywhere, Category="Mover Info")
+	int32 ParticleID = INDEX_NONE;
+};
+
+USTRUCT(DisplayName="Mover Sim Data")
+struct FMoverCVDSimDataWrapper : public FChaosVDWrapperDataBase
+{
+	GENERATED_BODY()
+	
+	inline static FStringView WrapperTypeName = TEXT("FMoverCVDSimDataWrapper");
+
+	int32 SolverID = INDEX_NONE;
+	int32 ParticleID = INDEX_NONE;
+	TArray<uint8> SyncStateBytes;
+	TArray<uint8> InputCmdBytes;
+	TArray<uint8> InputMoverDataCollectionBytes;
+
+	MOVERCVDDATA_API bool Serialize(FArchive& Ar);
+};
+
+CVD_IMPLEMENT_SERIALIZER(FMoverCVDSimDataWrapper)
+
+USTRUCT()
+struct FMoverCVDSimDataContainer
+{
+	GENERATED_BODY()
+
+	TMap<int32, TArray<TSharedPtr<FMoverCVDSimDataWrapper>>> SimDataBySolverID;
+};
\ No newline at end of file
diff --git a/Engine/Plugins/Experimental/Mover/Source/MoverCVDEditor/MoverCVDEditor.Build.cs b/Engine/Plugins/Experimental/Mover/Source/MoverCVDEditor/MoverCVDEditor.Build.cs
new file mode 100644
--- /dev/null
+++ b/Engine/Plugins/Experimental/Mover/Source/MoverCVDEditor/MoverCVDEditor.Build.cs
@@ -0,0 +1,39 @@
+using UnrealBuildTool;
+
+public class MoverCVDEditor : ModuleRules
+{
+    public MoverCVDEditor(ReadOnlyTargetRules Target) : base(Target)
+    {
+		PublicDependencyModuleNames.AddRange(
+			new string[]
+			{
+				"Core"
+			});
+
+		PrivateDependencyModuleNames.AddRange(
+		    new string[]
+		    {
+				"ApplicationCore",
+				"ChaosVD",
+				"ChaosVDData",
+				"ChaosVDRuntime",
+				"CoreUObject",
+				"EditorWidgets",
+				"Engine",
+				"Mover",
+				"MoverCVDData",
+				"Projects",
+				"Slate",
+				"SlateCore",
+				"ToolMenus",
+				"TraceServices",
+				"TypedElementFramework",
+				"UnrealEd",
+			}
+		);
+        
+        SetupModulePhysicsSupport(Target);
+
+		SetupModuleChaosVisualDebuggerSupport(Target);
+	}
+}
\ No newline at end of file
diff --git a/Engine/Plugins/Experimental/Mover/Source/MoverCVDEditor/Private/MoverCVDEditor.cpp b/Engine/Plugins/Experimental/Mover/Source/MoverCVDEditor/Private/MoverCVDEditor.cpp
new file mode 100644
--- /dev/null
+++ b/Engine/Plugins/Experimental/Mover/Source/MoverCVDEditor/Private/MoverCVDEditor.cpp
@@ -0,0 +1,32 @@
+// Copyright Epic Games, Inc. All Rights Reserved.
+
+#include "MoverCVDEditor.h"
+
+#include "MoverCVDExtension.h"
+#include "ExtensionsSystem/ChaosVDExtensionsManager.h"
+
+#define LOCTEXT_NAMESPACE "FMoverCVDEditorModule"
+
+void FMoverCVDEditorModule::StartupModule()
+{
+	TSharedRef<FMoverCVDExtension> NewExtension = MakeShared<FMoverCVDExtension>();
+	FChaosVDExtensionsManager::Get().RegisterExtension(NewExtension);
+	AvailableExtensions.Add(NewExtension);
+}
+
+void FMoverCVDEditorModule::ShutdownModule()
+{
+	for (const TWeakPtr<FChaosVDExtension>& Extension : AvailableExtensions)
+	{
+		if (const TSharedPtr<FChaosVDExtension>& ExtensionPtr = Extension.Pin())
+		{
+			FChaosVDExtensionsManager::Get().UnRegisterExtension(ExtensionPtr.ToSharedRef());
+		}
+	}
+
+	AvailableExtensions.Reset();
+}
+
+#undef LOCTEXT_NAMESPACE
+    
+IMPLEMENT_MODULE(FMoverCVDEditorModule, MoverCVDEditor)
\ No newline at end of file
diff --git a/Engine/Plugins/Experimental/Mover/Source/MoverCVDEditor/Private/MoverCVDExtension.cpp b/Engine/Plugins/Experimental/Mover/Source/MoverCVDEditor/Private/MoverCVDExtension.cpp
new file mode 100644
--- /dev/null
+++ b/Engine/Plugins/Experimental/Mover/Source/MoverCVDEditor/Private/MoverCVDExtension.cpp
@@ -0,0 +1,50 @@
+// Copyright Epic Games, Inc. All Rights Reserved.
+
+#include "MoverCVDExtension.h"
+
+#include "MoverCVDSimDataComponent.h"
+#include "MoverCVDSimDataProcessor.h"
+#include "Widgets/SChaosVDMainTab.h"
+#include "MoverCVDTab.h"
+#include "MoverCVDStyle.h"
+
+namespace NMoverCVDExtension
+{
+	static const FName MoverTabName = FName(TEXT("Mover Info"));
+	static const FName ExtensionName = FName(TEXT("FMoverCVDExtension"));
+};
+
+FMoverCVDExtension::FMoverCVDExtension() : FChaosVDExtension()
+{
+	DataComponentsClasses.Add(UMoverCVDSimDataComponent::StaticClass());
+
+	ExtensionName = NMoverCVDExtension::ExtensionName;
+
+	FMoverCVDStyle::Initialize();
+}
+
+FMoverCVDExtension::~FMoverCVDExtension()
+{
+	DataComponentsClasses.Reset();
+
+	FMoverCVDStyle::Shutdown();
+}
+
+void FMoverCVDExtension::RegisterDataProcessorsInstancesForProvider(const TSharedRef<FChaosVDTraceProvider>& InTraceProvider)
+{
+	FChaosVDExtension::RegisterDataProcessorsInstancesForProvider(InTraceProvider);
+
+    TSharedPtr<FMoverCVDSimDataProcessor> SimDataProcessor = MakeShared<FMoverCVDSimDataProcessor>();
+    SimDataProcessor->SetTraceProvider(InTraceProvider);
+    InTraceProvider->RegisterDataProcessor(SimDataProcessor);
+}
+
+TConstArrayView<TSubclassOf<UActorComponent>> FMoverCVDExtension::GetSolverDataComponentsClasses()
+{
+	return DataComponentsClasses;
+}
+
+void FMoverCVDExtension::RegisterCustomTabSpawners(const TSharedRef<SChaosVDMainTab>& InParentTabWidget)
+{
+	InParentTabWidget->RegisterTabSpawner<FMoverCVDTab>(NMoverCVDExtension::MoverTabName);
+}
diff --git a/Engine/Plugins/Experimental/Mover/Source/MoverCVDEditor/Private/MoverCVDExtension.h b/Engine/Plugins/Experimental/Mover/Source/MoverCVDEditor/Private/MoverCVDExtension.h
new file mode 100644
--- /dev/null
+++ b/Engine/Plugins/Experimental/Mover/Source/MoverCVDEditor/Private/MoverCVDExtension.h
@@ -0,0 +1,29 @@
+// Copyright Epic Games, Inc. All Rights Reserved.
+
+#pragma once
+
+#include "ExtensionsSystem/ChaosVDExtension.h"
+#include "Templates/SubclassOf.h"
+#include "Templates/SharedPointer.h"
+
+class FChaosVDTraceProvider;
+class UActorComponent;
+class SChaosVDMainTab;
+
+/** MoverCVDExtension is where we register MoverCVDTab as a displayable tab, register MoverCVDSimDataProcessor and give access to the MoverSimDataComponent */
+class FMoverCVDExtension final : public FChaosVDExtension
+{
+public:
+	
+	FMoverCVDExtension();
+	virtual ~FMoverCVDExtension() override;
+
+	virtual void RegisterDataProcessorsInstancesForProvider(const TSharedRef<FChaosVDTraceProvider>& InTraceProvider) override;
+	virtual TConstArrayView<TSubclassOf<UActorComponent>> GetSolverDataComponentsClasses() override;
+
+	// Registers all available Tab Spawner instances in this extension, if any
+	virtual void RegisterCustomTabSpawners(const TSharedRef<SChaosVDMainTab>& InParentTabWidget) override;
+
+private:
+	TArray<TSubclassOf<UActorComponent>> DataComponentsClasses;
+};
diff --git a/Engine/Plugins/Experimental/Mover/Source/MoverCVDEditor/Private/MoverCVDSimDataComponent.cpp b/Engine/Plugins/Experimental/Mover/Source/MoverCVDEditor/Private/MoverCVDSimDataComponent.cpp
new file mode 100644
--- /dev/null
+++ b/Engine/Plugins/Experimental/Mover/Source/MoverCVDEditor/Private/MoverCVDSimDataComponent.cpp
@@ -0,0 +1,68 @@
+// Copyright Epic Games, Inc. All Rights Reserved.
+
+#include "MoverCVDSimDataComponent.h"
+
+#include "ChaosVDRecording.h"
+#include "ChaosVDScene.h"
+#include "MoverCVDDataWrappers.h"
+#include "Actors/ChaosVDDataContainerBaseActor.h"
+#include "MoverCVDTab.h"
+#include "MoverSimulationTypes.h"
+#include "ChaosVisualDebugger/MoverCVDRuntimeTrace.h"
+
+void UMoverCVDSimDataComponent::UpdateFromSolverFrameData(const FChaosVDSolverFrameData& InSolverFrameData)
+{
+	Super::UpdateFromSolverFrameData(InSolverFrameData);
+
+	if (TSharedPtr<FMoverCVDSimDataContainer> SimDataContainer = InSolverFrameData.GetCustomData().GetData<FMoverCVDSimDataContainer>())
+	{
+		if (const TArray<TSharedPtr<FMoverCVDSimDataWrapper>>* RecordedData = SimDataContainer->SimDataBySolverID.Find(SolverID))
+		{
+			// Load the recorded data into the Physics Mover CVD component
+			FrameSimDataArray.Reset(RecordedData->Num());
+			FrameSimDataArray = *RecordedData;
+
+			// Also, clear all cached deserialized data, we're starting from scratch at a new frame
+			DeserializedSyncStates.Empty();
+			DeserializedInputCmds.Empty();
+		}
+	}
+}
+
+void UMoverCVDSimDataComponent::ClearData()
+{
+	FrameSimDataArray.Reset();
+}
+
+bool UMoverCVDSimDataComponent::FindAndUnwrapSimDataForParticle(uint32 ParticleID, TSharedPtr<FMoverCVDSimDataWrapper>& OutSimDataWrapper, TSharedPtr<FMoverSyncState>& OutSyncState, TSharedPtr<FMoverInputCmdContext>& OutInputCmd)
+{
+	// Look for a sim data corresponding to ParticleID
+	TSharedPtr<FMoverCVDSimDataWrapper>* FoundSimData = FrameSimDataArray.FindByPredicate
+		(
+			[&](const TSharedPtr<FMoverCVDSimDataWrapper>& SimData)
+			{
+				return (SimData->HasValidData() && (SimData->ParticleID == ParticleID));
+			}
+		);
+
+	if (!FoundSimData || !*FoundSimData)
+	{
+		return false;
+	}
+
+	// We use the data wrapper pointer as key in the arrays of deserialized structs (input command, sync state)
+	FMoverCVDSimDataWrapper* SimDataPtr = (*FoundSimData).Get();
+	OutSimDataWrapper = *FoundSimData;
+
+	// Did we previously deserialize?
+	OutSyncState = DeserializedSyncStates.FindOrAdd(SimDataPtr);
+	OutInputCmd = DeserializedInputCmds.FindOrAdd(SimDataPtr);
+	if (!OutSyncState || !OutInputCmd)
+	{
+		// This means that the SimData wasn't deserialized yet, so we do it now
+		// Otherwise we use the cached version
+		FMoverCVDRuntimeTrace::UnwrapSimData(*SimDataPtr, OutInputCmd, OutSyncState);
+	}
+
+	return true;
+}
diff --git a/Engine/Plugins/Experimental/Mover/Source/MoverCVDEditor/Private/MoverCVDSimDataComponent.h b/Engine/Plugins/Experimental/Mover/Source/MoverCVDEditor/Private/MoverCVDSimDataComponent.h
new file mode 100644
--- /dev/null
+++ b/Engine/Plugins/Experimental/Mover/Source/MoverCVDEditor/Private/MoverCVDSimDataComponent.h
@@ -0,0 +1,45 @@
+// Copyright Epic Games, Inc. All Rights Reserved.
+
+#pragma once
+
+#include "Components/ChaosVDSolverDataComponent.h"
+#include "MoverCVDSimDataComponent.generated.h"
+
+struct FMoverSyncState;
+struct FMoverCVDSimDataWrapper;
+struct FMoverInputCmdContext;
+
+/** Component holding Mover data for the current visualized frame */
+UCLASS()
+class UMoverCVDSimDataComponent : public UChaosVDSolverDataComponent
+{
+	GENERATED_BODY()
+
+public:
+	// That we chose to implement this function and not UpdateFromNewGameFrameData or UpdateFromNewSolverStageData is tied
+	// to the implementation of FMoverCVDSimDataProcessor, which currently add ths information to FChaosVDTraceProvider::GetCurrentSolverFrame()
+	// Eventually we will record information at different stages of a solver frame and will be using UpdateFromNewSolverStageData instead,
+	// to show the state of the sync state at the beginning of the frame, then at the end
+	virtual void UpdateFromSolverFrameData(const FChaosVDSolverFrameData& InSolverFrameData) override;
+
+	virtual void ClearData() override;
+	
+	TConstArrayView<TSharedPtr<FMoverCVDSimDataWrapper>> GetFrameSimDataArray() const
+	{
+		return FrameSimDataArray;
+	}
+
+	bool FindAndUnwrapSimDataForParticle(
+		uint32 ParticleID,
+		TSharedPtr<FMoverCVDSimDataWrapper>& OutSimDataWrapper,
+		TSharedPtr<FMoverSyncState>& OutSyncState,
+		TSharedPtr<FMoverInputCmdContext>& OutInputCmd);
+
+private:
+	// This is the array of FMoverCVDSimDataWrapper for the current frame (as it is updated in UpdateFromNewGameFrameData)
+	// and corresponding to this DataComponent's SolverID
+	TArray<TSharedPtr<FMoverCVDSimDataWrapper>> FrameSimDataArray;
+
+	TMap<FMoverCVDSimDataWrapper*, TSharedPtr<FMoverSyncState>> DeserializedSyncStates;
+	TMap<FMoverCVDSimDataWrapper*, TSharedPtr<FMoverInputCmdContext>> DeserializedInputCmds;
+};
diff --git a/Engine/Plugins/Experimental/Mover/Source/MoverCVDEditor/Private/MoverCVDSimDataProcessor.cpp b/Engine/Plugins/Experimental/Mover/Source/MoverCVDEditor/Private/MoverCVDSimDataProcessor.cpp
new file mode 100644
--- /dev/null
+++ b/Engine/Plugins/Experimental/Mover/Source/MoverCVDEditor/Private/MoverCVDSimDataProcessor.cpp
@@ -0,0 +1,36 @@
+// Copyright Epic Games, Inc. All Rights Reserved.
+
+#include "MoverCVDSimDataProcessor.h"
+
+#include "MoverCVDDataWrappers.h"
+
+FMoverCVDSimDataProcessor::FMoverCVDSimDataProcessor() : FChaosVDDataProcessorBase(FMoverCVDSimDataWrapper::WrapperTypeName)
+{
+}
+
+bool FMoverCVDSimDataProcessor::ProcessRawData(const TArray<uint8>& InData)
+{
+	FChaosVDDataProcessorBase::ProcessRawData(InData);
+
+	const TSharedPtr<FChaosVDTraceProvider> ProviderSharedPtr = TraceProvider.Pin();
+	if (!ensure(ProviderSharedPtr.IsValid()))
+	{
+		return false;
+	}
+
+	const TSharedPtr<FMoverCVDSimDataWrapper> SimData = MakeShared<FMoverCVDSimDataWrapper>();
+	const bool bSuccess = Chaos::VisualDebugger::ReadDataFromBuffer(InData, *SimData, ProviderSharedPtr.ToSharedRef());
+	
+	if (bSuccess)
+	{
+		if (FChaosVDSolverFrameData* CurrentSolverFrameData = ProviderSharedPtr->GetCurrentSolverFrame(SimData->SolverID))
+		{
+			if (TSharedPtr<FMoverCVDSimDataContainer> SimDataContainer = CurrentSolverFrameData->GetCustomData().GetOrAddDefaultData<FMoverCVDSimDataContainer>())
+			{
+				SimDataContainer->SimDataBySolverID.FindOrAdd(SimData->SolverID).Add(SimData);
+			}
+		}
+	}
+
+	return bSuccess;
+}
diff --git a/Engine/Plugins/Experimental/Mover/Source/MoverCVDEditor/Private/MoverCVDSimDataProcessor.h b/Engine/Plugins/Experimental/Mover/Source/MoverCVDEditor/Private/MoverCVDSimDataProcessor.h
new file mode 100644
--- /dev/null
+++ b/Engine/Plugins/Experimental/Mover/Source/MoverCVDEditor/Private/MoverCVDSimDataProcessor.h
@@ -0,0 +1,16 @@
+// Copyright Epic Games, Inc. All Rights Reserved.
+
+#pragma once
+
+#include "Trace/DataProcessors/ChaosVDDataProcessorBase.h"
+
+/**
+ * Data processor implementation that is able to deserialize traced Mover data
+ */
+class FMoverCVDSimDataProcessor final : public FChaosVDDataProcessorBase
+{
+public:
+	explicit FMoverCVDSimDataProcessor();
+	
+	virtual bool ProcessRawData(const TArray<uint8>& InData) override;
+};
diff --git a/Engine/Plugins/Experimental/Mover/Source/MoverCVDEditor/Private/MoverCVDSimDataSettings.cpp b/Engine/Plugins/Experimental/Mover/Source/MoverCVDEditor/Private/MoverCVDSimDataSettings.cpp
new file mode 100644
--- /dev/null
+++ b/Engine/Plugins/Experimental/Mover/Source/MoverCVDEditor/Private/MoverCVDSimDataSettings.cpp
@@ -0,0 +1,24 @@
+// Copyright Epic Games, Inc. All Rights Reserved.
+
+#include "MoverCVDSimDataSettings.h"
+
+#include "ChaosVDSettingsManager.h"
+
+void UMoverCVDSimDataSettings::SetDataVisualizationFlags(EMoverCVDSimDataVisualizationFlags NewFlags)
+{
+	if (UMoverCVDSimDataSettings* Settings = FChaosVDSettingsManager::Get().GetSettingsObject<UMoverCVDSimDataSettings>())
+	{
+		Settings->DebugDrawFlags = static_cast<uint32>(NewFlags);
+		Settings->BroadcastSettingsChanged();
+	}
+}
+
+EMoverCVDSimDataVisualizationFlags UMoverCVDSimDataSettings::GetDataVisualizationFlags()
+{
+	if (UMoverCVDSimDataSettings* Settings = FChaosVDSettingsManager::Get().GetSettingsObject<UMoverCVDSimDataSettings>())
+	{
+		return static_cast<EMoverCVDSimDataVisualizationFlags>(Settings->DebugDrawFlags);
+	}
+
+	return EMoverCVDSimDataVisualizationFlags::None;
+}
diff --git a/Engine/Plugins/Experimental/Mover/Source/MoverCVDEditor/Private/MoverCVDSimDataSettings.h b/Engine/Plugins/Experimental/Mover/Source/MoverCVDEditor/Private/MoverCVDSimDataSettings.h
new file mode 100644
--- /dev/null
+++ b/Engine/Plugins/Experimental/Mover/Source/MoverCVDEditor/Private/MoverCVDSimDataSettings.h
@@ -0,0 +1,44 @@
+// Copyright Epic Games, Inc. All Rights Reserved.
+
+#pragma once
+
+#include "Settings/ChaosVDCoreSettings.h"
+#include "MoverCVDSimDataSettings.generated.h"
+
+/** Set of visualization flags options for Mover Sim Data */
+UENUM(meta = (Bitflags, UseEnumValuesAsMaskValuesInEditor = "true"))
+enum class EMoverCVDSimDataVisualizationFlags : uint32
+{
+	None					= 0 UMETA(Hidden),
+	EnableDraw				= 1 << 0,
+};
+ENUM_CLASS_FLAGS(EMoverCVDSimDataVisualizationFlags);
+
+UCLASS(config=ChaosVD, PerObjectConfig)
+class UMoverCVDSimDataSettings : public UChaosVDVisualizationSettingsObjectBase
+{
+	GENERATED_BODY()
+
+public:
+
+	/** If true, any debug draw text available will be drawn */
+	UPROPERTY(EditAnywhere, Category=DebugDraw)
+	bool bShowDebugText = false;
+
+	/** The depth priority used for while drawing data. Can be World or Foreground (with this one the shapes will be drawn on top of the geometry and be always visible) */
+	UPROPERTY(EditAnywhere, Category=DebugDraw)
+	TEnumAsByte<ESceneDepthPriorityGroup> DepthPriority = ESceneDepthPriorityGroup::SDPG_Foreground;
+
+	/** Thickness to apply to any debug draw shape controlled by this setting. */
+	UPROPERTY(EditAnywhere, config, Category=DebugDraw)
+	float BaseThickness = 2.0f;
+
+	static void SetDataVisualizationFlags(EMoverCVDSimDataVisualizationFlags NewFlags);
+	static EMoverCVDSimDataVisualizationFlags GetDataVisualizationFlags();
+
+private:
+
+	/** Set of flags to enable/disable visualization of debug draw data shapes */
+	UPROPERTY(config, meta = (Bitmask, BitmaskEnum = "/Script/MoverCVDEditor.EMoverCVDSimDataVisualizationFlags"))
+	uint32 DebugDrawFlags = static_cast<uint32>(EMoverCVDSimDataVisualizationFlags::None);
+};
\ No newline at end of file
diff --git a/Engine/Plugins/Experimental/Mover/Source/MoverCVDEditor/Private/MoverCVDStyle.cpp b/Engine/Plugins/Experimental/Mover/Source/MoverCVDEditor/Private/MoverCVDStyle.cpp
new file mode 100644
--- /dev/null
+++ b/Engine/Plugins/Experimental/Mover/Source/MoverCVDEditor/Private/MoverCVDStyle.cpp
@@ -0,0 +1,60 @@
+// Copyright Epic Games, Inc. All Rights Reserved.
+
+#include "MoverCVDStyle.h"
+#include "Framework/Application/SlateApplication.h"
+#include "Interfaces/IPluginManager.h"
+#include "Slate/SlateGameResources.h"
+#include "Styling/SlateStyle.h"
+#include "Styling/SlateStyleMacros.h"
+#include "Styling/SlateStyleRegistry.h"
+
+#define RootToContentDir Style->RootToContentDir
+
+TSharedPtr<FSlateStyleSet> FMoverCVDStyle::StyleInstance = nullptr;
+
+void FMoverCVDStyle::Initialize()
+{
+	if (!StyleInstance.IsValid())
+	{
+		StyleInstance = Create();
+		FSlateStyleRegistry::RegisterSlateStyle(*StyleInstance);
+	}
+}
+
+void FMoverCVDStyle::Shutdown()
+{
+	FSlateStyleRegistry::UnRegisterSlateStyle(*StyleInstance);
+	ensure(StyleInstance.IsUnique());
+	StyleInstance.Reset();
+}
+
+FName FMoverCVDStyle::GetStyleSetName()
+{
+	static FName StyleSetName(TEXT("MoverCVDStyle"));
+	return StyleSetName;
+}
+
+const FVector2D Icon16x16(16.0f, 16.0f);
+
+TSharedRef< FSlateStyleSet > FMoverCVDStyle::Create()
+{
+	TSharedRef< FSlateStyleSet > Style = MakeShareable(new FSlateStyleSet("MoverCVDStyle"));
+	Style->SetContentRoot(IPluginManager::Get().FindPlugin("Mover")->GetBaseDir() / TEXT("Resources"));
+
+	Style->Set("TabIconMoverInfoPanel",  new IMAGE_BRUSH_SVG(TEXT("MoverInfo"), Icon16x16));
+
+	return Style;
+}
+
+void FMoverCVDStyle::ReloadTextures()
+{
+	if (FSlateApplication::IsInitialized())
+	{
+		FSlateApplication::Get().GetRenderer()->ReloadTextureResources();
+	}
+}
+
+const ISlateStyle& FMoverCVDStyle::Get()
+{
+	return *StyleInstance;
+}
diff --git a/Engine/Plugins/Experimental/Mover/Source/MoverCVDEditor/Private/MoverCVDStyle.h b/Engine/Plugins/Experimental/Mover/Source/MoverCVDEditor/Private/MoverCVDStyle.h
new file mode 100644
--- /dev/null
+++ b/Engine/Plugins/Experimental/Mover/Source/MoverCVDEditor/Private/MoverCVDStyle.h
@@ -0,0 +1,30 @@
+// Copyright Epic Games, Inc. All Rights Reserved.
+
+#pragma once
+
+#include "Templates/SharedPointer.h"
+
+class FName;
+class FSlateStyleSet;
+class ISlateStyle;
+
+class FMoverCVDStyle
+{
+public:
+
+	static void Initialize();
+
+	static void Shutdown();
+
+	static void ReloadTextures();
+
+	static const ISlateStyle& Get();
+
+	static FName GetStyleSetName();
+
+private:
+
+	static TSharedRef<FSlateStyleSet> Create();
+
+	static TSharedPtr<FSlateStyleSet> StyleInstance;
+};
\ No newline at end of file
diff --git a/Engine/Plugins/Experimental/Mover/Source/MoverCVDEditor/Private/MoverCVDTab.cpp b/Engine/Plugins/Experimental/Mover/Source/MoverCVDEditor/Private/MoverCVDTab.cpp
new file mode 100644
--- /dev/null
+++ b/Engine/Plugins/Experimental/Mover/Source/MoverCVDEditor/Private/MoverCVDTab.cpp
@@ -0,0 +1,220 @@
+// Copyright Epic Games, Inc. All Rights Reserved.
+
+#include "MoverCVDTab.h"
+
+#include "ChaosLog.h"
+#include "ChaosVDModule.h"
+#include "ChaosVDScene.h"
+#include "Framework/Docking/TabManager.h"
+#include "Widgets/Docking/SDockTab.h"
+#include "Widgets/SChaosVDMainTab.h"
+#include "MoverCVDStyle.h"
+#include "MoverSimulationTypes.h"
+#include "MoverCVDDataWrappers.h"
+#include "MoverCVDSimDataComponent.h"
+#include "TEDS/ChaosVDSelectionInterface.h"
+#include "ChaosVDSceneParticle.h"
+
+#define LOCTEXT_NAMESPACE "MoverCVDExtension"
+
+FMoverCVDTab::FMoverCVDTab(const FName& InTabID, TSharedPtr<FTabManager> InTabManager, TWeakPtr<SChaosVDMainTab> InOwningTabWidget)
+	: FChaosVDObjectDetailsTab(InTabID, InTabManager, InOwningTabWidget)
+{
+	// Cache a scene pointer for future reference
+	if (TSharedPtr<SChaosVDMainTab> OwningTabWidgetSharedPtr = OwningTabWidget.Pin())
+	{
+		if (TSharedPtr<FChaosVDScene> Scene = OwningTabWidgetSharedPtr->GetScene())
+		{
+			SceneWeakPtr = Scene;
+		}
+	}
+}
+
+FMoverCVDTab::~FMoverCVDTab()
+{
+	if (TSharedPtr<FChaosVDScene> Scene = SceneWeakPtr.Pin())
+	{
+		Scene->OnSceneUpdated().RemoveAll(this);
+	}
+}
+
+TSharedRef<SDockTab> FMoverCVDTab::HandleTabSpawnRequest(const FSpawnTabArgs& Args)
+{
+	// I can't add callbacks in the constructor as SharedFromThis would complain
+	if (TSharedPtr<FChaosVDScene> Scene = SceneWeakPtr.Pin())
+	{
+		Scene->OnSceneUpdated().AddSP(this, &FMoverCVDTab::HandleSceneUpdated);
+	}
+
+	TSharedRef<SDockTab> Tab = FChaosVDObjectDetailsTab::HandleTabSpawnRequest(Args);
+
+	// Let's rename the tab so its name is unique
+	Tab->SetLabel(Args.GetTabId().ToText());
+
+	Tab->SetTabIcon(FMoverCVDStyle::Get().GetBrush("TabIconMoverInfoPanel"));
+
+	return Tab;
+}
+
+void FMoverCVDTab::RetrieveAllSolversMoverDataComponents()
+{
+	SolverToSimDataComponentMap.Empty();
+
+	const TSharedPtr<FChaosVDScene> Scene = SceneWeakPtr.Pin();
+	if (Scene)
+	{
+		const FChaosVDSolverInfoByIDMap& SolverInfoByIDMap = Scene->GetSolverInfoActorsMap();
+		for (const auto& It : SolverInfoByIDMap)
+		{
+			int32 SolverID = It.Key;
+			// No need to know about AChaosVDSolverInfoActor as we intend to solely find an actor component,
+			// but also at the time of writing, AChaosVDSolverInfoActor was no a public include
+			const AActor* SolverInfoActor = reinterpret_cast<const AActor*>(It.Value);
+			if (SolverInfoActor)
+			{
+				if (UMoverCVDSimDataComponent* SolverSimDataComponent = SolverInfoActor->FindComponentByClass<UMoverCVDSimDataComponent>())
+				{
+					SolverToSimDataComponentMap.Add(SolverID, SolverSimDataComponent);
+				}
+			}
+		}
+	}
+}
+
+TWeakObjectPtr<UMoverCVDSimDataComponent>* FMoverCVDTab::FindMoverDataComponentForSolver(int32 SolverID)
+{
+	TWeakObjectPtr<UMoverCVDSimDataComponent>* FoundSimDataComponent = SolverToSimDataComponentMap.Find(SolverID);
+	if (!FoundSimDataComponent || FoundSimDataComponent->IsStale())
+	{
+		// Refresh SolverToSimDataComponentMap
+		RetrieveAllSolversMoverDataComponents();
+
+		// Try again
+		FoundSimDataComponent = SolverToSimDataComponentMap.Find(SolverID);
+	}
+
+	return FoundSimDataComponent;
+}
+
+void FMoverCVDTab::DisplaySingleParticleInfo(int32 SelectedSolverID, int32 SelectedParticleID)
+{
+	// At the moment this is really written to support only one particle selected at a time
+	// I we were to have multiple selection we would have to cache the whole list of particle IDs selected
+	bool IsParticleBeingSelected = (SelectedParticleID != INDEX_NONE && SelectedSolverID != INDEX_NONE);
+	bool WasDataFoundForParticle = false;
+	if (IsParticleBeingSelected)
+	{
+		// Find the sim data component weak pointer for the selected solver ID
+		TWeakObjectPtr<UMoverCVDSimDataComponent>* FoundSimDataComponent = FindMoverDataComponentForSolver(SelectedSolverID);
+		if (FoundSimDataComponent)
+		{
+			// Pin a strong pointer to it so it doesn't get away while in use
+			TStrongObjectPtr<UMoverCVDSimDataComponent> SimDataComponentSharedPtr = FoundSimDataComponent->Pin();
+			UMoverCVDSimDataComponent* SimDataComponent = SimDataComponentSharedPtr.Get();
+			if (IsValid(SimDataComponent))
+			{
+				if (SimDataComponent->FindAndUnwrapSimDataForParticle(SelectedParticleID, MoverSimDataWrapper, MoverSyncState, MoverInputCmd))
+				{
+					if (MoverSimDataWrapper && MoverInputCmd && MoverSyncState)
+					{
+						if (DetailsPanelView)
+						{
+							MultiViewWrapper.Clear();
+
+							FMoverCVDSimDataParticleInfo MoverParticleInfo;
+							MoverParticleInfo.ParticleID = MoverSimDataWrapper->ParticleID;
+							MoverParticleInfo.SolverID = MoverSimDataWrapper->SolverID;
+							MultiViewWrapper.AddData(MoverSimDataWrapper.Get());
+
+							for (const TSharedPtr<FMoverDataStructBase>& InputDataStructSharedPtr : MoverInputCmd->InputCollection.GetDataArray())
+							{
+								if (FMoverDataStructBase* MoverDataStructBase = MoverInputCmd->InputCollection.GetDataArray()[0].Get())
+								{
+									MultiViewWrapper.AddData(MakeShared<FStructOnScope>(MoverDataStructBase->GetScriptStruct(), reinterpret_cast<uint8*>(MoverDataStructBase)));
+								}
+							}
+							MultiViewWrapper.AddData(MoverSyncState.Get());
+
+							WasDataFoundForParticle = true;
+
+							SetStructToInspect(&MultiViewWrapper);
+						}
+					}
+					// Cache the particle ID so we can display info for that same particle when we scrub to a different frame that also has that particle
+					// WARNING: Apparently IDs can be reused across frames for different particles, so this might not be good enough
+					CurrentlyDisplayedSolverID = SelectedSolverID;
+					CurrentlyDisplayedParticleID = SelectedParticleID;
+				}
+			}
+		}
+	}
+	
+	if (!WasDataFoundForParticle)
+	{
+		CurrentlyDisplayedSolverID = INDEX_NONE;
+		CurrentlyDisplayedParticleID = INDEX_NONE;
+	}
+}
+
+void FMoverCVDTab::DisplayMoverInfoForSelectedElements(const TArray<FTypedElementHandle>& SelectedElementHandles)
+{
+	if (DetailsPanelView)
+	{
+		DetailsPanelView->SetSelectedStruct(nullptr);
+
+		bool ShouldClearSimDataDetailsPanel = true;
+		using namespace Chaos::VD::TypedElementDataUtil;
+		for (int32 SelectionIndex = 0; SelectionIndex < SelectedElementHandles.Num(); ++SelectionIndex)
+		{
+			if (const FChaosVDSceneParticle* Particle = GetStructDataFromTypedElementHandle<FChaosVDSceneParticle>(SelectedElementHandles[SelectionIndex]))
+			{
+				int32 SelectedParticleID = Particle->GetParticleData()->ParticleIndex; // Replace with ID of selected particle if one is selected
+				int32 SelectedSolverID = Particle->GetParticleData()->SolverID;
+				ShouldClearSimDataDetailsPanel = false;
+				DisplaySingleParticleInfo(SelectedSolverID, SelectedParticleID);
+
+				// Right now, we only handle displaying info for the first particle in the selection, but later we might not break here and display info for all particles
+				break;
+			}
+		}
+
+		// Clear the details panel view if we didn't find particles in the selection
+		if (ShouldClearSimDataDetailsPanel)
+		{
+			DetailsPanelView->SetSelectedStruct(nullptr);
+		}
+	}
+}
+
+void FMoverCVDTab::HandleSceneUpdated()
+{
+	// This is called when we scrub the timelines for instance
+	bool IsParticleSelected = (CurrentlyDisplayedSolverID != INDEX_NONE && CurrentlyDisplayedParticleID != INDEX_NONE);
+
+	if (!IsParticleSelected)
+	{
+		return;
+	}
+
+	const TSharedPtr<FChaosVDScene> Scene = SceneWeakPtr.Pin();
+	if (!Scene)
+	{
+		return;
+	}
+
+	bool ShouldClearSimDataDetailsPanel = true;
+	using namespace Chaos::VD::TypedElementDataUtil;
+	TArray<FTypedElementHandle> SelectedParticlesHandles = Scene->GetSelectedElementHandles();
+	DisplayMoverInfoForSelectedElements(SelectedParticlesHandles);
+}
+
+void FMoverCVDTab::HandlePostSelectionChange(const UTypedElementSelectionSet* ChangedSelectionSet)
+{
+	// This is called when the selection changes
+	TArray<FTypedElementHandle> SelectedParticlesHandles;
+	ChangedSelectionSet->GetSelectedElementHandles(SelectedParticlesHandles, UChaosVDSelectionInterface::StaticClass());	
+	DisplayMoverInfoForSelectedElements(SelectedParticlesHandles);
+}
+
+#undef LOCTEXT_NAMESPACE
+
diff --git a/Engine/Plugins/Experimental/Mover/Source/MoverCVDEditor/Private/MoverCVDTab.h b/Engine/Plugins/Experimental/Mover/Source/MoverCVDEditor/Private/MoverCVDTab.h
new file mode 100644
--- /dev/null
+++ b/Engine/Plugins/Experimental/Mover/Source/MoverCVDEditor/Private/MoverCVDTab.h
@@ -0,0 +1,58 @@
+// Copyright Epic Games, Inc. All Rights Reserved.
+
+#pragma once
+
+#include "Templates/SharedPointer.h"
+#include "ChaosVDTabSpawnerBase.h"
+#include "Widgets/SChaosVDDetailsView.h"
+#include "ChaosVDSolverDataSelection.h"
+#include "ChaosVDObjectDetailsTab.h"
+#include "TEDS/ChaosVDStructTypedElementData.h"
+
+class SOutputLog;
+class FOutputLogHistory;
+class FName;
+struct FMoverSyncState;
+struct FMoverInputCmdContext;
+struct FMoverCVDSimDataWrapper;
+class UMoverCVDSimDataComponent;
+
+/** This tab is an additional details tab displaying mover info corresponding to the selected particles if they are moved by a Mover component */
+class FMoverCVDTab : public FChaosVDObjectDetailsTab
+{
+public:
+
+	FMoverCVDTab(const FName& InTabID, TSharedPtr<FTabManager> InTabManager, TWeakPtr<SChaosVDMainTab> InOwningTabWidget);
+	virtual ~FMoverCVDTab();
+
+	// Implementation of FChaosVDObjectDetailsTab
+	virtual TSharedRef<SDockTab> HandleTabSpawnRequest(const FSpawnTabArgs& Args) override;
+	virtual void HandlePostSelectionChange(const UTypedElementSelectionSet* ChangedSelectionSet) override;
+
+	// Scene callbacks
+	void HandleSceneUpdated();
+
+private:
+	void DisplaySingleParticleInfo(int32 SelectedSolverID, int32 SelectedParticleID);
+	void DisplayMoverInfoForSelectedElements(const TArray<FTypedElementHandle>& SelectedElementHandles);
+
+	// This function retrieves and caches all the mover data components for all solvers, populating SolverToSimDataComponentMap
+	void RetrieveAllSolversMoverDataComponents();
+
+	TWeakObjectPtr<UMoverCVDSimDataComponent>* FindMoverDataComponentForSolver(int32 SolverID);
+
+	TWeakPtr<FChaosVDScene> SceneWeakPtr;
+
+	TMap<int32, TWeakObjectPtr<UMoverCVDSimDataComponent>> SolverToSimDataComponentMap;
+
+	int32 CurrentlyDisplayedParticleID = INDEX_NONE;
+	int32 CurrentlyDisplayedSolverID = INDEX_NONE;
+
+	FChaosVDSelectionMultipleView MultiViewWrapper;
+
+	TSharedPtr<FMoverCVDSimDataWrapper> MoverSimDataWrapper;
+	TSharedPtr<FMoverSyncState> MoverSyncState;
+	TSharedPtr<FMoverInputCmdContext> MoverInputCmd;
+};
+
+
diff --git a/Engine/Plugins/Experimental/Mover/Source/MoverCVDEditor/Public/MoverCVDEditor.h b/Engine/Plugins/Experimental/Mover/Source/MoverCVDEditor/Public/MoverCVDEditor.h
new file mode 100644
--- /dev/null
+++ b/Engine/Plugins/Experimental/Mover/Source/MoverCVDEditor/Public/MoverCVDEditor.h
@@ -0,0 +1,27 @@
+// Copyright Epic Games, Inc. All Rights Reserved.
+
+#pragma once
+
+#include "CoreMinimal.h"
+#include "Modules/ModuleManager.h"
+
+class FChaosVDExtension;
+
+/** This module contains the Editor side of the Mover support in Chaos Visual Debugger (CVD)
+* Mover is a module responsible for the movement of actors and its replication, for instance the input controlled movement of characters
+* This module allows the display and recording specific data in CVD that makes it easier to understand and debug Mover
+* MoverCVDTab is the tab in CVD where that information is displayed
+* MoverCVDSimDataProcessor is the receiving and processing end of the Mover data trace ("Sim Data") that the game sends to CVD
+* MoverSimDataComponent is a component holding Mover data for the current visualized frame
+* MoverCVDExtension is where we register MoverCVDTab as a displayable tab, register MoverCVDSimDataProcessor and give access to the MoverSimDataComponent
+*/
+class FMoverCVDEditorModule : public IModuleInterface
+{
+public:
+    virtual void StartupModule() override;
+    virtual void ShutdownModule() override;
+
+private:
+
+	TArray<TWeakPtr<FChaosVDExtension>> AvailableExtensions;
+};
