Crash fix in CVD after opening a file fails, the code still tried to use null loaded recording

diff --git a/Engine/Plugins/ChaosVD/Source/ChaosVD/Private/ChaosVDPlaybackController.cpp b/Engine/Plugins/ChaosVD/Source/ChaosVD/Private/ChaosVDPlaybackController.cpp
--- a/Engine/Plugins/ChaosVD/Source/ChaosVD/Private/ChaosVDPlaybackController.cpp
+++ b/Engine/Plugins/ChaosVD/Source/ChaosVD/Private/ChaosVDPlaybackController.cpp
@@ -355,8 +355,7 @@ void FChaosVDPlaybackController::GoToTrackFrameAndSync(FGuid InstigatorID, EChao
 		return;
 	}
 
-	FReadScopeLock ReadLock(LoadedRecording->GetRecordingDataLock());
-	
+	FReadScopeLock ReadLock(LoadedRecording->GetRecordingDataLock());	
 	GoToTrackFrame_AssumesLockedAndSync(InstigatorID, TrackType, InTrackID, FrameNumber, StageNumber);
 }
 
@@ -580,6 +579,11 @@ int32 FChaosVDPlaybackController::GetTrackCurrentStage(EChaosVDTrackType TrackTy
 
 int32 FChaosVDPlaybackController::GetTrackLastStageAtFrame(EChaosVDTrackType TrackType, int32 InTrackID, int32 InFrameNumber) const
 {
+	if (!ensure(LoadedRecording.IsValid()))
+	{
+		return INDEX_NONE;
+	}
+
 	FReadScopeLock ReadLock(LoadedRecording->GetRecordingDataLock());
 	return GetTrackLastStageAtFrame_AssumesLocked(TrackType, InTrackID, InFrameNumber);
 }
@@ -656,8 +660,13 @@ void FChaosVDPlaybackController::GetMutableAvailableTracks(EChaosVDTrackType Tra
 
 void FChaosVDPlaybackController::GetAvailableTrackInfosAtTrackFrame(EChaosVDTrackType TrackTypeToFind, const TSharedRef<const FChaosVDTrackInfo>& InFromTrack, TArray<TSharedPtr<const FChaosVDTrackInfo>>& OutTrackInfo)
 {
+	if (!ensure(LoadedRecording.IsValid()))
+	{
+		return;
+	}
+
 	FReadScopeLock ReadLock(LoadedRecording->GetRecordingDataLock());
-	return GetAvailableTrackInfosAtTrackFrame_AssumesLocked(TrackTypeToFind, InFromTrack, OutTrackInfo);
+	GetAvailableTrackInfosAtTrackFrame_AssumesLocked(TrackTypeToFind, InFromTrack, OutTrackInfo);
 }
 
 void FChaosVDPlaybackController::GetAvailableTrackInfosAtTrackFrame_AssumesLocked(EChaosVDTrackType TrackTypeToFind, const TSharedRef<const FChaosVDTrackInfo>& InFromTrack, TArray<TSharedPtr<const FChaosVDTrackInfo>>& OutTrackInfo)
@@ -1079,8 +1088,7 @@ void FChaosVDPlaybackController::SyncTracks(const TSharedRef<const FChaosVDTrack
 		return;
 	}
 
-	FReadScopeLock ReadLock(LoadedRecording->GetRecordingDataLock());
-	
+	FReadScopeLock ReadLock(LoadedRecording->GetRecordingDataLock());	
 	SyncTracks_AssumesLocked(FromTrack, TrackSyncMode);
 }
 
