[ChaosVD] Implemented persistent layout support

- Now each time CVD is closed, the layout will be serialized and re-used the next time it loads
- Expanded the settings menu to also contain layout options
- A new Reset Layout entry was added that allows you to reset CVD's layout to the default one, without having to fully restart CVD

diff --git a/Engine/Plugins/ChaosVD/Source/ChaosVD/Private/ChaosVDModule.cpp b/Engine/Plugins/ChaosVD/Source/ChaosVD/Private/ChaosVDModule.cpp
--- a/Engine/Plugins/ChaosVD/Source/ChaosVD/Private/ChaosVDModule.cpp
+++ b/Engine/Plugins/ChaosVD/Source/ChaosVD/Private/ChaosVDModule.cpp
@@ -226,6 +226,30 @@ bool FChaosVDModule::IsStandaloneChaosVisualDebugger()
 	return FPlatformProperties::IsProgram() && FApp::GetProjectName() == ChaosVisualDebuggerProgramName;
 }
 
+void FChaosVDModule::ReloadInstanceUI(FGuid InstanceGUID)
+{
+	// CVD UI is (or it should be) fully de-coupled from the tool's non-ui state
+	// Therefore to reload its UI we can just re-spawn the tab widget and initialize it with the existing CVD engine instance
+
+	TWeakPtr<SDockTab>* ActiveTabPtrPtr = ActiveCVDTabs.Find(InstanceGUID);
+	if (TSharedPtr<SDockTab> TabPtr = ActiveTabPtrPtr ? ActiveTabPtrPtr->Pin() : nullptr)
+	{
+		TSharedPtr<FChaosVDEngine>* ChaosVDEngineInstancePtrPtr = ActiveChaosVDInstances.Find(InstanceGUID);
+		if (TSharedPtr<FChaosVDEngine> ChaosVDEngineInstancePtr = ChaosVDEngineInstancePtrPtr ? *ChaosVDEngineInstancePtrPtr : nullptr)
+		{
+			TabPtr->SetContent(SNew(SChaosVDMainTab, ChaosVDEngineInstancePtr).OwnerTab(TabPtr));
+		}
+		else
+		{
+			UE_LOG(LogChaosVDEditor, Error, TEXT("[%hs]: Failed to find CVD engine instance with ID [%s]. The UI was not reloaded."), __func__, *InstanceGUID.ToString());
+		}
+	}
+	else
+	{
+		UE_LOG(LogChaosVDEditor, Error, TEXT("[%hs]: Failed to find CVD tab instance with ID [%s]. The UI was not reloaded."), __func__, *InstanceGUID.ToString());
+	}
+}
+
 #undef LOCTEXT_NAMESPACE
 	
 IMPLEMENT_MODULE(FChaosVDModule, ChaosVD)
diff --git a/Engine/Plugins/ChaosVD/Source/ChaosVD/Private/Widgets/SChaosVDMainTab.cpp b/Engine/Plugins/ChaosVD/Source/ChaosVD/Private/Widgets/SChaosVDMainTab.cpp
--- a/Engine/Plugins/ChaosVD/Source/ChaosVD/Private/Widgets/SChaosVDMainTab.cpp
+++ b/Engine/Plugins/ChaosVD/Source/ChaosVD/Private/Widgets/SChaosVDMainTab.cpp
@@ -48,6 +48,7 @@
 #include "ExtensionsSystem/ChaosVDExtensionsManager.h"
 
 #include "Features/IModularFeatures.h"
+#include "Framework/Docking/LayoutService.h"
 #include "HAL/FileManager.h"
 #include "Settings/ChaosVDGeneralSettings.h"
 #include "Settings/ChaosVDMiscSettings.h"
@@ -93,6 +94,8 @@ void SChaosVDMainTab::Construct(const FArguments& InArgs, TSharedPtr<FChaosVDEng
 
 	TabManager = FGlobalTabmanager::Get()->NewTabManager(InArgs._OwnerTab.ToSharedRef()).ToSharedPtr();
 
+	TabManager->SetOnPersistLayout(FTabManager::FOnPersistLayout::CreateSP(this, &SChaosVDMainTab::HandlePersistLayout));
+
 	RegisterTabSpawner<FChaosVDWorldOutlinerTab>(FChaosVDTabID::WorldOutliner);
 	RegisterTabSpawner<FChaosVDObjectDetailsTab>(FChaosVDTabID::DetailsPanel);
 	RegisterTabSpawner<FChaosVDOutputLogTab>(FChaosVDTabID::OutputLog);
@@ -187,7 +190,7 @@ void SChaosVDMainTab::Construct(const FArguments& InArgs, TSharedPtr<FChaosVDEng
 		+SVerticalBox::Slot()
 		.Padding(FMargin(0.0f,5.0f,0.0f,0.0f))
 		[
-			TabManager->RestoreFrom(GenerateMainLayout(), TabManager->GetOwnerTab()->GetParentWindow()).ToSharedRef()
+			TabManager->RestoreFrom(FLayoutSaveRestore::LoadFromConfig(GEditorLayoutIni, GenerateDefaultLayout()), TabManager->GetOwnerTab()->GetParentWindow()).ToSharedRef()
 		]
 		+SVerticalBox::Slot()
 		.Padding(0.0f, 2.0f, 0.0f, 0.0f)
@@ -280,6 +283,19 @@ const TSharedPtr<FChaosVDIndependentDetailsPanelManager>& SChaosVDMainTab::GetIn
 	return IndependentDetailsPanelManager;
 }
 
+void SChaosVDMainTab::HandlePersistLayout(const TSharedRef<FTabManager::FLayout>& InLayoutToSave)
+{
+	if (!bCanTabManagerPersistLayout)
+	{
+		return;
+	}
+
+	if (TSharedPtr<FTabManager::FArea> PrimaryArea = InLayoutToSave->GetPrimaryArea().Pin())
+	{
+		FLayoutSaveRestore::SaveToConfig(GEditorLayoutIni, InLayoutToSave);
+	}
+}
+
 void SChaosVDMainTab::HandlePostInitializationExtensionRegistered(const TSharedRef<FChaosVDExtension>& NewExtension)
 {	
 	NewExtension->RegisterCustomTabSpawners(StaticCastSharedRef<SChaosVDMainTab>(AsShared()));
@@ -307,11 +323,11 @@ void SChaosVDMainTab::HandleTabDestroyed(TSharedRef<SDockTab> Tab, FName TabID)
 	ActiveTabsByID.Remove(TabID);
 }
 
-TSharedRef<FTabManager::FLayout> SChaosVDMainTab::GenerateMainLayout()
+TSharedRef<FTabManager::FLayout> SChaosVDMainTab::GenerateDefaultLayout()
 {
 	const float DPIScaleFactor = FPlatformApplicationMisc::GetDPIScaleFactorAtPoint(10.0f, 10.0f);
 
-	return FTabManager::NewLayout("ChaosVisualDebugger_Layout")
+	return FTabManager::NewLayout("ChaosVisualDebugger_Layout_V1")
 		->AddArea
 		(
 			FTabManager::NewPrimaryArea()
@@ -383,6 +399,17 @@ TSharedRef<FTabManager::FLayout> SChaosVDMainTab::GenerateMainLayout()
 	));
 }
 
+void SChaosVDMainTab::ResetLayout()
+{
+	// During a layout reset, we manually stomp the currently saved layout, therefore we don't want the layout
+	// to be resaved for the remaining of this tab instance lifespan (which should be not longer than this scope).
+	bCanTabManagerPersistLayout = false;
+	FLayoutSaveRestore::SaveToConfig(GEditorLayoutIni, GenerateDefaultLayout());
+
+	TabManager->CloseAllAreas();
+	FChaosVDModule::Get().ReloadInstanceUI(GetChaosVDEngineInstance()->GetInstanceGuid());
+}
+
 void SChaosVDMainTab::GenerateMainWindowMenu()
 {
 	FMenuBarBuilder MenuBarBuilder = FMenuBarBuilder(TSharedPtr<FUICommandList >());
@@ -923,7 +950,36 @@ void SChaosVDMainTab::RegisterSettingsMenu()
 	using namespace Chaos::VisualDebugger::Utils;
 
 	FToolMenuSection& Section = ToolBar->AddSection("SettingsMenu");
-	FNewToolMenuDelegate MainSettingsMenuBuilder = FNewToolMenuDelegate::CreateStatic(&CreateMenuEntryForSettingsObject<UChaosVDGeneralSettings>, EChaosVDSaveSettingsOptions::ShowResetButton);
+
+	FNewToolMenuDelegate MainSettingsMenuBuilder = FNewToolMenuDelegate::CreateStatic([](UToolMenu* Menu)
+	{
+		FToolMenuSection& GeneralSection = Menu->AddSection("GeneralSettingsMenu", LOCTEXT("CommonSettingsMenuLabel", "General"));
+		constexpr bool bOpenSubMenuOnClick = false;
+		FNewToolMenuDelegate MainCommonSettingsMenuBuilder = FNewToolMenuDelegate::CreateStatic(&CreateMenuEntryForSettingsObject<UChaosVDGeneralSettings>, EChaosVDSaveSettingsOptions::ShowResetButton);
+		GeneralSection.AddSubMenu(TEXT("MainCommonSettingsMenu"), LOCTEXT("MainCommonSettingsMenuLabel", "Common"), LOCTEXT("MainCommonSettingsMenuTip", "Common Settings that controls general behavior of CVD"), MainCommonSettingsMenuBuilder, bOpenSubMenuOnClick,  FSlateIcon(FAppStyle::Get().GetStyleSetName(), TEXT("Icons.Toolbar.Settings")));
+
+		FNewToolMenuDelegate LayoutSettingsMenuBuilder = FNewToolMenuDelegate::CreateStatic([](UToolMenu* Menu)
+		{
+			const UChaosVDMainToolbarMenuContext* Context = Menu ? Menu->FindContext<UChaosVDMainToolbarMenuContext>() : nullptr;
+			TSharedPtr<SChaosVDMainTab> MainTabPtr = Context ? Context->MainTab.Pin() : nullptr;
+			if (!MainTabPtr)
+			{
+				return;
+			}
+
+			FToolMenuEntry ResetMenuEntry = FToolMenuEntry::InitMenuEntry(FName("ResetLayoutMenu"),
+				LOCTEXT("ResetLayoutMenuEntryLabel", "Reset Layout"),
+				LOCTEXT("ResetLayoutMenuEntryLabelToolTip", "Reset the current layout to the defaults one"),
+				FSlateIcon(FAppStyle::Get().GetStyleSetName(), TEXT("Mainframe.LoadLayout")),
+				FUIAction(FExecuteAction::CreateSP(MainTabPtr.ToSharedRef(), &SChaosVDMainTab::ResetLayout)));
+
+			Menu->AddMenuEntry(NAME_None, ResetMenuEntry);
+		});
+		
+		FToolMenuSection& AppearanceSection = Menu->AddSection("AppearanceSectionMenu", LOCTEXT("AppearanceSectionMenuLabel", "Appearance"));
+		AppearanceSection.AddSubMenu(TEXT("MainLayoutSettingsMenu"), LOCTEXT("MainLayoutSettingsMenuLabel", "Layout"), LOCTEXT("MainLayoutSettingsMenuTip", "Set of options to alter CVD's UI layout"), LayoutSettingsMenuBuilder, bOpenSubMenuOnClick,  FSlateIcon(FAppStyle::Get().GetStyleSetName(), TEXT("Icons.Layout")));
+	});
+	
 	constexpr bool bOpenSubMenuOnClick = true;
 	Section.AddSubMenu(TEXT("MainSettingsMenu"), LOCTEXT("MainSettingsMenuLabel", "Settings"), LOCTEXT("MainSettingsMenuTip", "Settings that controls general behavior of CVD"), MainSettingsMenuBuilder, bOpenSubMenuOnClick,  FSlateIcon(FAppStyle::Get().GetStyleSetName(), TEXT("Icons.Toolbar.Settings")));
 }
diff --git a/Engine/Plugins/ChaosVD/Source/ChaosVD/Public/ChaosVDModule.h b/Engine/Plugins/ChaosVD/Source/ChaosVD/Public/ChaosVDModule.h
--- a/Engine/Plugins/ChaosVD/Source/ChaosVD/Public/ChaosVDModule.h
+++ b/Engine/Plugins/ChaosVD/Source/ChaosVD/Public/ChaosVDModule.h
@@ -38,6 +38,9 @@ public:
 
 	static bool IsStandaloneChaosVisualDebugger();
 
+	/** Re-instantiates CVD's widgets without restarting the underlying CVD engine instance */
+	void ReloadInstanceUI(FGuid InstanceGUID);
+
 private:
 
 	TSharedRef<SDockTab> SpawnMainTab(const FSpawnTabArgs& Args);
diff --git a/Engine/Plugins/ChaosVD/Source/ChaosVD/Public/Widgets/SChaosVDMainTab.h b/Engine/Plugins/ChaosVD/Source/ChaosVD/Public/Widgets/SChaosVDMainTab.h
--- a/Engine/Plugins/ChaosVD/Source/ChaosVD/Public/Widgets/SChaosVDMainTab.h
+++ b/Engine/Plugins/ChaosVD/Source/ChaosVD/Public/Widgets/SChaosVDMainTab.h
@@ -137,6 +137,8 @@ public:
 
 private:
 
+	void HandlePersistLayout(const TSharedRef<FTabManager::FLayout>& InLayoutToSave);
+
 	void HandlePostInitializationExtensionRegistered(const TSharedRef<FChaosVDExtension>& NewExtension);
 
 	void SetCustomPropertyLayouts(IDetailsView* DetailsView);
@@ -150,7 +152,9 @@ private:
 	CHAOSVD_API void HandleTabSpawned(TSharedRef<SDockTab> Tab, FName TabID);
 	CHAOSVD_API void HandleTabDestroyed(TSharedRef<SDockTab> Tab, FName TabID);
 
-	TSharedRef<FTabManager::FLayout> GenerateMainLayout();
+	TSharedRef<FTabManager::FLayout> GenerateDefaultLayout();
+
+	void ResetLayout();
 
 	void GenerateMainWindowMenu();
 	void GenerateRecentFilesMenu(FMenuBuilder& MenuBuilder);
@@ -167,6 +171,8 @@ private:
 	void BrowseLiveSessionsFromTraceStore() const;
 
 	bool ShouldDisableCPUThrottling() const;
+
+	bool bCanTabManagerPersistLayout = true;
 	
 	TArray<TWeakPtr<IDetailsView>> CustomizedDetailsPanels;
 
