[ChaosVD] Added support to persist data channels enabled state

Now if you enable or disable a data channel and then restart the editor, the channels you enabled or disabled will stay in the same state

diff --git a/Engine/Plugins/ChaosVD/Source/ChaosVD/Private/ChaosVDEngine.cpp b/Engine/Plugins/ChaosVD/Source/ChaosVD/Private/ChaosVDEngine.cpp
--- a/Engine/Plugins/ChaosVD/Source/ChaosVD/Private/ChaosVDEngine.cpp
+++ b/Engine/Plugins/ChaosVD/Source/ChaosVD/Private/ChaosVDEngine.cpp
@@ -8,6 +8,7 @@
 #include "ChaosVDScene.h"
 #include "ChaosVDSettingsManager.h"
 #include "Chaos/ChaosVDEngineEditorBridge.h"
+#include "ChaosVisualDebugger/ChaosVDOptionalDataChannel.h"
 #include "Settings/ChaosVDMiscSettings.h"
 #include "Trace/ChaosVDTraceManager.h"
 
@@ -42,6 +43,8 @@ void FChaosVDEngine::Initialize()
 			}
 		});
 	}
+
+	RestoreDataChannelsEnabledStateFromSave();
 	
 	bIsInitialized = true;
 }
@@ -95,9 +98,63 @@ void FChaosVDEngine::DeInitialize()
 
 	CollectGarbage(GARBAGE_COLLECTION_KEEPFLAGS);
 
+#if WITH_CHAOS_VISUAL_DEBUGGER
+
+	if (DataChannelStateUpdatedHandle.IsValid())
+	{
+		Chaos::VisualDebugger::FChaosVDDataChannelsManager::Get().OnChannelStateChanged().Remove(DataChannelStateUpdatedHandle);
+	}
+
+#endif
+
 	bIsInitialized = false;
 }
 
+void FChaosVDEngine::RestoreDataChannelsEnabledStateFromSave()
+{
+#if WITH_CHAOS_VISUAL_DEBUGGER
+
+	UChaosVDMiscSettings* MiscSettings = FChaosVDSettingsManager::Get().GetSettingsObject<UChaosVDMiscSettings>();
+	if (!MiscSettings)
+	{
+		return;
+	}
+
+	using namespace Chaos::VisualDebugger;
+	FChaosVDDataChannelsManager::Get().EnumerateChannels([MiscSettings](const TSharedRef<FChaosVDOptionalDataChannel>& Channel)
+	{
+		if (bool* SavedEnabledState = MiscSettings->DataChannelEnabledState.Find(Channel->GetId().ToString()))
+		{
+			Channel->SetChannelEnabled(*SavedEnabledState);
+		}
+		return true;
+	});
+
+	DataChannelStateUpdatedHandle = FChaosVDDataChannelsManager::Get().OnChannelStateChanged().AddSP(this, &FChaosVDEngine::UpdateSavedDataChannelsEnabledState);
+
+#endif
+}
+
+void FChaosVDEngine::UpdateSavedDataChannelsEnabledState(TWeakPtr<Chaos::VisualDebugger::FChaosVDOptionalDataChannel> DataChannelChanged)
+{
+#if WITH_CHAOS_VISUAL_DEBUGGER
+
+	UChaosVDMiscSettings* MiscSettings = FChaosVDSettingsManager::Get().GetSettingsObject<UChaosVDMiscSettings>();
+	if (!MiscSettings)
+	{
+		return;
+	}
+
+	if (TSharedPtr<Chaos::VisualDebugger::FChaosVDOptionalDataChannel> DataChannelPtr = DataChannelChanged.Pin())
+	{
+		MiscSettings->DataChannelEnabledState.FindOrAdd(DataChannelPtr->GetId().ToString()) = DataChannelPtr->IsChannelEnabled();
+	}
+
+	MiscSettings->SaveConfig();
+
+#endif
+}
+
 void FChaosVDEngine::UpdateRecentFilesList(const FString& InFilename)
 {
 	UChaosVDMiscSettings* MiscSettings = FChaosVDSettingsManager::Get().GetSettingsObject<UChaosVDMiscSettings>();
diff --git a/Engine/Plugins/ChaosVD/Source/ChaosVD/Private/ChaosVDEngine.h b/Engine/Plugins/ChaosVD/Source/ChaosVD/Private/ChaosVDEngine.h
--- a/Engine/Plugins/ChaosVD/Source/ChaosVD/Private/ChaosVDEngine.h
+++ b/Engine/Plugins/ChaosVD/Source/ChaosVD/Private/ChaosVDEngine.h
@@ -9,6 +9,11 @@
 
 #include "ChaosVDEngine.generated.h"
 
+namespace Chaos::VisualDebugger
+{
+	struct FChaosVDOptionalDataChannel;
+}
+
 class FChaosVDTraceManager;
 class FChaosVDPlaybackController;
 class FChaosVDScene;
@@ -81,6 +86,8 @@ public:
 	}
 
 private:
+	void RestoreDataChannelsEnabledStateFromSave();
+	void UpdateSavedDataChannelsEnabledState(TWeakPtr<Chaos::VisualDebugger::FChaosVDOptionalDataChannel> DataChannelChanged);
 
 	void UpdateRecentFilesList(const FString& InFilename);
 
@@ -94,6 +101,8 @@ private:
 	bool bIsInitialized = false;
 
 	FDelegateHandle LiveSessionStoppedDelegateHandle;
+
+	FDelegateHandle DataChannelStateUpdatedHandle;
 	
 	FSessionStateChangedDelegate OnSessionOpenedDelegate;
 	FSessionStateChangedDelegate OnSessionClosedDelegate;
diff --git a/Engine/Plugins/ChaosVD/Source/ChaosVD/Private/Settings/ChaosVDMiscSettings.h b/Engine/Plugins/ChaosVD/Source/ChaosVD/Private/Settings/ChaosVDMiscSettings.h
--- a/Engine/Plugins/ChaosVD/Source/ChaosVD/Private/Settings/ChaosVDMiscSettings.h
+++ b/Engine/Plugins/ChaosVD/Source/ChaosVD/Private/Settings/ChaosVDMiscSettings.h
@@ -67,4 +67,8 @@ public:
 	/** Maximum number of recent files we can keep track of*/
 	UPROPERTY(config)
 	int32 MaxRecentFilesNum = 10;
+
+	/** Saved data channel enabled state */
+	UPROPERTY(config)
+	TMap<FString, bool> DataChannelEnabledState;
 };
