[ChaosVD] Addressing remaining issues of the multi file support implementation.

- The ids used in the serializable name table, are not guarantee to be stable between recordings. To fix that, now we keep the name table loaded in the trace provider, which is created per loaded file, instead on the main (shared) recording structure.
- Similar to above, we now also keep the Archive header data in the Trace Provider instance. this allow us load files from different versions at the same time (useful when AB testing against a previous recording)
- Fixed an issue where the collision channels name data was being serialized before we serialize the archive data. Without this fix, this will cause backward compatibility issues if we need to change its serialization code.

diff --git a/Engine/Plugins/ChaosVD/Source/ChaosVD/Private/Trace/ChaosVDTraceProvider.cpp b/Engine/Plugins/ChaosVD/Source/ChaosVD/Private/Trace/ChaosVDTraceProvider.cpp
--- a/Engine/Plugins/ChaosVD/Source/ChaosVD/Private/Trace/ChaosVDTraceProvider.cpp
+++ b/Engine/Plugins/ChaosVD/Source/ChaosVD/Private/Trace/ChaosVDTraceProvider.cpp
@@ -37,7 +37,11 @@ FName FChaosVDTraceProvider::ProviderName("ChaosVDProvider");
 
 FChaosVDTraceProvider::FChaosVDTraceProvider(TraceServices::IAnalysisSession& InSession): Session(InSession)
 {
-	DefaultHeaderData = Chaos::VisualDebugger::FChaosVDArchiveHeader::Current();
+	using namespace Chaos::VisualDebugger;
+	NameTable = MakeShared<FChaosVDSerializableNameTable>();
+
+	// Start with a default header data as a fallback
+	HeaderData = FChaosVDArchiveHeader::Current();
 }
 
 void FChaosVDTraceProvider::CreateRecordingInstanceForSession(const FString& InSessionName)
@@ -306,12 +310,14 @@ bool FChaosVDTraceProvider::ProcessBinaryData(const int32 DataID)
 					if (ensure(DataProcessorPtr->ProcessRawData(*RawData)))
 					{
 						UnprocessedDataByID.Remove(DataID);
+						DataProcessedSoFarCounter++;
 						return true;
 					}
 					else
 					{
 						UE_LOG(LogChaosVDEditor, Verbose, TEXT("[%s] Failed to serialize Binary Data with ID [%d] | Type [%s]"), ANSI_TO_TCHAR(__FUNCTION__), DataID, *UnprocessedData->TypeName);
 						TypesFailedToSerialize.Add(UnprocessedData->TypeName);
+						DataProcessedSoFarCounter++;
 					}
 				}
 			}
diff --git a/Engine/Plugins/ChaosVD/Source/ChaosVD/Private/Trace/DataProcessors/ChaosVDArchiveHeaderProcessor.cpp b/Engine/Plugins/ChaosVD/Source/ChaosVD/Private/Trace/DataProcessors/ChaosVDArchiveHeaderProcessor.cpp
--- a/Engine/Plugins/ChaosVD/Source/ChaosVD/Private/Trace/DataProcessors/ChaosVDArchiveHeaderProcessor.cpp
+++ b/Engine/Plugins/ChaosVD/Source/ChaosVD/Private/Trace/DataProcessors/ChaosVDArchiveHeaderProcessor.cpp
@@ -54,30 +54,20 @@ bool FChaosVDArchiveHeaderProcessor::ProcessRawData(const TArray<uint8>& InData)
 	FChaosVDArchiveHeader RecordedHeaderData;
 	const bool bSuccess = ReadDataFromBuffer(InData, RecordedHeaderData, ProviderSharedPtr.ToSharedRef());
 
-	if (const TSharedPtr<FChaosVDRecording> RecordingInstance = ProviderSharedPtr->GetRecordingForSession())
-	{
-		// This works under the assumption the header is the first thing written and therefore the first thing read
-		// If that didn't happen, we need to know to investigate further.
-		// Also, this assumption breaks under the new multi file support feature
+	// This works under the assumption the header is the first thing written and therefore the first thing read
+	// If that didn't happen, we need to know to investigate further.
 
-		//TODO: Add support for multiple headers in case we want to compare a new recording to a very old one
+	constexpr int64 DataExpectedToBeProcessedAtThisPoint = 0;
 
-		// Note: This is not a fatal error, unless pretty drastic serialization changes were made and the loaded file is old.
-		// CVD can gracefully handle serialization errors (as long the types serializers themselves don't crash and properly error out as expected).
-		if (!RecordingInstance->IsEmpty())
-		{
-			if (ensure(Private::IsCompatibleHeader(RecordingInstance->GetHeaderData(), RecordedHeaderData)))
-			{
-				UE_LOG(LogChaosVDEditor, Warning, TEXT("[%s] Processed an archive header while the recording already had data loaded. That initially loaded data used the default header and serialization errors might have occured | This should not happen..."), ANSI_TO_TCHAR(__FUNCTION__));
-			}	
-		}
+	// Note: This is not a fatal error, unless pretty drastic serialization changes were made and the loaded file is old.
+	// CVD can gracefully handle serialization errors (as long the types serializers themselves don't crash and properly error out as expected).
 
-		RecordingInstance->SetHeaderData(RecordedHeaderData);
-	}
-	else
+	if (!ensure(ProviderSharedPtr->GetDataProcessedSoFarNum() == DataExpectedToBeProcessedAtThisPoint))
 	{
-		UE_LOG(LogChaosVDEditor, Error, TEXT("[%s] Processed an archive header without a valid recording instance. Any follwing archive reads will use the default header and serialization errors might occur..."), ANSI_TO_TCHAR(__FUNCTION__));	
+		UE_LOG(LogChaosVDEditor, Warning, TEXT("[%s] Processed an archive header while the recording already had data loaded. That initially loaded data used the default header and serialization errors might have occured | This should not happen..."), ANSI_TO_TCHAR(__FUNCTION__));
 	}
 
+	ProviderSharedPtr->SetHeaderData(RecordedHeaderData);
+
 	return bSuccess;
 }
diff --git a/Engine/Plugins/ChaosVD/Source/ChaosVD/Private/Trace/DataProcessors/ChaosVDCollisionChannelsInfoDataProcessor.cpp b/Engine/Plugins/ChaosVD/Source/ChaosVD/Private/Trace/DataProcessors/ChaosVDCollisionChannelsInfoDataProcessor.cpp
--- a/Engine/Plugins/ChaosVD/Source/ChaosVD/Private/Trace/DataProcessors/ChaosVDCollisionChannelsInfoDataProcessor.cpp
+++ b/Engine/Plugins/ChaosVD/Source/ChaosVD/Private/Trace/DataProcessors/ChaosVDCollisionChannelsInfoDataProcessor.cpp
@@ -2,6 +2,9 @@
 
 #include "ChaosVDCollisionChannelsInfoDataProcessor.h"
 
+#include "ChaosVDModule.h"
+
+
 FChaosVDCollisionChannelsInfoDataProcessor::FChaosVDCollisionChannelsInfoDataProcessor() : FChaosVDDataProcessorBase(FChaosVDCollisionChannelsInfoContainer::WrapperTypeName)
 {
 	
@@ -17,13 +20,76 @@ bool FChaosVDCollisionChannelsInfoDataProcessor::ProcessRawData(const TArray<uin
 		return false;
 	}
 
+	const bool bNeedsLoadingOrderFixup = ProviderSharedPtr->GetDataProcessedSoFarNum() == 0;
+
+	if (bNeedsLoadingOrderFixup)
+	{
+		// All serializable data in CVD needs to be backwards compatible. In order to do that the archive header we use needs to be the first thing we serialize and the first thing we read
+		// There was an issue in the current implementation where the collision data container was being traced before the header data, therefore as a workaround if we are trying to read the collision
+		// container data first, we need to delay it at least one frame. This is ok as this data is only used for the UI to decode collision channel ids to human-readable names.
+		FTSTicker::GetCoreTicker().AddTicker(FTickerDelegate::CreateLambda([this, WeakTraceProvider = ProviderSharedPtr.ToWeakPtr()](float DeltaTime)
+		{
+			TSharedPtr<FChaosVDTraceProvider> ProviderSharedPtr = WeakTraceProvider.Pin();
+
+			if (ProviderSharedPtr)
+			{
+				if (!ProcessPendingCollisionChannelData(PendingCollisionChannelDataToProcess))
+				{
+					UE_LOG(LogChaosVDEditor, Warning, TEXT("[%hs] Failed to process deferred collision channel data."), __func__);
+					ProviderSharedPtr->TypesFailedToSerialize.Add(FString(CompatibleType));
+				}
+
+				ProviderSharedPtr->DataProcessedSoFarCounter++;
+				PendingCollisionChannelDataToProcess.Empty();
+			}
+				
+			return false;
+		}));
+
+		PendingCollisionChannelDataToProcess = InData;
+		ProviderSharedPtr->DataProcessedSoFarCounter--;
+		
+		UE_LOG(LogChaosVDEditor, Warning, TEXT("[%hs] Deferring loading collision channel data to the next frame, because the archive header is not loaded yet."), __func__);
+
+		return true;
+	}
+	else
+	{
+		return ProcessPendingCollisionChannelData(InData);
+	}
+
+	return false;
+}
+
+bool FChaosVDCollisionChannelsInfoDataProcessor::ProcessPendingCollisionChannelData(const TArray<uint8>& InData)
+{
+	TSharedPtr<FChaosVDTraceProvider> ProviderSharedPtr = TraceProvider.Pin();
+	TSharedPtr<FChaosVDRecording> RecordingInstancePtr = ProviderSharedPtr ? ProviderSharedPtr->GetRecordingForSession() : nullptr;
+	if (!ensure(ProviderSharedPtr))
+	{
+		return false;
+	}
+
 	TSharedPtr<FChaosVDCollisionChannelsInfoContainer> CollisionChannelsData = MakeShared<FChaosVDCollisionChannelsInfoContainer>();
 	const bool bSuccess = Chaos::VisualDebugger::ReadDataFromBuffer(InData, *CollisionChannelsData, ProviderSharedPtr.ToSharedRef());
 	
 	if (bSuccess)
 	{
+		if (RecordingInstancePtr->GetCollisionChannelsInfoContainer())
+		{
+			EChaosVDRecordingAttributes RecordingAttributes = RecordingInstancePtr->GetAttributes();
+			if (EnumHasAnyFlags(RecordingAttributes, EChaosVDRecordingAttributes::Merged))
+			{
+				UE_LOG(LogChaosVDEditor, Warning, TEXT("[%hs] Multiple Collision data channel in multi file mode is not supported yet. \n The collision channel data from the last loaded recording will be used. This is used for to decode collision channels ids into names, \n so as long the recordings are from the same build, this warning should not cause issues."), __func__);
+			}
+			else
+			{
+				UE_LOG(LogChaosVDEditor, Error, TEXT("[%hs] Collision channel data already loaded. This should not happen. Was the collision data serialized twice?"), __func__);
+			}
+			
+		}
 		ProviderSharedPtr->GetRecordingForSession()->SetCollisionChannelsInfoContainer(CollisionChannelsData);
 	}
-	
+
 	return bSuccess;
 }
diff --git a/Engine/Plugins/ChaosVD/Source/ChaosVD/Private/Trace/DataProcessors/ChaosVDCollisionChannelsInfoDataProcessor.h b/Engine/Plugins/ChaosVD/Source/ChaosVD/Private/Trace/DataProcessors/ChaosVDCollisionChannelsInfoDataProcessor.h
--- a/Engine/Plugins/ChaosVD/Source/ChaosVD/Private/Trace/DataProcessors/ChaosVDCollisionChannelsInfoDataProcessor.h
+++ b/Engine/Plugins/ChaosVD/Source/ChaosVD/Private/Trace/DataProcessors/ChaosVDCollisionChannelsInfoDataProcessor.h
@@ -14,5 +14,11 @@ public:
 	explicit FChaosVDCollisionChannelsInfoDataProcessor();
 	
 	virtual bool ProcessRawData(const TArray<uint8>& InData) override;
+
+private:
+
+	bool ProcessPendingCollisionChannelData(const TArray<uint8>& InData);
+	
+	TArray<uint8> PendingCollisionChannelDataToProcess;
 };
 
diff --git a/Engine/Plugins/ChaosVD/Source/ChaosVD/Private/Trace/DataProcessors/ChaosVDSerializedNameEntryDataProcessor.cpp b/Engine/Plugins/ChaosVD/Source/ChaosVD/Private/Trace/DataProcessors/ChaosVDSerializedNameEntryDataProcessor.cpp
--- a/Engine/Plugins/ChaosVD/Source/ChaosVD/Private/Trace/DataProcessors/ChaosVDSerializedNameEntryDataProcessor.cpp
+++ b/Engine/Plugins/ChaosVD/Source/ChaosVD/Private/Trace/DataProcessors/ChaosVDSerializedNameEntryDataProcessor.cpp
@@ -22,10 +22,7 @@ bool FChaosVDSerializedNameEntryDataProcessor::ProcessRawData(const TArray<uint8
 		return false;
 	}
 
-	TSharedPtr<FChaosVDRecording> RecordingInstance = ProviderSharedPtr->GetRecordingForSession();
-
-	TSharedPtr<Chaos::VisualDebugger::FChaosVDSerializableNameTable> NameTableInstance = RecordingInstance ? RecordingInstance->GetNameTableInstance() : nullptr;
-
+	TSharedPtr<Chaos::VisualDebugger::FChaosVDSerializableNameTable> NameTableInstance = ProviderSharedPtr->GetNameTableInstance();
 	if (!ensure(NameTableInstance.IsValid()))
 	{
 		return false;
@@ -35,7 +32,7 @@ bool FChaosVDSerializedNameEntryDataProcessor::ProcessRawData(const TArray<uint8
 	
 	FMemoryReader MemReader(InData);
 
-	const Chaos::VisualDebugger::FChaosVDArchiveHeader& RecordedHeader = RecordingInstance->GetHeaderData();
+	const Chaos::VisualDebugger::FChaosVDArchiveHeader& RecordedHeader = ProviderSharedPtr->GetHeaderData();
 	ApplyHeaderDataToArchive(MemReader, RecordedHeader);
 	
 	MemReader << NameEntry;
diff --git a/Engine/Plugins/ChaosVD/Source/ChaosVD/Public/Trace/ChaosVDTraceProvider.h b/Engine/Plugins/ChaosVD/Source/ChaosVD/Public/Trace/ChaosVDTraceProvider.h
--- a/Engine/Plugins/ChaosVD/Source/ChaosVD/Public/Trace/ChaosVDTraceProvider.h
+++ b/Engine/Plugins/ChaosVD/Source/ChaosVD/Public/Trace/ChaosVDTraceProvider.h
@@ -134,13 +134,36 @@ public:
 	virtual void EditAccessCheck() const override{};
 	
 	int32 CHAOSVD_API GetRemappedSolverID(int32 SolverID);
-	int32 CHAOSVD_API GetRemappedSolverID_AssumesLocked(int32 SolverID);
 
 	bool DoesOwnRecordingInstance() const
 	{
 		return !bHasRecordingOverride;
 	}
 
+	/** Returns the name table instances used to de-duplicate strings serialization */
+	TSharedPtr<Chaos::VisualDebugger::FChaosVDSerializableNameTable> GetNameTableInstance() const
+	{
+		return NameTable;
+	}
+
+	/** Returns the FArchive header used to read the serialized binary data */
+	const Chaos::VisualDebugger::FChaosVDArchiveHeader& GetHeaderData() const
+	{
+		return HeaderData;
+	}
+
+	/** Sets the FArchive header used to read the serialized binary data */
+	void SetHeaderData(Chaos::VisualDebugger::FChaosVDArchiveHeader& InNewHeader)
+	{
+		HeaderData = MoveTemp(InNewHeader);
+	}
+
+	/** Returns how many pieces of data we processed so far with a data processor (even if it failed) */
+	int64 GetDataProcessedSoFarNum() const
+	{
+		return DataProcessedSoFarCounter;
+	}
+
 private:
 	
 	int32 GetCurrentGameThreadTrackID() const
@@ -152,6 +175,10 @@ private:
 	{
 		RemappedGameThreadTrackID = NewID;
 	}
+	
+	Chaos::VisualDebugger::FChaosVDArchiveHeader HeaderData;
+	
+	TSharedPtr<Chaos::VisualDebugger::FChaosVDSerializableNameTable> NameTable;
 
 	TraceServices::IAnalysisSession& Session;
 
@@ -167,7 +194,7 @@ private:
 
 	TWeakPtr<FChaosVDGameFrameData> CurrentGameFrame = nullptr;
 
-	Chaos::VisualDebugger::FChaosVDArchiveHeader DefaultHeaderData;
+	TUniquePtr<Chaos::VisualDebugger::FChaosVDArchiveHeader> CurrentHeaderData;
 
 	int32 CurrentGameFrameQueueSize = 0;
 
@@ -185,6 +212,12 @@ private:
 	TSet<FString> MissingDataProcessors;
 	TSet<FString> TypesFailedToSerialize;
 
+	int64 DataProcessedSoFarCounter = 0; 
+
 	friend class FChaosVDTraceAnalyzer;
+
+	// Not idea, but needed to for the serialization fix up logic of the collision data channel info
+	// An alternative is expose a plugin function deprecated from the get go
+	friend class FChaosVDCollisionChannelsInfoDataProcessor;
 };
 
diff --git a/Engine/Plugins/ChaosVD/Source/ChaosVD/Public/Trace/DataProcessors/ChaosVDDataProcessorBase.h b/Engine/Plugins/ChaosVD/Source/ChaosVD/Public/Trace/DataProcessors/ChaosVDDataProcessorBase.h
--- a/Engine/Plugins/ChaosVD/Source/ChaosVD/Public/Trace/DataProcessors/ChaosVDDataProcessorBase.h
+++ b/Engine/Plugins/ChaosVD/Source/ChaosVD/Public/Trace/DataProcessors/ChaosVDDataProcessorBase.h
@@ -31,8 +31,7 @@ namespace Chaos::VisualDebugger
 	template<typename TDataToSerialize>
 	bool ReadDataFromBuffer(const TArray<uint8>& InDataBuffer, TDataToSerialize& Data, const TSharedRef<FChaosVDTraceProvider>& DataProvider)
 	{
-		const TSharedPtr<FChaosVDRecording> RecordingInstance = DataProvider->GetRecordingForSession();
-		const TSharedPtr<FChaosVDSerializableNameTable> NameTableInstance = RecordingInstance ? RecordingInstance->GetNameTableInstance() : nullptr;
+		const TSharedPtr<FChaosVDSerializableNameTable> NameTableInstance = DataProvider->GetNameTableInstance();
 
 		if (!ensure(NameTableInstance.IsValid()))
 		{
@@ -40,7 +39,7 @@ namespace Chaos::VisualDebugger
 		}
 		
 		FChaosVDMemoryReader MemReader(InDataBuffer, NameTableInstance.ToSharedRef());
-		const FChaosVDArchiveHeader& RecordedHeader =  RecordingInstance->GetHeaderData();
+		const FChaosVDArchiveHeader& RecordedHeader = DataProvider->GetHeaderData();
 		ApplyHeaderDataToArchive(MemReader, RecordedHeader);
 
 		bool bSuccess = false;
diff --git a/Engine/Source/Runtime/Experimental/Chaos/Public/ChaosVisualDebugger/ChaosVisualDebuggerTrace.h b/Engine/Source/Runtime/Experimental/Chaos/Public/ChaosVisualDebugger/ChaosVisualDebuggerTrace.h
--- a/Engine/Source/Runtime/Experimental/Chaos/Public/ChaosVisualDebugger/ChaosVisualDebuggerTrace.h
+++ b/Engine/Source/Runtime/Experimental/Chaos/Public/ChaosVisualDebugger/ChaosVisualDebuggerTrace.h
@@ -173,6 +173,8 @@ using FChaosVDSerializableNameTable = Chaos::VisualDebugger::FChaosVDSerializabl
 enum class EChaosVDTraceBinaryDataOptions
 {
 	None = 0,
+	/** Will trace the provided data buffer into CVD's trace channel, regardless id CVD's trace system is fully initialized.
+	 * Data traced in this way, will not be backwards compatible as the required header data will not be ready yet */
 	ForceTrace = 1 << 0
 };
 ENUM_CLASS_FLAGS(EChaosVDTraceBinaryDataOptions)
diff --git a/Engine/Source/Runtime/Experimental/ChaosSolverEngine/Private/Chaos/ChaosVDEngineEditorBridge.cpp b/Engine/Source/Runtime/Experimental/ChaosSolverEngine/Private/Chaos/ChaosVDEngineEditorBridge.cpp
--- a/Engine/Source/Runtime/Experimental/ChaosSolverEngine/Private/Chaos/ChaosVDEngineEditorBridge.cpp
+++ b/Engine/Source/Runtime/Experimental/ChaosSolverEngine/Private/Chaos/ChaosVDEngineEditorBridge.cpp
@@ -47,7 +47,12 @@ void FChaosVDEngineEditorBridge::RemoveOnScreenRecordingMessage()
 
 void FChaosVDEngineEditorBridge::HandleCVDRecordingStarted()
 {
-	SerializeCollisionChannelsNames();
+	FTSTicker::GetCoreTicker().AddTicker(FTickerDelegate::CreateLambda([](float DeltaTime)
+	{
+		// We need to wait at least one frame to serialize the collision channel names to ensure the Archive header used for the whole session is traced
+		FChaosVDEngineEditorBridge::Get().SerializeCollisionChannelsNames();
+		return false;
+	}));
 
 	AddOnScreenRecordingMessage();
 
@@ -121,7 +126,7 @@ void FChaosVDEngineEditorBridge::SerializeCollisionChannelsNames()
 
 	Chaos::VisualDebugger::WriteDataToBuffer(CollisionChannelsDataBuffer, CollisionChannelInfoContainer);
 
-	CVD_TRACE_BINARY_DATA(CollisionChannelsDataBuffer, FChaosVDCollisionChannelsInfoContainer::WrapperTypeName, EChaosVDTraceBinaryDataOptions::ForceTrace)
+	CVD_TRACE_BINARY_DATA(CollisionChannelsDataBuffer, FChaosVDCollisionChannelsInfoContainer::WrapperTypeName);
 }
 
 FChaosVDEngineEditorBridge& FChaosVDEngineEditorBridge::Get()
diff --git a/Engine/Source/Runtime/Experimental/ChaosVDData/Private/ChaosVDRecording.cpp b/Engine/Source/Runtime/Experimental/ChaosVDData/Private/ChaosVDRecording.cpp
--- a/Engine/Source/Runtime/Experimental/ChaosVDData/Private/ChaosVDRecording.cpp
+++ b/Engine/Source/Runtime/Experimental/ChaosVDData/Private/ChaosVDRecording.cpp
@@ -22,10 +22,6 @@ void FChaosVDCustomFrameData::AddData(const FChaosVDCustomUserDataHandle& InData
 
 FChaosVDRecording::FChaosVDRecording()
 {
-	NameTable = MakeShared<Chaos::VisualDebugger::FChaosVDSerializableNameTable>();
-
-	// Start with a default header data as a fallback
-	HeaderData = Chaos::VisualDebugger::FChaosVDArchiveHeader::Current();
 }
 
 int32 FChaosVDRecording::GetAvailableGameFramesNumber() const
diff --git a/Engine/Source/Runtime/Experimental/ChaosVDData/Public/ChaosVDRecording.h b/Engine/Source/Runtime/Experimental/ChaosVDData/Public/ChaosVDRecording.h
--- a/Engine/Source/Runtime/Experimental/ChaosVDData/Public/ChaosVDRecording.h
+++ b/Engine/Source/Runtime/Experimental/ChaosVDData/Public/ChaosVDRecording.h
@@ -640,24 +640,6 @@ struct CHAOSVDDATA_API FChaosVDRecording
 		return RecordingAttributes;
 	}
 
-	/** Returns the name table instances used to de-duplicate strings serialization */
-	TSharedPtr<Chaos::VisualDebugger::FChaosVDSerializableNameTable> GetNameTableInstance() const
-	{
-		return NameTable;
-	}
-
-	/** Returns the FArchive header used to read the serialized binary data */
-	const Chaos::VisualDebugger::FChaosVDArchiveHeader& GetHeaderData() const
-	{
-		return HeaderData;
-	}
-	
-	/** Sets the FArchive header used to read the serialized binary data */
-	void SetHeaderData(const Chaos::VisualDebugger::FChaosVDArchiveHeader& InNewHeader)
-	{
-		HeaderData = InNewHeader;
-	}
-
 	/** Returns true if this recording does not have any usable data */
 	bool IsEmpty() const;
 
@@ -707,8 +689,6 @@ protected:
 	/** Id to Ptr map of all shared geometry data required to visualize */
 	TMap<uint32, Chaos::FConstImplicitObjectPtr> ImplicitObjects;
 
-	TSharedPtr<Chaos::VisualDebugger::FChaosVDSerializableNameTable> NameTable;
-
 	mutable FRWLock RecordingDataLock;
 
 	EChaosVDRecordingAttributes RecordingAttributes = EChaosVDRecordingAttributes::None; 
@@ -719,8 +699,6 @@ protected:
 	/** Map that temporary holds generated particle data during the key frame generation process, keeping its memory allocation between generated frames*/
 	TMap<int32, TSharedPtr<FChaosVDParticleDataWrapper>> ParticlesOnCurrentGeneratedKeyframe;
 
-	Chaos::VisualDebugger::FChaosVDArchiveHeader HeaderData;
-
 	TSharedPtr<FChaosVDCollisionChannelsInfoContainer> CollisionChannelsInfoContainer;
 
 	TSet<int32> ReservedSolverIDs;
