[ChaosVD] Fix for an issue where the TEDS Outliner is constantly iterating over all data rows to check if parent data has changed.

This is a problem because in recordings from complex maps like city sample, this can take up to 5ms in worst cases every frame.

This updates CVD code to add a different tag to indicate the TEDS Outliner need to update the UI, each time we update the parent relationship between CVD objects, and updates the outliner to only use that tag.

diff --git a/Engine/Plugins/ChaosVD/Source/ChaosVD/Private/TEDS/ChaosVDParentDataStorageFactory.cpp b/Engine/Plugins/ChaosVD/Source/ChaosVD/Private/TEDS/ChaosVDParentDataStorageFactory.cpp
--- a/Engine/Plugins/ChaosVD/Source/ChaosVD/Private/TEDS/ChaosVDParentDataStorageFactory.cpp
+++ b/Engine/Plugins/ChaosVD/Source/ChaosVD/Private/TEDS/ChaosVDParentDataStorageFactory.cpp
@@ -21,7 +21,7 @@ namespace Chaos::VD::ParentDataStorageUtils::Private
 {
 	using namespace UE::Editor::DataStorage::Queries;
 
-	void UpdateParentData(ICoreProvider& DataStorage, FChaosVDTableRowParentColumn& RowParentColumn, RowHandle ItemRowHandle, RowHandle NewParentRowHandle)
+	void UpdateParentData(ICoreProvider& DataStorage, IQueryContext& Context, FChaosVDTableRowParentColumn& RowParentColumn, RowHandle ItemRowHandle, RowHandle NewParentRowHandle)
 	{
 		if (DataStorage.IsRowAvailable(NewParentRowHandle))
 		{
@@ -30,6 +30,7 @@ namespace Chaos::VD::ParentDataStorageUtils::Private
 				if (FChaosVDTableRowParentColumn* OldParentData = DataStorage.GetColumn<FChaosVDTableRowParentColumn>(RowParentColumn.ParentObject))
 				{
 					OldParentData->ChildrenSet.Remove(ItemRowHandle);
+					Context.AddColumns<FTypedElementSyncBackToWorldTag>(RowParentColumn.ParentObject);
 				}
 
 				RowParentColumn.ParentObject = NewParentRowHandle;
@@ -37,6 +38,7 @@ namespace Chaos::VD::ParentDataStorageUtils::Private
 				if (FChaosVDTableRowParentColumn* ParentData = DataStorage.GetColumn<FChaosVDTableRowParentColumn>(NewParentRowHandle))
 				{
 					ParentData->ChildrenSet.Emplace(NewParentRowHandle);
+					Context.AddColumns<FTypedElementSyncBackToWorldTag>(NewParentRowHandle);
 				}
 			}
 		}
@@ -123,7 +125,7 @@ void UChaosVDParentDataStorageFactory::RegisterUpdateOrRemoveParentColumn(UE::Ed
 						NewParentRowHandle = Context.LookupMappedRow(IdKey);
 					}
 
-					UpdateParentData(DataStorage, ParentColumn, Row, NewParentRowHandle);
+					UpdateParentData(DataStorage, Context, ParentColumn, Row, NewParentRowHandle);
 				}
 			})
 		.Where()
@@ -147,7 +149,7 @@ void UChaosVDParentDataStorageFactory::RegisterUpdateOrRemoveParentColumn(UE::Ed
 						NewParentRowHandle = Context.LookupMappedRow(IdKey);
 					}
 					
-					UpdateParentData(DataStorage, Parent, Row, NewParentRowHandle);
+					UpdateParentData(DataStorage, Context, Parent, Row, NewParentRowHandle);
 				}
 			})
 		.Where()
@@ -165,6 +167,8 @@ void UChaosVDParentDataStorageFactory::RegisterUpdateOrRemoveParentColumn(UE::Ed
 						if (FChaosVDTableRowParentColumn* OldParentData = DataStorage.GetColumn<FChaosVDTableRowParentColumn>(ParentData->ParentObject))
 						{
 							OldParentData->ChildrenSet.Remove(Row);
+							Context.AddColumns<FTypedElementSyncBackToWorldTag>(ParentData->ParentObject);
+							Context.AddColumns<FTypedElementSyncBackToWorldTag>(Row);
 						}
 					}
 				})
diff --git a/Engine/Plugins/Experimental/EditorDataStorageFeatures/Source/TedsOutliner/Private/TedsOutlinerImpl.cpp b/Engine/Plugins/Experimental/EditorDataStorageFeatures/Source/TedsOutliner/Private/TedsOutlinerImpl.cpp
--- a/Engine/Plugins/Experimental/EditorDataStorageFeatures/Source/TedsOutliner/Private/TedsOutlinerImpl.cpp
+++ b/Engine/Plugins/Experimental/EditorDataStorageFeatures/Source/TedsOutliner/Private/TedsOutlinerImpl.cpp
@@ -634,7 +634,7 @@ void FTedsOutlinerImpl::RecompileQueries()
 			.ReadOnly(ParentColumnType, EOptional::Yes)
 			.ReadOnly<FTedsOutlinerColumn>(EOptional::Yes)
 		.Where()
-			.Any<FTypedElementSyncFromWorldTag, FTypedElementSyncBackToWorldTag>()
+			.All<FTypedElementSyncBackToWorldTag>()
 		.Compile();
 		
 		// Add the conditions from FinalQueryDescription to ensure we are the rows the user requested
