Add debug drawing for Trimesh's BVH to CVD

diff --git a/Engine/Plugins/ChaosVD/Source/ChaosVD/Private/Settings/ChaosVDParticleVisualizationSettings.h b/Engine/Plugins/ChaosVD/Source/ChaosVD/Private/Settings/ChaosVDParticleVisualizationSettings.h
--- a/Engine/Plugins/ChaosVD/Source/ChaosVD/Private/Settings/ChaosVDParticleVisualizationSettings.h
+++ b/Engine/Plugins/ChaosVD/Source/ChaosVD/Private/Settings/ChaosVDParticleVisualizationSettings.h
@@ -271,6 +271,14 @@ public:
 	UPROPERTY(config, EditAnywhere, Category=DebugDraw)
 	float CenterOfMassRadius = 10.0f;
 
+	/** Should a triangle mesh's BVH draw. */
+	UPROPERTY(config, EditAnywhere, Category=DebugDraw)
+	bool bDrawTriMeshBVH = false;
+
+	/** What depth of a triangle mesh's BVH to draw. Level '-1' means draw everything. */
+	UPROPERTY(config, EditAnywhere, Category=DebugDraw, meta = (EditCondition = "bDrawTriMeshBVH", EditConditionHides))
+	int TriMeshBVHDrawLevel = 0;
+
 	UPROPERTY(config, EditAnywhere, Category=DebugDraw)
 	FChaosParticleDataDebugDrawColors ColorSettings;
 
diff --git a/Engine/Plugins/ChaosVD/Source/ChaosVD/Private/Visualizers/ChaosVDParticleDataComponentVisualizer.cpp b/Engine/Plugins/ChaosVD/Source/ChaosVD/Private/Visualizers/ChaosVDParticleDataComponentVisualizer.cpp
--- a/Engine/Plugins/ChaosVD/Source/ChaosVD/Private/Visualizers/ChaosVDParticleDataComponentVisualizer.cpp
+++ b/Engine/Plugins/ChaosVD/Source/ChaosVD/Private/Visualizers/ChaosVDParticleDataComponentVisualizer.cpp
@@ -18,6 +18,9 @@
 #include "Visualizers/ChaosVDDebugDrawUtils.h"
 #include "Widgets/SChaosVDEnumFlagsMenu.h"
 #include "Widgets/SChaosVDViewportToolbar.h"
+#include "Chaos/ImplicitObject.h"
+#include "Chaos/ImplicitObjectScaled.h"
+#include "Chaos/TriangleMeshImplicitObject.h"
 
 #define LOCTEXT_NAMESPACE "ChaosVisualDebugger"
 
@@ -337,6 +340,109 @@ void FChaosVDParticleDataComponentVisualizer::DrawVisualizationForParticleData(c
 			}
 		}
 	}
+
+	// Draw the selected tri mesh's BVH
+	if (InVisualizationContext.DebugDrawSettings->bDrawTriMeshBVH)
+	{
+		if (const TSharedPtr<FChaosVDInstancedMeshData> SelectedMeshInstance = InParticleInstance.GetSelectedMeshInstance().Pin())
+		{
+			const FChaosVDParticlePositionRotation& ParticlePositionRotation = InParticleInstance.GetParticleData()->ParticlePositionRotation;
+			FTransform Transform = FTransform::Identity;
+			Transform.SetTranslation(ParticlePositionRotation.MX);
+			Transform.SetRotation(ParticlePositionRotation.MR);
+
+			const TSharedRef<FChaosVDExtractedGeometryDataHandle>& GeometryHandle = SelectedMeshInstance->GetGeometryHandle();
+			if (const Chaos::FImplicitObject* ImplicitObject = GeometryHandle->GetImplicitObject())
+			{
+				if (const Chaos::TImplicitObjectScaled<Chaos::FTriangleMeshImplicitObject>* ScaledTriMesh = ImplicitObject->AsA<Chaos::TImplicitObjectScaled<Chaos::FTriangleMeshImplicitObject>>())
+				{
+					FChaosVDTrimeshBVHVisualizer TrimeshBVHVisualizer;
+					const FVector Scale = ScaledTriMesh->GetScale();
+					Transform.SetScale3D(Scale);
+					TrimeshBVHVisualizer.Draw(PDI, InVisualizationContext, Transform, *ScaledTriMesh->GetUnscaledObject());
+				}
+				else if (const Chaos::TImplicitObjectInstanced<Chaos::FTriangleMeshImplicitObject>* InstancedTriMesh = ImplicitObject->AsA<Chaos::TImplicitObjectInstanced<Chaos::FTriangleMeshImplicitObject>>())
+				{
+					FChaosVDTrimeshBVHVisualizer TrimeshBVHVisualizer;
+					TrimeshBVHVisualizer.Draw(PDI, InVisualizationContext, Transform, *InstancedTriMesh->GetInstancedObject());
+				}
+				else if (const Chaos::FTriangleMeshImplicitObject* TriMesh = ImplicitObject->AsA<Chaos::FTriangleMeshImplicitObject>())
+				{
+					FChaosVDTrimeshBVHVisualizer TrimeshBVHVisualizer;
+					TrimeshBVHVisualizer.Draw(PDI, InVisualizationContext, Transform, *TriMesh);
+				}
+			}
+		}
+	}
+}
+
+void FChaosVDTrimeshBVHVisualizer::Draw(FPrimitiveDrawInterface* PDI, const FChaosVDParticleDataVisualizationContext& InVisualizationContext, const FTransform& LocalToWorldTransform, const Chaos::FTriangleMeshImplicitObject& TriMesh) const
+{
+	const int32 TargetLevel = InVisualizationContext.DebugDrawSettings->TriMeshBVHDrawLevel;
+	DrawBVH(PDI, LocalToWorldTransform, TriMesh.GetBVH(), TargetLevel);
+}
+
+void FChaosVDTrimeshBVHVisualizer::DrawBVH(FPrimitiveDrawInterface* PDI, const FTransform& LocalToWorldTransform, const Chaos::FTrimeshBVH& BVH, const int32 TargetLevel) const
+{
+	const int32 RootIndex = 0;
+	// The root has to be handled differently as there is no aabb in the tree for the root, it's implicit from the first two child aabbs.
+	if (TargetLevel == 0)
+	{
+		Chaos::FAABB3 RootAabb = ToAabb(BVH.Nodes[RootIndex].Children.GetBounds(0));
+		RootAabb.GrowToInclude(ToAabb(BVH.Nodes[RootIndex].Children.GetBounds(1)));
+		DrawAabb(PDI, LocalToWorldTransform, RootAabb, 0);
+	}
+	else
+	{
+		// Start at level 1 instead of 0 as the root is handled above and the aabbs for the children are embedded in the node
+		DrawBVHLevel(PDI, LocalToWorldTransform, BVH, RootIndex, 1, TargetLevel);
+	}
+}
+
+void FChaosVDTrimeshBVHVisualizer::DrawBVHLevel(FPrimitiveDrawInterface* PDI, const FTransform& LocalToWorldTransform, const Chaos::FTrimeshBVH& BVH, const int32 NodeIndex, const int32 CurrentLevel, const int32 TargetLevel) const
+{
+	const Chaos::FTrimeshBVH::FNode& Node = BVH.Nodes[NodeIndex];
+
+	if (TargetLevel == -1 || TargetLevel == CurrentLevel)
+	{
+		DrawAabb(PDI, LocalToWorldTransform, ToAabb(Node.Children.GetBounds(0)), CurrentLevel);
+		DrawAabb(PDI, LocalToWorldTransform, ToAabb(Node.Children.GetBounds(1)), CurrentLevel);
+	}
+	if (TargetLevel == -1 || CurrentLevel < TargetLevel)
+	{
+		for (int32 ChildIndex = 0; ChildIndex < 2; ++ChildIndex)
+		{
+			const bool bIsLeaf = Node.Children.GetFaceCount(ChildIndex) != 0;
+			if (!bIsLeaf)
+			{
+				const int32 ChildNodeIndex = Node.Children.GetChildOrFaceIndex(ChildIndex);
+				DrawBVHLevel(PDI, LocalToWorldTransform, BVH, ChildNodeIndex, CurrentLevel + 1, TargetLevel);
+			}
+		}
+	}
+}
+
+void FChaosVDTrimeshBVHVisualizer::DrawAabb(FPrimitiveDrawInterface* PDI, const FTransform& LocalToWorldTransform, const Chaos::FAABB3& Aabb, const int32 ColorSeed) const
+{
+	// Build the local to world transform for a "unit" aabb (unit being half-extents of 0.5)
+	const FVector UnitHalfExtent(0.5f);
+	const FTransform LocalAabbTransform(FQuat::Identity, Aabb.GetCenter(), Aabb.Extents());
+	const FTransform WorldAabbTransform = LocalAabbTransform * LocalToWorldTransform;
+	const FColor Color = FColor::MakeRandomSeededColor(ColorSeed);
+	FChaosVDDebugDrawUtils::DrawBox(PDI, UnitHalfExtent, Color, WorldAabbTransform, FText());
+}
+
+Chaos::FAABB3 FChaosVDTrimeshBVHVisualizer::ToAabb(const Chaos::FAABBVectorized& AabbVectorized) const
+{
+	auto ToVector = [](VectorRegister4Float VectorRegister)
+	{
+		double Values[4];
+		VectorStore(VectorRegister, Values);
+		return FVector(Values[0], Values[1], Values[2]);
+	};
+	const FVector Min = ToVector(AabbVectorized.GetMin());
+	const FVector Max = ToVector(AabbVectorized.GetMax());
+	return Chaos::FAABB3(Min, Max);
 }
 
 #undef LOCTEXT_NAMESPACE
diff --git a/Engine/Plugins/ChaosVD/Source/ChaosVD/Private/Visualizers/ChaosVDParticleDataComponentVisualizer.h b/Engine/Plugins/ChaosVD/Source/ChaosVD/Private/Visualizers/ChaosVDParticleDataComponentVisualizer.h
--- a/Engine/Plugins/ChaosVD/Source/ChaosVD/Private/Visualizers/ChaosVDParticleDataComponentVisualizer.h
+++ b/Engine/Plugins/ChaosVD/Source/ChaosVD/Private/Visualizers/ChaosVDParticleDataComponentVisualizer.h
@@ -15,6 +15,13 @@ class UChaosVDParticleVisualizationDebugDrawSettings;
 struct FChaosParticleDataDebugDrawSettings;
 struct FChaosVDParticleDataWrapper;
 
+namespace Chaos
+{
+	struct FTrimeshBVH;
+	class FTriangleMeshImplicitObject;
+	class FAABBVectorized;
+} // namespace Chaos
+
 
 struct FChaosVDParticleDataVisualizationContext : public FChaosVDVisualizationContext
 {
@@ -52,3 +59,14 @@ protected:
 	void DrawParticleVector(FPrimitiveDrawInterface* PDI, const FVector& StartLocation, const FVector& InVector, EChaosVDParticleDataVisualizationFlags VectorID, const FChaosVDParticleDataVisualizationContext& InVisualizationContext, float LineThickness);
 	void DrawVisualizationForParticleData(const UActorComponent* Component, FPrimitiveDrawInterface* PDI, const FSceneView* View, const FChaosVDParticleDataVisualizationContext& InVisualizationContext, const FChaosVDSceneParticle& InParticleInstance);
 };
+
+class FChaosVDTrimeshBVHVisualizer
+{
+public:
+	void Draw(FPrimitiveDrawInterface* PDI, const FChaosVDParticleDataVisualizationContext& InVisualizationContext, const FTransform& LocalToWorldTransform, const Chaos::FTriangleMeshImplicitObject& TriMesh) const;
+private:
+	void DrawBVH(FPrimitiveDrawInterface* PDI, const FTransform& LocalToWorldTransform, const Chaos::FTrimeshBVH& BVH, const int32 TargetLevel) const;
+	void DrawBVHLevel(FPrimitiveDrawInterface* PDI, const FTransform& LocalToWorldTransform, const Chaos::FTrimeshBVH& BVH, const int32 NodeIndex, const int32 CurrentLevel, const int32 TargetLevel) const;
+	void DrawAabb(FPrimitiveDrawInterface* PDI, const FTransform& LocalToWorldTransform, const Chaos::FAABB3& Aabb, const int32 ColorSeed) const;
+	Chaos::FAABB3 ToAabb(const Chaos::FAABBVectorized& AabbVectorized) const;
+};
diff --git a/Engine/Source/Runtime/Experimental/Chaos/Private/Chaos/TriangleMeshImplicitObject.cpp b/Engine/Source/Runtime/Experimental/Chaos/Private/Chaos/TriangleMeshImplicitObject.cpp
--- a/Engine/Source/Runtime/Experimental/Chaos/Private/Chaos/TriangleMeshImplicitObject.cpp
+++ b/Engine/Source/Runtime/Experimental/Chaos/Private/Chaos/TriangleMeshImplicitObject.cpp
@@ -887,6 +887,11 @@ void FTriangleMeshImplicitObject::CollectTriangles(const FAABB3& MeshQueryBounds
 	Collector.EndCollect();
 }
 
+const FTrimeshBVH& FTriangleMeshImplicitObject::GetBVH() const
+{
+	return FastBVH;
+}
+
 template <typename QueryGeomType>
 bool FTriangleMeshImplicitObject::GJKContactPointImp(const QueryGeomType& QueryGeom, const FRigidTransform3& QueryTM, const FReal WorldThickness, FVec3& Location, FVec3& Normal, FReal& OutContactPhi, int32& OutFaceIndex, FVec3 TriMeshScale) const
 {
diff --git a/Engine/Source/Runtime/Experimental/Chaos/Public/Chaos/TriangleMeshImplicitObject.h b/Engine/Source/Runtime/Experimental/Chaos/Public/Chaos/TriangleMeshImplicitObject.h
--- a/Engine/Source/Runtime/Experimental/Chaos/Public/Chaos/TriangleMeshImplicitObject.h
+++ b/Engine/Source/Runtime/Experimental/Chaos/Public/Chaos/TriangleMeshImplicitObject.h
@@ -800,6 +800,9 @@ namespace Chaos
 		// Internal: do not use - this API will change as we optimize mesh collision
 		void CollectTriangles(const FAABB3& MeshQueryBounds, const FRigidTransform3& MeshToObjectTransform, const FAABB3& ObjectBounds, Private::FMeshContactGenerator& Collector) const;
 
+		// Internal: Only for CVD
+		UE_INTERNAL CHAOS_API const FTrimeshBVH& GetBVH() const;
+
 	private:
 		using BVHType = TAABBTree<int32, TAABBTreeLeafArray<int32, /*bComputeBounds=*/ false, FRealSingle>, /*bMutable=*/false, FRealSingle>;
 		CHAOS_API void RebuildFastBVHFromTree(const BVHType& BVH);
