[ChaosVD] Making meshes created to represent landscape collision, use Instanced Static Mesh components instead of Static Mesh Components.

This improves rendering performance when you are loading multiple CVD files from different build types but with the same base content (like a server + multiple clients, or multiples server (DSTM)).

The limitation of using static mesh components for landscape collisions comes from the time where CVD loaded the geometry for the entire scene, and given the complexity of a mesh representing a landscape, we neded LODs support.

Since a few months ago, CVD only loads the geometry within a configurable range from the editor's camera, so LODs are not needed anymore. (and if at some point we need them, we could look into adding support ofr HISMC).

diff --git a/Engine/Plugins/ChaosVD/Source/ChaosVD/Private/ChaosVDSceneParticle.cpp b/Engine/Plugins/ChaosVD/Source/ChaosVD/Private/ChaosVDSceneParticle.cpp
--- a/Engine/Plugins/ChaosVD/Source/ChaosVD/Private/ChaosVDSceneParticle.cpp
+++ b/Engine/Plugins/ChaosVD/Source/ChaosVD/Private/ChaosVDSceneParticle.cpp
@@ -32,6 +32,12 @@ namespace Chaos::VisualDebugger::Cvars
 		TEXT("p.Chaos.VD.Tool.ForceStaticMeshComponentUse"),
 		bForceStaticMeshComponentUse,
 		TEXT("If true, static mesh components will be used instead of Instanced Static mesh components when recreating the geometry for each particle"));
+
+	static bool bUseInstancedStaticMeshForLandscape = true;
+	static FAutoConsoleVariableRef CVarChaosVDbUseInstancedStaticMeshForLandscape(
+		TEXT("p.Chaos.VD.Tool.UseInstancedStaticMeshForLandscape"),
+		bForceStaticMeshComponentUse,
+		TEXT("If true, instanced static mesh components will be used instead of static mesh components when recreating the geometry for particles from Landscapes"));
 }
 
 FChaosVDSceneParticle::~FChaosVDSceneParticle()
@@ -434,8 +440,14 @@ void FChaosVDSceneParticle::UpdateGeometry(const Chaos::FConstImplicitObjectPtr&
 	TArray<TSharedPtr<FChaosVDExtractedGeometryDataHandle>> OutExtractedGeometryDataHandles;
 	OutExtractedGeometryDataHandles.Reserve(ObjectsToGenerateNum);
 
-	// Heightfields need to be created as Static meshes and use normal Static Mesh components because we need LODs for them due to their high triangle count
-	const bool bHasToUseStaticMeshComponent = Chaos::VisualDebugger::Cvars::bForceStaticMeshComponentUse || FChaosVDGeometryBuilder::DoesImplicitContainType(InImplicitObject, Chaos::ImplicitObjectType::HeightField);
+	bool bHasToUseStaticMeshComponent = Chaos::VisualDebugger::Cvars::bForceStaticMeshComponentUse;
+
+	if (!bHasToUseStaticMeshComponent && !Chaos::VisualDebugger::Cvars::bUseInstancedStaticMeshForLandscape)
+	{
+		// Heightfields need to be created as Static meshes and use normal Static Mesh components because we need LODs for them due to their high triangle count
+		bHasToUseStaticMeshComponent = FChaosVDGeometryBuilder::DoesImplicitContainType(InImplicitObject, Chaos::ImplicitObjectType::HeightField);
+	}
+
 	constexpr int32 LODsToGenerateNum = 3;
 	constexpr int32 LODsToGenerateNumForInstancedStaticMesh = 0;
 
