[ChaosVD] Hiding auto-generated tracks for game thread data, in multi-source mode.

When you load multiple CVD files in CVD into a single instance, we need to auto-generate a track to hold what would normally go to the main game thread data track.

These auto-generated tracks have no real use currently, and they just add noise if you are working with multiple clients or servers.

This CL makes it so we can completely hide these tracks. I am not excluding the widget generation for these completely yet, as I would like to explore the possibility of allowing showing them as an option in the future. They might be usable once we start recording Game Thread particle data.

diff --git a/Engine/Plugins/ChaosVD/Source/ChaosVD/Private/ChaosVDPlaybackController.cpp b/Engine/Plugins/ChaosVD/Source/ChaosVD/Private/ChaosVDPlaybackController.cpp
--- a/Engine/Plugins/ChaosVD/Source/ChaosVD/Private/ChaosVDPlaybackController.cpp
+++ b/Engine/Plugins/ChaosVD/Source/ChaosVD/Private/ChaosVDPlaybackController.cpp
@@ -1275,8 +1275,8 @@ void FChaosVDPlaybackController::UpdateSolverTracksData()
 		if (bIsNewTrack)
 		{
 			SolverTrackInfo = CreateTrackInfo();
-			// Solver frame tracks only holding GT Data, are not able to change visibility states
-			SolverTrackInfo->bSupportsVisibilityChange = !SolverIDPair.Value.IsEmpty() ? !EnumHasAnyFlags(SolverIDPair.Value[0].GetAttributes(), EChaosVDSolverFrameAttributes::HasGTDataToReRoute) : true;
+			// Solver frame tracks only holding GT Data, these should not be shown in the timeline controls UI
+			SolverTrackInfo->bCanShowTrackControls = !SolverIDPair.Value.IsEmpty() ? !EnumHasAnyFlags(SolverIDPair.Value[0].GetAttributes(), EChaosVDSolverFrameAttributes::HasGTDataToReRoute) : true;
 			SolverTrackInfo->CurrentFrame = 0;
 			SolverTrackInfo->CurrentStage = 0;
 		}
diff --git a/Engine/Plugins/ChaosVD/Source/ChaosVD/Private/Widgets/SChaosVDSolverTracks.cpp b/Engine/Plugins/ChaosVD/Source/ChaosVD/Private/Widgets/SChaosVDSolverTracks.cpp
--- a/Engine/Plugins/ChaosVD/Source/ChaosVD/Private/Widgets/SChaosVDSolverTracks.cpp
+++ b/Engine/Plugins/ChaosVD/Source/ChaosVD/Private/Widgets/SChaosVDSolverTracks.cpp
@@ -161,6 +161,11 @@ TSharedRef<ITableRow> SChaosVDSolverTracks::MakeSolverTrackControlsFromTrackInfo
 	if (TrackInfo)
 	{
 		RowWidget = SNew(SVerticalBox)
+				.Visibility_Lambda([WeakTrackInfo = TrackInfo.ToWeakPtr()]()
+				{
+					TSharedPtr<const FChaosVDTrackInfo> TrackInfoPtr = StaticCastSharedPtr<const FChaosVDTrackInfo>(WeakTrackInfo.Pin());
+					return TrackInfoPtr && TrackInfoPtr->bCanShowTrackControls ?  EVisibility::Visible : EVisibility::Collapsed;
+				})
 				+SVerticalBox::Slot()
 				.AutoHeight()
 				.Padding(10.0f, 4.0f, 10.0f, 0.0f)
diff --git a/Engine/Plugins/ChaosVD/Source/ChaosVD/Public/ChaosVDPlaybackController.h b/Engine/Plugins/ChaosVD/Source/ChaosVD/Public/ChaosVDPlaybackController.h
--- a/Engine/Plugins/ChaosVD/Source/ChaosVD/Public/ChaosVDPlaybackController.h
+++ b/Engine/Plugins/ChaosVD/Source/ChaosVD/Public/ChaosVDPlaybackController.h
@@ -51,6 +51,7 @@ struct FChaosVDTrackInfo
 	bool bIsServer = false;
 	bool bHasNetworkSyncData = false;
 	bool bSupportsVisibilityChange = true;
+	bool bCanShowTrackControls = true;
 
 	bool operator==(const FChaosVDTrackInfo& Other) const;
 
