[ChaosVD] Fix for an issue where in complex map, if you try to start recording & connect to a remote instance using the one click workflows, the recording will start in the remote instance but CVD will fail to auto-connect.

The issue is that if the remote instance stalls while it performs the initial geometry serialization step for the whole world, or the trace system takes a while to respond to the start recording command, CVD might determined the session has expired as it didn't get a response after multiple pings.

When a session is marked as expired, it is removed and therefore any further attempts to connect of using the now expired session object will fail (no matter how many attempts you try).

The fix this CL does is:

- Pass the Session GUID to the retry system instead of the session object
- Increase the default retries to 14 from 4
- Add a way to configure how many retries we should perform in CVD's settings UI.

A further improvement to this will be to implement exponential backoff in the retry system, and a way for the user to cancel them in the UI.

diff --git a/Engine/Plugins/ChaosVD/Source/ChaosVD/Private/Settings/ChaosVDGeneralSettings.h b/Engine/Plugins/ChaosVD/Source/ChaosVD/Private/Settings/ChaosVDGeneralSettings.h
--- a/Engine/Plugins/ChaosVD/Source/ChaosVD/Private/Settings/ChaosVDGeneralSettings.h
+++ b/Engine/Plugins/ChaosVD/Source/ChaosVD/Private/Settings/ChaosVDGeneralSettings.h
@@ -22,7 +22,11 @@ public:
 	UPROPERTY(Config, EditAnywhere, Category = Streaming)
 	float StreamingBoxExtentSize = 10000.0f;
 
-	// If set to true CVD will process any updates to the streaming accel structure in worker threads, in between streaming updates
+	/** If set to true CVD will process any updates to the streaming accel structure in worker threads, in between streaming updates */
 	UPROPERTY(Config, EditAnywhere, Category = Streaming)
 	bool bProcessPendingOperationsQueueInWorkerThread = true;
+
+	/** How many times CVD will attempt to connect to a live trace or load from file session if the first attempt failed */
+	UPROPERTY(Config, EditAnywhere, Category = Connection)
+	int32 MaxConnectionRetries = 14;
 };
diff --git a/Engine/Plugins/ChaosVD/Source/ChaosVD/Private/Widgets/SChaosVDRecordingControls.cpp b/Engine/Plugins/ChaosVD/Source/ChaosVD/Private/Widgets/SChaosVDRecordingControls.cpp
--- a/Engine/Plugins/ChaosVD/Source/ChaosVD/Private/Widgets/SChaosVDRecordingControls.cpp
+++ b/Engine/Plugins/ChaosVD/Source/ChaosVD/Private/Widgets/SChaosVDRecordingControls.cpp
@@ -6,6 +6,7 @@
 #include "ChaosVDEngine.h"
 #include "ChaosVDModule.h"
 #include "ChaosVDRecordingDetails.h"
+#include "ChaosVDSettingsManager.h"
 #include "ChaosVDStyle.h"
 #include "ChaosVisualDebugger/ChaosVisualDebuggerTrace.h"
 #include "Editor.h"
@@ -20,6 +21,7 @@
 #include "Chaos/ChaosVDEngineEditorBridge.h"
 #include "Framework/MultiBox/MultiBoxBuilder.h"
 #include "Misc/ScopedSlowTask.h"
+#include "Settings/ChaosVDGeneralSettings.h"
 #include "Widgets/Images/SImage.h"
 #include "Widgets/Input/SButton.h"
 #include "Widgets/SChaosVDMainTab.h"
@@ -29,6 +31,7 @@
 
 #define LOCTEXT_NAMESPACE "ChaosVisualDebugger"
 
+class UChaosVDGeneralSettings;
 const FName SChaosVDRecordingControls::RecordingControlsToolbarName = FName("ChaosVD.MainToolBar.RecordingControls");
 
 void SChaosVDRecordingControls::Construct(const FArguments& InArgs, const TSharedRef<SChaosVDMainTab>& InMainTabSharedRef)
@@ -594,6 +597,17 @@ void SChaosVDRecordingControls::SetSessionRecordingState(bool bIsRecording, ECha
 	{
 		TArray<FString, TInlineAllocator<1>> RecordingArgs;
 
+		int32 RemainingRetries = 4;
+
+		if (UChaosVDGeneralSettings* Settings = FChaosVDSettingsManager::Get().GetSettingsObject<UChaosVDGeneralSettings>())
+		{
+			RemainingRetries = Settings->MaxConnectionRetries;
+		}
+		else
+		{
+			UE_LOG(LogChaosVDEditor, Warning, TEXT("Failed to obtain setting object. Setting the retries attempts to connect to a session to 4 as a fallback."))
+		}
+
 		if (RecordingMode == EChaosVDRecordingMode::Live)
 		{
 			FChaosVDStartRecordingCommandMessage RecordingParams;
@@ -611,10 +625,10 @@ void SChaosVDRecordingControls::SetSessionRecordingState(bool bIsRecording, ECha
 
 			// Once the start recording command is issue, we can try to connect to the created session (if everything went well).
 			// If it didn't go well, this will take care of update the UI to notify the user
-			constexpr int32 RemainingRetries = 4;
-			ExecuteAsyncConnectionAttemptTaskWithRetry(RemainingRetries, [WeakThis = AsWeak(), SessionInfoAsWeak = SessionInfoRef.ToWeakPtr()]()
+			ExecuteAsyncConnectionAttemptTaskWithRetry(RemainingRetries, [WeakThis = AsWeak(), SessionInstanceId = SessionInfoRef->InstanceId]()
 			{
-				TSharedPtr<FChaosVDSessionInfo> SessionInfoPtr = SessionInfoAsWeak.Pin();
+				TSharedPtr<FChaosVDRemoteSessionsManager> RemoteSessionManagerPtr = FChaosVDEngineEditorBridge::Get().GetRemoteSessionsManager();
+				TSharedPtr<FChaosVDSessionInfo> SessionInfoPtr = RemoteSessionManagerPtr ? RemoteSessionManagerPtr->GetSessionInfo(SessionInstanceId).Pin() : nullptr;
 
 				if (!SessionInfoPtr)
 				{
@@ -640,7 +654,6 @@ void SChaosVDRecordingControls::SetSessionRecordingState(bool bIsRecording, ECha
 
 			// Once the start recording command is issued, we need to check if the recording started, which might take a few frames.
 			// This will take care of retrying, waiting and update the UI to notify the user if needed.
-			constexpr int32 RemainingRetries = 4;
 			ExecuteAsyncConnectionAttemptTaskWithRetry(RemainingRetries, [WeakThis = AsWeak()]()
 			{
 				TSharedPtr<SChaosVDRecordingControls> Controls = StaticCastSharedPtr<SChaosVDRecordingControls>(WeakThis.Pin());
