[ChaosVD] Fix for ensures that might trigger if you playback back and forth a recording that contains particles with no geometry.

This is something that is possible with cluster unions

diff --git a/Engine/Plugins/ChaosVD/Source/ChaosVD/Private/ChaosVDSceneParticle.cpp b/Engine/Plugins/ChaosVD/Source/ChaosVD/Private/ChaosVDSceneParticle.cpp
--- a/Engine/Plugins/ChaosVD/Source/ChaosVD/Private/ChaosVDSceneParticle.cpp
+++ b/Engine/Plugins/ChaosVD/Source/ChaosVD/Private/ChaosVDSceneParticle.cpp
@@ -83,7 +83,11 @@ void FChaosVDSceneParticle::RemoveAllGeometry()
 
 FBox FChaosVDSceneParticle::GetStreamingBounds() const
 {
-	ensure(CachedBounds.IsValid);
+	if (!CachedBounds.IsValid)
+	{
+		// The only valid case to not have bounds, is if we don't have geometry
+		ensure(!CurrentRootGeometry);
+	}
 
 	return CachedBounds;
 }
@@ -211,7 +215,22 @@ void FChaosVDSceneParticle::PreUpdateFromRecordedParticleData(const TSharedPtr<F
 	{
 		if (const TSharedPtr<FChaosVDScene>& ScenePtr = SceneWeakPtr.Pin())
 		{
-			CurrentRootGeometry = ScenePtr->GetUpdatedGeometry(ParticleDataPtr->GeometryHash);
+			uint32 GeometryHash = ParticleDataPtr->GeometryHash;
+
+			if (GeometryHash != 0)
+			{
+				CurrentRootGeometry = ScenePtr->GetUpdatedGeometry(GeometryHash);
+
+				if (!ensure(CurrentRootGeometry))
+				{
+					// We intentionally let the code continue, as passing down a null geometry  will take care of removing any existing mesh representation for this particle
+					UE_LOG(LogChaosVDEditor, Warning, TEXT("Failed to find Geometry for Particle ID [%d] | Geometry Hash [%u] | Debug Name [%s]"), ParticleDataPtr->ParticleIndex, GeometryHash, *ParticleDataPtr->DebugName);
+				}
+			}
+			else
+			{
+				CurrentRootGeometry = nullptr;
+			}
 		}
 
 		EnumAddFlags(DirtyFlags, EChaosVDSceneParticleDirtyFlags::Geometry);
@@ -405,6 +424,8 @@ void FChaosVDSceneParticle::UpdateGeometry(const Chaos::FConstImplicitObjectPtr&
 
 	if (!InImplicitObject.IsValid())
 	{
+		RemoveAllGeometry();
+		EnumRemoveFlags(DirtyFlags, EChaosVDSceneParticleDirtyFlags::Geometry);
 		return;
 	}
 
@@ -524,7 +545,7 @@ void FChaosVDSceneParticle::UpdateGeometry(const Chaos::FConstImplicitObjectPtr&
 
 void FChaosVDSceneParticle::CalculateAndCacheBounds() const
 {
-	if (!ensure(CurrentRootGeometry))
+	if (CurrentRootGeometry == nullptr)
 	{
 		return;
 	}
diff --git a/Engine/Plugins/ChaosVD/Source/ChaosVD/Private/ChaosVDSceneStreaming.cpp b/Engine/Plugins/ChaosVD/Source/ChaosVD/Private/ChaosVDSceneStreaming.cpp
--- a/Engine/Plugins/ChaosVD/Source/ChaosVD/Private/ChaosVDSceneStreaming.cpp
+++ b/Engine/Plugins/ChaosVD/Source/ChaosVD/Private/ChaosVDSceneStreaming.cpp
@@ -111,7 +111,12 @@ void FChaosVDSceneStreaming::UpdateStreamingSourceLocation(const FVector& NewLoc
 
 bool FChaosVDSceneStreaming::IsInStreamingRange(const FBox& Bounds) const
 {
-	return bIsStreamingSystemEnabled ? SteamingViewBox.Intersect(Bounds) : true;
+	if (bIsStreamingSystemEnabled)
+	{
+		return Bounds.IsValid ? SteamingViewBox.Intersect(Bounds) : false;
+	}
+
+	return true;
 }
 
 void FChaosVDSceneStreaming::UpdateStreamingStateForObject(const TSharedRef<FChaosVDBaseSceneObject>& InSceneObject) const
