[ChaosVD] Allow to users enable scene outliner updates during playback.

Currently this is always disabled by default, because updating the outliner has a significant performance update, but in small scenes the impact is acceptable and having the updates enabled is a better UX.

We will still keep it disabled by default for now, but we will now allow the user to change it.

Long term, the plan is to have scene outliner updates enabled by default, and detect if such updates are going over budget. If that is detected we will prompt the use to disable updates or just accept the degraded performance.

diff --git a/Engine/Plugins/ChaosVD/Source/ChaosVD/Private/ChaosVDWorldOutlinerMode.cpp b/Engine/Plugins/ChaosVD/Source/ChaosVD/Private/ChaosVDWorldOutlinerMode.cpp
--- a/Engine/Plugins/ChaosVD/Source/ChaosVD/Private/ChaosVDWorldOutlinerMode.cpp
+++ b/Engine/Plugins/ChaosVD/Source/ChaosVD/Private/ChaosVDWorldOutlinerMode.cpp
@@ -7,6 +7,7 @@
 #include "ChaosVDSceneParticle.h"
 #include "ChaosVDPlaybackController.h"
 #include "ChaosVDScene.h"
+#include "ChaosVDSettingsManager.h"
 #include "TedsOutlinerItem.h"
 #include "Elements/Actor/ActorElementData.h"
 #include "Elements/Columns/TypedElementCompatibilityColumns.h"
@@ -15,6 +16,7 @@
 #include "Elements/Framework/EngineElementsLibrary.h"
 #include "Elements/Framework/TypedElementSelectionSet.h"
 #include "Elements/Interfaces/TypedElementDataStorageCompatibilityInterface.h"
+#include "Settings/ChaosVDGeneralSettings.h"
 #include "TEDS/ChaosVDParticleEditorDataFactory.h"
 #include "TEDS/ChaosVDSelectionInterface.h"
 #include "TEDS/ChaosVDStructTypedElementData.h"
@@ -245,10 +247,23 @@ bool FChaosVDWorldOutlinerMode::CanPopulate() const
 {
 	if (TSharedPtr<FChaosVDPlaybackController> PlaybackControllerPtr = PlaybackController.Pin())
 	{
+		const bool bIsContinuousPlayback = PlaybackControllerPtr->IsScrubbingTimeline() || PlaybackControllerPtr->IsPlaying();
+
+		if (!bIsContinuousPlayback)
+		{
+			return true;
+		}
+
 		// Updating the scene outliner during playback it is very expensive and can tank framerate,
 		// as it need to re-build the hierarchy when things are added in ad removed. So if we are playing we want to pause any updates to the outliner
-		return !PlaybackControllerPtr->IsPlaying() && !PlaybackControllerPtr->IsScrubbingTimeline();
+		if (const UChaosVDGeneralSettings* GeneralSettings = FChaosVDSettingsManager::Get().GetSettingsObject<UChaosVDGeneralSettings>())
+		{
+			return GeneralSettings->bUpdateSceneOutlinerDuringPlayback && !PlaybackControllerPtr->IsScrubbingTimeline();
+		}
+
+		return false;	
 	}
+
 	return true;
 }
 
diff --git a/Engine/Plugins/ChaosVD/Source/ChaosVD/Private/Settings/ChaosVDGeneralSettings.h b/Engine/Plugins/ChaosVD/Source/ChaosVD/Private/Settings/ChaosVDGeneralSettings.h
--- a/Engine/Plugins/ChaosVD/Source/ChaosVD/Private/Settings/ChaosVDGeneralSettings.h
+++ b/Engine/Plugins/ChaosVD/Source/ChaosVD/Private/Settings/ChaosVDGeneralSettings.h
@@ -15,17 +15,22 @@ class CHAOSVD_API UChaosVDGeneralSettings : public UChaosVDSettingsObjectBase
 
 public:
 	/** If true, CVD will only load collision geometry that is visible */
-	UPROPERTY(Config, EditAnywhere, Category = Streaming)
+	UPROPERTY(Config, EditAnywhere, Category = "Performance | Streaming")
 	bool bStreamingSystemEnabled = true;
 
 	/** Extent size of the box used for calculate what should be streamed in */
-	UPROPERTY(Config, EditAnywhere, Category = Streaming)
+	UPROPERTY(Config, EditAnywhere, Category = "Performance | Streaming")
 	float StreamingBoxExtentSize = 10000.0f;
 
 	/** If set to true CVD will process any updates to the streaming accel structure in worker threads, in between streaming updates */
-	UPROPERTY(Config, EditAnywhere, Category = Streaming)
+	UPROPERTY(Config, EditAnywhere, Category = "Performance | Streaming")
 	bool bProcessPendingOperationsQueueInWorkerThread = true;
 
+	/** If set to true CVD will keep the scene outliner up to date as the recording is played. If during the recording a
+	 * significant amount of objects are loaded/unloaded, the performance impact will be significant enough to degrade the playback stability. */
+	UPROPERTY(Config, EditAnywhere, Category = "Performance | Scene Outliner")
+	bool bUpdateSceneOutlinerDuringPlayback = false;
+
 	/** How many times CVD will attempt to connect to a live trace or load from file session if the first attempt failed */
 	UPROPERTY(Config, EditAnywhere, Category = Connection)
 	int32 MaxConnectionRetries = 14;
