[ChaosVD] Fix for an issue where generic debug data recorded without an associated solver ID, fails to load in multi-source mode.

The issue is that when you record debug draw data without a solver id, that goes into the game thread data track. But the CVD component in charge of re-rounting that data to the correct components in multi-source mode, was only evaluating solver data actors.

diff --git a/Engine/Plugins/ChaosVD/Source/ChaosVD/Private/Components/ChaosVDAdditionalGTDataRouterComponent.cpp b/Engine/Plugins/ChaosVD/Source/ChaosVD/Private/Components/ChaosVDAdditionalGTDataRouterComponent.cpp
--- a/Engine/Plugins/ChaosVD/Source/ChaosVD/Private/Components/ChaosVDAdditionalGTDataRouterComponent.cpp
+++ b/Engine/Plugins/ChaosVD/Source/ChaosVD/Private/Components/ChaosVDAdditionalGTDataRouterComponent.cpp
@@ -19,9 +19,9 @@ void UChaosVDAdditionalGTDataRouterComponent::UpdateFromSolverFrameData(const FC
 		{
 			if (TSharedPtr<FChaosVDGameFrameDataWrapper> GTFrameDataWrapper = InSolverFrameData.GetCustomData().GetData<FChaosVDGameFrameDataWrapper>())
 			{
-				for (const TPair<int32, AChaosVDSolverInfoActor*>& SolverInfoWithID : CVDScene->GetSolverInfoActorsMap())
+				for (const TObjectPtr<AChaosVDDataContainerBaseActor>& DataContainerActor : CVDScene->GetDataContainerActorsView())
 				{
-					SolverInfoWithID.Value->UpdateFromNewGameFrameData(*GTFrameDataWrapper->FrameData);
+					DataContainerActor->UpdateFromNewGameFrameData(*GTFrameDataWrapper->FrameData);
 				}
 			}			
 		}
diff --git a/Engine/Plugins/ChaosVD/Source/ChaosVDBuiltInExtensions/Private/GenericDebugDraw/Components/ChaosVDGenericDebugDrawDataComponent.cpp b/Engine/Plugins/ChaosVD/Source/ChaosVDBuiltInExtensions/Private/GenericDebugDraw/Components/ChaosVDGenericDebugDrawDataComponent.cpp
--- a/Engine/Plugins/ChaosVD/Source/ChaosVDBuiltInExtensions/Private/GenericDebugDraw/Components/ChaosVDGenericDebugDrawDataComponent.cpp
+++ b/Engine/Plugins/ChaosVD/Source/ChaosVDBuiltInExtensions/Private/GenericDebugDraw/Components/ChaosVDGenericDebugDrawDataComponent.cpp
@@ -20,6 +20,10 @@ void UChaosVDGenericDebugDrawDataComponent::UpdateFromNewGameFrameData(const FCh
 		CopyDataFromSourceMap<FChaosVDDebugDrawSphereDataWrapper>(DebugDrawShapeData->RecordedDebugDrawSpheresBySolverID, DebugDrawSpheres, SolverID);
 		CopyDataFromSourceMap<FChaosVDDebugDrawImplicitObjectDataWrapper>(DebugDrawShapeData->RecordedDebugDrawImplicitObjectsBySolverID, DebugDrawImplicitObjects, SolverID);
 	}
+	else
+	{
+		ClearData();
+	}
 }
 
 void UChaosVDGenericDebugDrawDataComponent::ClearData()
