[ChaosVD] Added options to start recordings to multiple targets of specific types, like "All Servers" and "All Clients".

Also, now there is a "Custom Selection" target where you can select which specific targets you want in a multi-session recording action.

Lastly, fixed a bug where if you try to connect to a multi-target item twice in a row, the previously loaded trace sessions were not unloaded

diff --git a/Engine/Plugins/ChaosVD/Source/ChaosVD/Private/ChaosVDEngine.cpp b/Engine/Plugins/ChaosVD/Source/ChaosVD/Private/ChaosVDEngine.cpp
--- a/Engine/Plugins/ChaosVD/Source/ChaosVD/Private/ChaosVDEngine.cpp
+++ b/Engine/Plugins/ChaosVD/Source/ChaosVD/Private/ChaosVDEngine.cpp
@@ -63,6 +63,11 @@ void FChaosVDEngine::CloseActiveTraceSessions()
 	}
 
 	CurrentSessionDescriptors.Reset();
+
+	if (PlaybackController)
+	{
+		PlaybackController->UnloadCurrentRecording(EChaosVDUnloadRecordingFlags::BroadcastChanges);
+	}
 }
 
 void FChaosVDEngine::StopActiveTraceSessions()
diff --git a/Engine/Plugins/ChaosVD/Source/ChaosVD/Private/ChaosVDPlaybackController.cpp b/Engine/Plugins/ChaosVD/Source/ChaosVD/Private/ChaosVDPlaybackController.cpp
--- a/Engine/Plugins/ChaosVD/Source/ChaosVD/Private/ChaosVDPlaybackController.cpp
+++ b/Engine/Plugins/ChaosVD/Source/ChaosVD/Private/ChaosVDPlaybackController.cpp
@@ -323,7 +323,7 @@ void FChaosVDPlaybackController::GoToRecordedGameFrame_AssumesLocked(const int32
 
 void FChaosVDPlaybackController::GoToTrackFrame(FGuid InstigatorID, EChaosVDTrackType TrackType, int32 InTrackID, int32 FrameNumber, int32 StageNumber)
 {
-	if (!ensure(LoadedRecording.IsValid()))
+	if (!LoadedRecording.IsValid())
 	{
 		return;
 	}
@@ -569,7 +569,7 @@ int32 FChaosVDPlaybackController::GetTrackCurrentStage(EChaosVDTrackType TrackTy
 
 int32 FChaosVDPlaybackController::GetTrackLastStageAtFrame(EChaosVDTrackType TrackType, int32 InTrackID, int32 InFrameNumber) const
 {
-	if (!ensure(LoadedRecording.IsValid()))
+	if (!LoadedRecording.IsValid())
 	{
 		return INDEX_NONE;
 	}
@@ -648,7 +648,7 @@ void FChaosVDPlaybackController::GetMutableAvailableTracks(EChaosVDTrackType Tra
 
 void FChaosVDPlaybackController::GetAvailableTrackInfosAtTrackFrame(EChaosVDTrackType TrackTypeToFind, const TSharedRef<const FChaosVDTrackInfo>& InFromTrack, TArray<TSharedPtr<const FChaosVDTrackInfo>>& OutTrackInfo)
 {
-	if (!ensure(LoadedRecording.IsValid()))
+	if (!LoadedRecording.IsValid())
 	{
 		return;
 	}
@@ -1071,7 +1071,7 @@ void FChaosVDPlaybackController::GetTracksByType(EChaosVDTrackType Type, TArray<
 
 void FChaosVDPlaybackController::SyncTracks(const TSharedRef<const FChaosVDTrackInfo>& FromTrack, EChaosVDSyncTimelinesMode TrackSyncMode)
 {
-	if (!ensure(LoadedRecording.IsValid()))
+	if (!LoadedRecording.IsValid())
 	{
 		return;
 	}
diff --git a/Engine/Plugins/ChaosVD/Source/ChaosVD/Private/Widgets/SChaosVDMainTab.cpp b/Engine/Plugins/ChaosVD/Source/ChaosVD/Private/Widgets/SChaosVDMainTab.cpp
--- a/Engine/Plugins/ChaosVD/Source/ChaosVD/Private/Widgets/SChaosVDMainTab.cpp
+++ b/Engine/Plugins/ChaosVD/Source/ChaosVD/Private/Widgets/SChaosVDMainTab.cpp
@@ -706,7 +706,7 @@ void SChaosVDMainTab::BrowseChaosVDRecordingFromFolder(FStringView FolderPath, E
 	}
 
 	// Ideally, we should not need to do this, but the UI to support multi source is not robust yet, The UI improvement task to avoid this scenario is planned as UE-197418
-	if (const bool bShouldOverrideMode = OutOpenFilenames.Num() > 0 && LoadingMode == EChaosVDLoadRecordedDataMode::SingleSource)
+	if (const bool bShouldOverrideMode = OutOpenFilenames.Num() > 1 && LoadingMode == EChaosVDLoadRecordedDataMode::SingleSource)
 	{
 		GetChaosVDEngineInstance()->CloseActiveTraceSessions();
 		LoadingMode = EChaosVDLoadRecordedDataMode::MultiSource;
diff --git a/Engine/Plugins/ChaosVD/Source/ChaosVD/Private/Widgets/SChaosVDRecordingControls.cpp b/Engine/Plugins/ChaosVD/Source/ChaosVD/Private/Widgets/SChaosVDRecordingControls.cpp
--- a/Engine/Plugins/ChaosVD/Source/ChaosVD/Private/Widgets/SChaosVDRecordingControls.cpp
+++ b/Engine/Plugins/ChaosVD/Source/ChaosVD/Private/Widgets/SChaosVDRecordingControls.cpp
@@ -6,6 +6,7 @@
 #include "ChaosVDEngine.h"
 #include "ChaosVDModule.h"
 #include "ChaosVDRecordingDetails.h"
+#include "ChaosVDScene.h"
 #include "ChaosVDSettingsManager.h"
 #include "ChaosVDStyle.h"
 #include "ChaosVisualDebugger/ChaosVisualDebuggerTrace.h"
@@ -181,15 +182,43 @@ TSharedRef<SWidget> SChaosVDRecordingControls::GenerateTargetSessionDropdown()
 			AllRemoteTargetsLabel,
 			LOCTEXT("MultiRemoteTargetTooltip", "Select this to act on all remote targets"),
 			GetIconForSession(FChaosVDRemoteSessionsManager::AllRemoteSessionsWrapperGUID),
-			FUIAction(FExecuteAction::CreateSP(this, &SChaosVDRecordingControls::SelectTargetSession, FChaosVDRemoteSessionsManager::AllRemoteSessionsWrapperGUID), EUIActionRepeatMode::RepeatDisabled));
+			FUIAction(FExecuteAction::CreateSP(this, &SChaosVDRecordingControls::SelectTargetSession, FChaosVDRemoteSessionsManager::AllRemoteSessionsWrapperGUID),
+			FCanExecuteAction::CreateSP(this, &SChaosVDRecordingControls::CanSelectMultiSessionTarget, FChaosVDRemoteSessionsManager::AllRemoteSessionsWrapperGUID),
+			EUIActionRepeatMode::RepeatDisabled));
+
+	FText AllRemoteServersTargetsLabel = LOCTEXT("AllRemoteServersOption", "All Remote Servers");
+	MenuBuilder.AddMenuEntry(
+			AllRemoteServersTargetsLabel,
+			LOCTEXT("MultiRemoteServerTargetTooltip", "Select this to act on all remote server targets"),
+			GetIconForSession(FChaosVDRemoteSessionsManager::AllRemoteServersWrapperGUID),
+			FUIAction(FExecuteAction::CreateSP(this, &SChaosVDRecordingControls::SelectTargetSession, FChaosVDRemoteSessionsManager::AllRemoteServersWrapperGUID),
+			FCanExecuteAction::CreateSP(this, &SChaosVDRecordingControls::CanSelectMultiSessionTarget, FChaosVDRemoteSessionsManager::AllRemoteServersWrapperGUID),
+			EUIActionRepeatMode::RepeatDisabled));
+
+	FText AllRemoteClientsTargetsLabel = LOCTEXT("AllRemoteClientsOption", "All Remote Clients");
+	MenuBuilder.AddMenuEntry(
+			AllRemoteClientsTargetsLabel,
+			LOCTEXT("MultiRemoteClientTargetTooltip", "Select this to act on all remote client targets"),
+			GetIconForSession(FChaosVDRemoteSessionsManager::AllRemoteClientsWrapperGUID),
+			FUIAction(FExecuteAction::CreateSP(this, &SChaosVDRecordingControls::SelectTargetSession, FChaosVDRemoteSessionsManager::AllRemoteClientsWrapperGUID),
+			FCanExecuteAction::CreateSP(this, &SChaosVDRecordingControls::CanSelectMultiSessionTarget, FChaosVDRemoteSessionsManager::AllRemoteClientsWrapperGUID),
+			EUIActionRepeatMode::RepeatDisabled));
 
 	FText AllTargets = LOCTEXT("AllTargetsOption", "All");
 	MenuBuilder.AddMenuEntry(
 		AllTargets,
-		LOCTEXT("MultiTargetTooltip", "Select this to act on all targets, both Local and Remote"),
+		LOCTEXT("MultiAllTargetTooltip", "Select this to act on all targets, both Local and Remote"),
 		GetIconForSession(FChaosVDRemoteSessionsManager::AllSessionsWrapperGUID),
 		FUIAction(FExecuteAction::CreateSP(this, &SChaosVDRecordingControls::SelectTargetSession, FChaosVDRemoteSessionsManager::AllSessionsWrapperGUID), EUIActionRepeatMode::RepeatDisabled));
 
+	MenuBuilder.AddMenuSeparator();
+
+	FText CustomTargets = LOCTEXT("CustomTargetsOption", "Custom Selection");
+	MenuBuilder.AddSubMenu(CustomTargets,
+		LOCTEXT("MultiCustomTargetTooltip", "Select this to act on the specific targets you selected"),
+		FNewMenuDelegate::CreateSP(this, &SChaosVDRecordingControls::GenerateCustomTargetsMenu)
+		);
+
 	MenuBuilder.EndSection();
 
 	return MenuBuilder.MakeWidget();
@@ -346,6 +375,103 @@ void SChaosVDRecordingControls::SelectTargetSession(FGuid SessionId)
 	CurrentSelectedSessionId = SessionId;
 }
 
+void SChaosVDRecordingControls::GenerateCustomTargetsMenu(FMenuBuilder& MenuBuilder)
+{
+	if (TSharedPtr<FChaosVDRemoteSessionsManager> RemoteSessionManager = FChaosVDEngineEditorBridge::Get().GetRemoteSessionsManager())
+	{
+		RemoteSessionManager->EnumerateActiveSessions([this, &MenuBuilder](const TSharedRef<FChaosVDSessionInfo>& InSessionInfoRef)
+		{
+			if (EnumHasAnyFlags(InSessionInfoRef->GetSessionTypeAttributes(), EChaosVDRemoteSessionAttributes::IsMultiSessionWrapper))
+			{
+				return true;
+			}
+	
+			FText SessionNameAsText = FText::AsCultureInvariant(InSessionInfoRef->SessionName);
+			MenuBuilder.AddMenuEntry(
+			SessionNameAsText,
+			FText::Format(LOCTEXT("MultiTargetItemTooltip", "Select {0} session as one of the current targets"), SessionNameAsText),
+			GetIconForSession(InSessionInfoRef->InstanceId),
+			FUIAction(
+			FExecuteAction::CreateSP(this, &SChaosVDRecordingControls::ToggleSessionSelectionInCustomTarget, InSessionInfoRef->InstanceId),
+			FCanExecuteAction::CreateSP(this, &SChaosVDRecordingControls::CanSelectInCustomTarget, InSessionInfoRef->InstanceId),
+			FIsActionChecked::CreateSP(this, &SChaosVDRecordingControls::IsSessionPartOfCustomTargetSelection, InSessionInfoRef->InstanceId)),
+			NAME_None, EUserInterfaceActionType::ToggleButton);
+
+			return true;
+		});
+	}
+}
+
+bool SChaosVDRecordingControls::IsSessionPartOfCustomTargetSelection(FGuid SessionGuid)
+{
+	if (TSharedPtr<FChaosVDRemoteSessionsManager> RemoteSessionManager = FChaosVDEngineEditorBridge::Get().GetRemoteSessionsManager())
+	{
+		if (TSharedPtr<FChaosVDMultiSessionInfo> CustomSessionTarget = StaticCastSharedPtr<FChaosVDMultiSessionInfo>(RemoteSessionManager->GetSessionInfo(FChaosVDRemoteSessionsManager::CustomSessionsWrapperGUID).Pin()))
+		{
+			return CustomSessionTarget->InnerSessionsByInstanceID.Find(SessionGuid) != nullptr;
+		}
+	}
+
+	return false;
+}
+
+void SChaosVDRecordingControls::ToggleSessionSelectionInCustomTarget(FGuid SessionGuid)
+{
+	if (TSharedPtr<FChaosVDRemoteSessionsManager> RemoteSessionManager = FChaosVDEngineEditorBridge::Get().GetRemoteSessionsManager())
+	{
+		if (TSharedPtr<FChaosVDMultiSessionInfo> CustomSessionTarget = StaticCastSharedPtr<FChaosVDMultiSessionInfo>(RemoteSessionManager->GetSessionInfo(FChaosVDRemoteSessionsManager::CustomSessionsWrapperGUID).Pin()))
+		{
+			if (CustomSessionTarget->InnerSessionsByInstanceID.Find(SessionGuid) != nullptr)
+			{
+				CustomSessionTarget->InnerSessionsByInstanceID.Remove(SessionGuid);
+			}
+			else
+			{
+				CustomSessionTarget->InnerSessionsByInstanceID.Add(SessionGuid, RemoteSessionManager->GetSessionInfo(SessionGuid));
+			}
+
+			SelectTargetSession(FChaosVDRemoteSessionsManager::CustomSessionsWrapperGUID);
+		}
+	}
+}
+
+bool SChaosVDRecordingControls::CanSelectInCustomTarget(FGuid SessionGuid) const
+{
+	if (TSharedPtr<FChaosVDRemoteSessionsManager> RemoteSessionManager = FChaosVDEngineEditorBridge::Get().GetRemoteSessionsManager())
+	{
+		if (TSharedPtr<FChaosVDSessionInfo> CustomSessionTarget = RemoteSessionManager->GetSessionInfo(SessionGuid).Pin())
+		{
+			return CustomSessionTarget->ReadyState == EChaosVDRemoteSessionReadyState::Ready;
+		}
+	}
+
+	return false;
+}
+
+bool SChaosVDRecordingControls::CanSelectMultiSessionTarget(FGuid SessionGuid) const
+{
+	if (TSharedPtr<FChaosVDRemoteSessionsManager> RemoteSessionManager = FChaosVDEngineEditorBridge::Get().GetRemoteSessionsManager())
+	{
+		if (TSharedPtr<FChaosVDSessionInfo> CustomSessionTarget = RemoteSessionManager->GetSessionInfo(SessionGuid).Pin())
+		{
+			return CanSelectMultiSessionTarget(CustomSessionTarget.ToSharedRef());
+		}
+	}
+
+	return false;
+}
+
+bool SChaosVDRecordingControls::CanSelectMultiSessionTarget(const TSharedRef<FChaosVDSessionInfo>& SessionInfoRef) const
+{
+	if (EnumHasAnyFlags(SessionInfoRef->GetSessionTypeAttributes(), EChaosVDRemoteSessionAttributes::IsMultiSessionWrapper))
+	{
+		TSharedRef<FChaosVDMultiSessionInfo> AsMultiSessionInfo = StaticCastSharedRef<FChaosVDMultiSessionInfo>(SessionInfoRef);
+		return !AsMultiSessionInfo->InnerSessionsByInstanceID.IsEmpty();
+	}
+
+	return false;
+}
+
 FSlateIcon SChaosVDRecordingControls::GetIconForSession(FGuid SessionId)
 {
 	if (TSharedPtr<FChaosVDSessionInfo> SessionInfo = GetSessionInfo(SessionId))
@@ -475,19 +601,19 @@ void SChaosVDRecordingControls::HandleRecordingStart(TWeakPtr<FChaosVDSessionInf
 	RecordingMessageHandle = StatusBarSubsystem->PushStatusBarMessage(StatusBarID, LOCTEXT("RecordingMessage", "Recording..."));
 }
 
-void SChaosVDRecordingControls::ExecuteAsyncConnectionAttemptTaskWithRetry(int32 RemainingRetries, TFunction<bool()> RecordingStartAttemptCallback)
+void SChaosVDRecordingControls::ExecuteAsyncConnectionAttemptTaskWithRetry(int32 RemainingRetries, const TFunction<bool()>& InRecordingStartAttemptCallback, const TFunction<void()>& InRecordingFailedCallback)
 {
 	TSharedPtr<SNotificationItem> AttemptNotification = PushConnectionAttemptNotification();
 
-	ExecutePostRecordingAsyncTaskWithRetry_Internal(RemainingRetries, RecordingStartAttemptCallback, AttemptNotification);
+	ExecutePostRecordingAsyncTaskWithRetry_Internal(RemainingRetries, InRecordingStartAttemptCallback, AttemptNotification, InRecordingFailedCallback);
 }
 
-void SChaosVDRecordingControls::ExecutePostRecordingAsyncTaskWithRetry_Internal(int32 RemainingRetries, TFunction<bool()> RecordingStartAttemptCallback, const TSharedPtr<SNotificationItem>& InProgressNotification)
+void SChaosVDRecordingControls::ExecutePostRecordingAsyncTaskWithRetry_Internal(int32 RemainingRetries, TFunction<bool()> RecordingStartAttemptCallback, const TSharedPtr<SNotificationItem>& InProgressNotification, const TFunction<void()> InRecordingFailedCallback)
 {
 	// We need to wait at least one tick before attempting to connect to give it time to the trace to be initialized, write to disk, and for the
 	// session manager to hear back from a remote instance
 
-	FTSTicker::GetCoreTicker().AddTicker(FTickerDelegate::CreateLambda([InProgressNotification, RemainingRetries, WeakThis = AsWeak(), RecordingStartAttemptCallback](float DeltaTime)
+	FTSTicker::GetCoreTicker().AddTicker(FTickerDelegate::CreateLambda([InProgressNotification, RemainingRetries, WeakThis = AsWeak(), RecordingStartAttemptCallback, InRecordingFailedCallback](float DeltaTime)
 	{
 		if (const TSharedPtr<SChaosVDRecordingControls> RecordingControlsPtr = StaticCastSharedPtr<SChaosVDRecordingControls>(WeakThis.Pin()))
 		{
@@ -512,17 +638,19 @@ void SChaosVDRecordingControls::ExecutePostRecordingAsyncTaskWithRetry_Internal(
 					if (NewRemainingRetries > 0)
 					{
 						UE_LOG(LogChaosVDEditor, Verbose, TEXT("[%s] Failed to connect to live session | Attempting again in [%f]..."), ANSI_TO_TCHAR(__FUNCTION__), RecordingControlsPtr->IntervalBetweenAutoplayConnectionAttemptsSeconds);
-						RecordingControlsPtr->ExecutePostRecordingAsyncTaskWithRetry_Internal(NewRemainingRetries, RecordingStartAttemptCallback, InProgressNotification);
+						RecordingControlsPtr->ExecutePostRecordingAsyncTaskWithRetry_Internal(NewRemainingRetries, RecordingStartAttemptCallback, InProgressNotification, InRecordingFailedCallback);
 					}
 					else
 					{
 						RecordingControlsPtr->HandleConnectionAttemptResult(EChaosVDLiveConnectionAttemptResult::Failed, InProgressNotification);
+						InRecordingFailedCallback();
 						UE_LOG(LogChaosVDEditor, Error, TEXT("[%s] Failed to connect to live session | attempts exhausted..."), ANSI_TO_TCHAR(__FUNCTION__));	
 					}
 				}
 				else
 				{
 					RecordingControlsPtr->HandleConnectionAttemptResult(EChaosVDLiveConnectionAttemptResult::Success, InProgressNotification);
+					InRecordingFailedCallback();
 				}
 			}
 		}
@@ -608,6 +736,21 @@ void SChaosVDRecordingControls::SetSessionRecordingState(bool bIsRecording, ECha
 			UE_LOG(LogChaosVDEditor, Warning, TEXT("Failed to obtain setting object. Setting the retries attempts to connect to a session to 4 as a fallback."))
 		}
 
+		SessionInfoRef->ReadyState = EChaosVDRemoteSessionReadyState::Busy;
+
+		auto RecordingAttemptFailedCallback = [SessionGUID = SessionInfoRef->InstanceId]()
+		{
+			TSharedPtr<FChaosVDRemoteSessionsManager> RemoteSessionManagerPtr = FChaosVDEngineEditorBridge::Get().GetRemoteSessionsManager();
+			TSharedPtr<FChaosVDSessionInfo> SessionInfoPtr = RemoteSessionManagerPtr ? RemoteSessionManagerPtr->GetSessionInfo(SessionGUID).Pin() : nullptr;
+
+			if (!SessionInfoPtr)
+			{
+				return;
+			}
+
+			SessionInfoPtr->ReadyState = EChaosVDRemoteSessionReadyState::Ready;
+		};
+
 		if (RecordingMode == EChaosVDRecordingMode::Live)
 		{
 			FChaosVDStartRecordingCommandMessage RecordingParams;
@@ -644,7 +787,8 @@ void SChaosVDRecordingControls::SetSessionRecordingState(bool bIsRecording, ECha
 
 				FLiveConnectionSettings ConnectionDetails = Controls->GetTargetLiveConnectionSettings(SessionInfoPtr.ToSharedRef());
 				return MainTabSharedPtr->ConnectToLiveSession(ConnectionDetails.TraceID, ConnectionDetails.SessionAddress, Controls->CurrentLoadingMode);
-			});
+			},
+			RecordingAttemptFailedCallback);
 		}
 		else
 		{
@@ -658,7 +802,8 @@ void SChaosVDRecordingControls::SetSessionRecordingState(bool bIsRecording, ECha
 			{
 				TSharedPtr<SChaosVDRecordingControls> Controls = StaticCastSharedPtr<SChaosVDRecordingControls>(WeakThis.Pin());
 				return Controls ? Controls->IsRecording() : false;
-			});
+			},
+			RecordingAttemptFailedCallback);
 		}
 	}
 	else
@@ -702,7 +847,17 @@ bool SChaosVDRecordingControls::IsRecordingToggleButtonEnabled(EChaosVDRecording
 	// If we are recording, don't show the stop button for the mode that is disabled
 	if (TSharedPtr<FChaosVDSessionInfo> SessionInfo = GetCurrentSessionInfo())
 	{
-		return SessionInfo->GetRecordingMode() == RecordingMode;
+		if (SessionInfo->GetRecordingMode() == RecordingMode)
+		{
+			if (EnumHasAnyFlags(SessionInfo->GetSessionTypeAttributes(), EChaosVDRemoteSessionAttributes::IsMultiSessionWrapper))
+			{
+				TSharedPtr<FChaosVDMultiSessionInfo> AsMultiSessionInfo = StaticCastSharedPtr<FChaosVDMultiSessionInfo>(SessionInfo);
+
+				return CanSelectMultiSessionTarget(SessionInfo.ToSharedRef());
+			}
+		}
+
+		return false;
 	}
 
 	return false;
diff --git a/Engine/Plugins/ChaosVD/Source/ChaosVD/Private/Widgets/SChaosVDRecordingControls.h b/Engine/Plugins/ChaosVD/Source/ChaosVD/Private/Widgets/SChaosVDRecordingControls.h
--- a/Engine/Plugins/ChaosVD/Source/ChaosVD/Private/Widgets/SChaosVDRecordingControls.h
+++ b/Engine/Plugins/ChaosVD/Source/ChaosVD/Private/Widgets/SChaosVDRecordingControls.h
@@ -10,6 +10,8 @@
 
 #include "SChaosVDRecordingControls.generated.h"
 
+class FMenuBuilder;
+
 namespace Chaos::VisualDebugger
 {
 	struct FChaosVDOptionalDataChannel;
@@ -70,6 +72,15 @@ protected:
 
 	void SelectTargetSession(FGuid SessionId);
 
+	void GenerateCustomTargetsMenu(FMenuBuilder& MenuBuilder);
+
+	bool IsSessionPartOfCustomTargetSelection(FGuid SessionGuid);
+	void ToggleSessionSelectionInCustomTarget(FGuid SessionGuid);
+	bool CanSelectInCustomTarget(FGuid SessionGuid) const;
+
+	bool CanSelectMultiSessionTarget(FGuid SessionGuid) const;
+	bool CanSelectMultiSessionTarget(const TSharedRef<FChaosVDSessionInfo>& SessionInfoRef) const;
+
 	FSlateIcon GetIconForSession(FGuid SessionId);
 
 	TSharedPtr<FChaosVDSessionInfo> GetCurrentSessionInfo() const;
@@ -84,10 +95,10 @@ protected:
 	void HandleRecordingStop(TWeakPtr<FChaosVDSessionInfo> SessionInfo);
 	void HandleRecordingStart(TWeakPtr<FChaosVDSessionInfo> SessionInfo);
 
-	void ExecuteAsyncConnectionAttemptTaskWithRetry(int32 RemainingRetries, TFunction<bool()> RecordingStartAttemptCallback);
+	void ExecuteAsyncConnectionAttemptTaskWithRetry(int32 RemainingRetries, const TFunction<bool()>& InRecordingStartAttemptCallback, const TFunction<void()>& InRecordingFailedCallback);
 
 private:
-	void ExecutePostRecordingAsyncTaskWithRetry_Internal(int32 RemainingRetries, TFunction<bool()> RecordingStartAttemptCallback, const TSharedPtr<SNotificationItem>& InProgressNotification);
+	void ExecutePostRecordingAsyncTaskWithRetry_Internal(int32 RemainingRetries, TFunction<bool()> RecordingStartAttemptCallback, const TSharedPtr<SNotificationItem>& InProgressNotification, const TFunction<void()> InRecordingFailedCallback);
 protected:
 
 	FReply ToggleRecordingState(EChaosVDRecordingMode RecordingMode);
diff --git a/Engine/Source/Runtime/Experimental/ChaosSolverEngine/Private/Chaos/ChaosVDRemoteSessionsManager.cpp b/Engine/Source/Runtime/Experimental/ChaosSolverEngine/Private/Chaos/ChaosVDRemoteSessionsManager.cpp
--- a/Engine/Source/Runtime/Experimental/ChaosSolverEngine/Private/Chaos/ChaosVDRemoteSessionsManager.cpp
+++ b/Engine/Source/Runtime/Experimental/ChaosSolverEngine/Private/Chaos/ChaosVDRemoteSessionsManager.cpp
@@ -11,14 +11,20 @@
 #include "Misc/App.h"
 
 const FGuid FChaosVDRemoteSessionsManager::AllRemoteSessionsWrapperGUID = FGuid::NewGuid();
+const FGuid FChaosVDRemoteSessionsManager::AllRemoteServersWrapperGUID = FGuid::NewGuid();
+const FGuid FChaosVDRemoteSessionsManager::AllRemoteClientsWrapperGUID = FGuid::NewGuid();
 const FGuid FChaosVDRemoteSessionsManager::AllSessionsWrapperGUID = FGuid::NewGuid();
+const FGuid FChaosVDRemoteSessionsManager::CustomSessionsWrapperGUID = FGuid::NewGuid();
 const FGuid FChaosVDRemoteSessionsManager::InvalidSessionGUID = FGuid();
 
 const FString FChaosVDRemoteSessionsManager::LocalEditorSessionName = TEXT("Local Editor");
 const FGuid FChaosVDRemoteSessionsManager::LocalEditorSessionID = IsController() ? FApp::GetInstanceId() : InvalidSessionGUID;
 const FName FChaosVDRemoteSessionsManager::MessageBusEndPointName = FName("CVDSessionManagerEndPoint");
-const FString FChaosVDRemoteSessionsManager::AllRemoteSessionsTargetName = TEXT("AllRemote");
-const FString FChaosVDRemoteSessionsManager::AllSessionsTargetName = TEXT("AllSessions");
+const FString FChaosVDRemoteSessionsManager::AllRemoteSessionsTargetName = TEXT("All Remote");
+const FString FChaosVDRemoteSessionsManager::AllRemoteServersTargetName = TEXT("All Remote Servers");
+const FString FChaosVDRemoteSessionsManager::AllRemoteClientsTargetName = TEXT("All Remote Clients");
+const FString FChaosVDRemoteSessionsManager::AllSessionsTargetName = TEXT("All Sessions");
+const FString FChaosVDRemoteSessionsManager::CustomSessionsTargetName = TEXT("Custom Selection");
 
 DEFINE_LOG_CATEGORY(LogChaosVDRemoteSession)
 
@@ -167,7 +173,10 @@ void FChaosVDRemoteSessionsManager::Initialize()
 	MessageEndpoint->Subscribe<FChaosVDFullSessionInfoRequestMessage>();
 
 	ActiveSessionsByInstanceId.Add(AllRemoteSessionsWrapperGUID, CreatedWrapperSessionInfo(AllRemoteSessionsWrapperGUID, AllRemoteSessionsTargetName));
+	ActiveSessionsByInstanceId.Add(AllRemoteServersWrapperGUID, CreatedWrapperSessionInfo(AllRemoteServersWrapperGUID, AllRemoteServersTargetName));
+	ActiveSessionsByInstanceId.Add(AllRemoteClientsWrapperGUID, CreatedWrapperSessionInfo(AllRemoteClientsWrapperGUID, AllRemoteClientsTargetName));
 	ActiveSessionsByInstanceId.Add(AllSessionsWrapperGUID, CreatedWrapperSessionInfo(AllSessionsWrapperGUID, AllSessionsTargetName));
+	ActiveSessionsByInstanceId.Add(CustomSessionsWrapperGUID, CreatedWrapperSessionInfo(CustomSessionsWrapperGUID, CustomSessionsTargetName));
 
 	constexpr float TickInterval = 1.0f;
 	TickHandle = FTSTicker::GetCoreTicker().AddTicker(FTickerDelegate::CreateRaw(this, &FChaosVDRemoteSessionsManager::Tick), TickInterval);
@@ -313,6 +322,7 @@ void FChaosVDRemoteSessionsManager::SendPong(const FChaosVDSessionPing& InMessag
 	{
 		PongMessage->InstanceId = FApp::GetInstanceId();
 		PongMessage->SessionId = FApp::GetSessionId();
+		PongMessage->BuildTargetType = static_cast<uint8>(FApp::GetBuildTargetType());
 
 		if (InMessage.ControllerInstanceId == PongMessage->InstanceId)
 		{
@@ -333,6 +343,15 @@ void FChaosVDRemoteSessionsManager::RegisterSessionInMultiSessionWrapper(const T
 	if (InSessionInfoRef->SessionName != LocalEditorSessionName )
 	{
 		StaticCastSharedPtr<FChaosVDMultiSessionInfo>(ActiveSessionsByInstanceId.FindChecked(AllRemoteSessionsWrapperGUID))->InnerSessionsByInstanceID.Emplace(InSessionInfoRef->InstanceId, InSessionInfoRef);
+
+		if (InSessionInfoRef->BuildTargetType == EBuildTargetType::Server)
+		{
+			StaticCastSharedPtr<FChaosVDMultiSessionInfo>(ActiveSessionsByInstanceId.FindChecked(AllRemoteServersWrapperGUID))->InnerSessionsByInstanceID.Emplace(InSessionInfoRef->InstanceId, InSessionInfoRef);
+		}
+		else
+		{
+			StaticCastSharedPtr<FChaosVDMultiSessionInfo>(ActiveSessionsByInstanceId.FindChecked(AllRemoteClientsWrapperGUID))->InnerSessionsByInstanceID.Emplace(InSessionInfoRef->InstanceId, InSessionInfoRef);
+		}
 	}
 
 	StaticCastSharedPtr<FChaosVDMultiSessionInfo>(ActiveSessionsByInstanceId.FindChecked(AllSessionsWrapperGUID))->InnerSessionsByInstanceID.Emplace(InSessionInfoRef->InstanceId, InSessionInfoRef);
@@ -342,6 +361,8 @@ void FChaosVDRemoteSessionsManager::DeRegisterSessionInMultiSessionWrapper(const
 {
 	StaticCastSharedPtr<FChaosVDMultiSessionInfo>(ActiveSessionsByInstanceId.FindChecked(AllRemoteSessionsWrapperGUID))->InnerSessionsByInstanceID.Remove(InSessionInfoRef->InstanceId);
 	StaticCastSharedPtr<FChaosVDMultiSessionInfo>(ActiveSessionsByInstanceId.FindChecked(AllSessionsWrapperGUID))->InnerSessionsByInstanceID.Remove(InSessionInfoRef->InstanceId);
+	StaticCastSharedPtr<FChaosVDMultiSessionInfo>(ActiveSessionsByInstanceId.FindChecked(AllRemoteServersWrapperGUID))->InnerSessionsByInstanceID.Remove(InSessionInfoRef->InstanceId);
+	StaticCastSharedPtr<FChaosVDMultiSessionInfo>(ActiveSessionsByInstanceId.FindChecked(AllRemoteClientsWrapperGUID))->InnerSessionsByInstanceID.Remove(InSessionInfoRef->InstanceId);
 }
 
 void FChaosVDRemoteSessionsManager::HandleSessionPongMessage(const FChaosVDSessionPong& InMessage, const TSharedRef<IMessageContext>& InContext)
@@ -354,6 +375,7 @@ void FChaosVDRemoteSessionsManager::HandleSessionPongMessage(const FChaosVDSessi
 		SessionInfoPtr->Address = InContext->GetSender();
 		SessionInfoPtr->InstanceId = InMessage.InstanceId;
 		SessionInfoPtr->SessionName = InMessage.SessionName;
+		SessionInfoPtr->BuildTargetType = static_cast<EBuildTargetType>(InMessage.BuildTargetType);
 
 		RegisterSessionInMultiSessionWrapper(SessionInfoPtr.ToSharedRef());
 
@@ -536,7 +558,9 @@ void FChaosVDRemoteSessionsManager::RemoveExpiredSessions()
 
 		FTimespan ElapsedTime = CurrentTime - SessionInfoPtr->LastPingTime;
 
-		constexpr float MaxAllowedTimeBetweenPings = 3.0f;
+		// A session goes into busy state if we are attempting to issue a command that might stall the target, currently that only happens on recording start commands of complex maps
+		// In these cases, we need to allow more time between pings. If a recording command failed, it is expected the state to be changed to Ready again
+		const float MaxAllowedTimeBetweenPings = SessionInfoPtr->ReadyState == EChaosVDRemoteSessionReadyState::Busy ? 3.0f : 60.0f;
 		if (ElapsedTime > FTimespan::FromSeconds(MaxAllowedTimeBetweenPings))
 		{
 			DeRegisterSessionInMultiSessionWrapper(SessionInfoPtr.ToSharedRef());
diff --git a/Engine/Source/Runtime/Experimental/ChaosSolverEngine/Public/Chaos/ChaosVDRemoteSessionsManager.h b/Engine/Source/Runtime/Experimental/ChaosSolverEngine/Public/Chaos/ChaosVDRemoteSessionsManager.h
--- a/Engine/Source/Runtime/Experimental/ChaosSolverEngine/Public/Chaos/ChaosVDRemoteSessionsManager.h
+++ b/Engine/Source/Runtime/Experimental/ChaosSolverEngine/Public/Chaos/ChaosVDRemoteSessionsManager.h
@@ -38,6 +38,9 @@ struct FChaosVDSessionPong
 
 	UPROPERTY()
 	FString SessionName;
+
+	UPROPERTY()
+	uint8 BuildTargetType = static_cast<uint8>(EBuildTargetType::Unknown);
 };
 
 USTRUCT()
@@ -146,9 +149,17 @@ enum class EChaosVDRemoteSessionAttributes
 	CanExpire = 1 << 1,
 	IsMultiSessionWrapper = 1 << 2,
 };
-
 ENUM_CLASS_FLAGS(EChaosVDRemoteSessionAttributes)
 
+UENUM()
+enum class EChaosVDRemoteSessionReadyState : uint8
+{
+	/** The session is ready to execute commands */
+	Ready,
+	/** We are executing a command in the session we expect to take a while without hearing anything from the target. */
+	Busy
+};
+
 /**
  * Session object that contains all the information needed to communicate with a remote instance, and the state of that instance
  */
@@ -173,6 +184,8 @@ public:
 	FString SessionName;
 	FMessageAddress Address;
 	FDateTime LastPingTime;
+	EBuildTargetType BuildTargetType = EBuildTargetType::Unknown;
+	EChaosVDRemoteSessionReadyState ReadyState = EChaosVDRemoteSessionReadyState::Ready;
 
 	FChaosVDRecordingStatusMessage LastKnownRecordingState;
 	TMap<FString, FChaosVDDataChannelState> DataChannelsStatesByName;
@@ -323,11 +336,17 @@ public:
 	}
 	
 	CHAOSSOLVERENGINE_API static const FString AllRemoteSessionsTargetName;
+	CHAOSSOLVERENGINE_API static const FString AllRemoteServersTargetName;
+	CHAOSSOLVERENGINE_API static const FString AllRemoteClientsTargetName;
 	CHAOSSOLVERENGINE_API static const FString AllSessionsTargetName;
+	CHAOSSOLVERENGINE_API static const FString CustomSessionsTargetName;
 	CHAOSSOLVERENGINE_API static const FString LocalEditorSessionName;
 	CHAOSSOLVERENGINE_API static const FName MessageBusEndPointName;
 	CHAOSSOLVERENGINE_API static const FGuid AllRemoteSessionsWrapperGUID;
+	CHAOSSOLVERENGINE_API static const FGuid AllRemoteServersWrapperGUID;
+	CHAOSSOLVERENGINE_API static const FGuid AllRemoteClientsWrapperGUID;
 	CHAOSSOLVERENGINE_API static const FGuid AllSessionsWrapperGUID;
+	CHAOSSOLVERENGINE_API static const FGuid CustomSessionsWrapperGUID;
 	CHAOSSOLVERENGINE_API static const FGuid InvalidSessionGUID;
 	CHAOSSOLVERENGINE_API static const FGuid LocalEditorSessionID;
 
