[ChaosVD] Updated the Generic Debug Draw trace system to allow users trace data in specific solver stages or solver frame.

Tracing debug draw shape data from any thread, from anywhere was already possible, but during playback the data was organized in game thread frames, which means if you wanted to trace a shape at a specific solver stage, you wouldn't see it when that stage was played back. Instead you would see it when the game frame that contained was played (which can contain data from multiple solver frames/stages).

User API wise, nothing changes. The underlying system takes care of finding the correct context for the data we are tracing, as we already do for particle and other physics data.

diff --git a/Engine/Plugins/ChaosVD/Source/ChaosVDBuiltInExtensions/Private/GenericDebugDraw/Components/ChaosVDGenericDebugDrawDataComponent.cpp b/Engine/Plugins/ChaosVD/Source/ChaosVDBuiltInExtensions/Private/GenericDebugDraw/Components/ChaosVDGenericDebugDrawDataComponent.cpp
--- a/Engine/Plugins/ChaosVD/Source/ChaosVDBuiltInExtensions/Private/GenericDebugDraw/Components/ChaosVDGenericDebugDrawDataComponent.cpp
+++ b/Engine/Plugins/ChaosVD/Source/ChaosVDBuiltInExtensions/Private/GenericDebugDraw/Components/ChaosVDGenericDebugDrawDataComponent.cpp
@@ -13,24 +13,87 @@ UChaosVDGenericDebugDrawDataComponent::UChaosVDGenericDebugDrawDataComponent()
 
 void UChaosVDGenericDebugDrawDataComponent::UpdateFromNewGameFrameData(const FChaosVDGameFrameData& InGameFrameData)
 {
-	if (TSharedPtr<FChaosVDDebugShapeDataContainer> DebugDrawShapeData = InGameFrameData.GetCustomDataHandler().GetData<FChaosVDDebugShapeDataContainer>())
+	Super::UpdateFromNewGameFrameData(InGameFrameData);
+
+	if (TSharedPtr<FChaosVDMultiSolverDebugShapeDataContainer> MultiSolverData = InGameFrameData.GetCustomDataHandler().GetData<FChaosVDMultiSolverDebugShapeDataContainer>())
 	{
-		CopyDataFromSourceMap<FChaosVDDebugDrawBoxDataWrapper>(DebugDrawShapeData->RecordedDebugDrawBoxesBySolverID, DebugDrawBoxes, SolverID);
-		CopyDataFromSourceMap<FChaosVDDebugDrawLineDataWrapper>(DebugDrawShapeData->RecordedDebugDrawLinesBySolverID, DebugDrawLines, SolverID);
-		CopyDataFromSourceMap<FChaosVDDebugDrawSphereDataWrapper>(DebugDrawShapeData->RecordedDebugDrawSpheresBySolverID, DebugDrawSpheres, SolverID);
-		CopyDataFromSourceMap<FChaosVDDebugDrawImplicitObjectDataWrapper>(DebugDrawShapeData->RecordedDebugDrawImplicitObjectsBySolverID, DebugDrawImplicitObjects, SolverID);
+		TSharedPtr<FChaosVDDebugShapeDataContainer>* SolverData = MultiSolverData->DataBySolverID.Find(SolverID);
+		CurrentGameFrameDebugDrawData = SolverData ? *SolverData : nullptr;
 	}
 	else
 	{
-		ClearData();
+		CurrentGameFrameDebugDrawData = nullptr;
 	}
 }
 
+void UChaosVDGenericDebugDrawDataComponent::UpdateFromSolverFrameData(const FChaosVDSolverFrameData& InSolverFrameData)
+{
+	CurrentSolverFrameDebugDrawData = InSolverFrameData.GetCustomData().GetData<FChaosVDDebugShapeDataContainer>();
+}
+
+void UChaosVDGenericDebugDrawDataComponent::UpdateFromNewSolverStageData(const FChaosVDSolverFrameData& InSolverFrameData, const FChaosVDFrameStageData& InSolverFrameStageData)
+{
+	CurrentSolverFrameDebugDrawData = InSolverFrameStageData.GetCustomDataHandler().GetData<FChaosVDDebugShapeDataContainer>();
+}
+
 void UChaosVDGenericDebugDrawDataComponent::ClearData()
 {
-	DebugDrawBoxes.Reset();
-	DebugDrawLines.Reset();
-	DebugDrawSpheres.Reset();
-	DebugDrawLines.Reset();
-	DebugDrawImplicitObjects.Reset();
+	CurrentSolverStageDebugDrawData.Reset();
+	CurrentSolverFrameDebugDrawData.Reset();
+	CurrentGameFrameDebugDrawData.Reset();
+}
+
+TConstArrayView<TSharedPtr<FChaosVDDebugDrawBoxDataWrapper>> UChaosVDGenericDebugDrawDataComponent::GetDebugDrawBoxesDataView(EChaosVDDrawDataContainerSource Source) const
+{
+	if (TSharedPtr<FChaosVDDebugShapeDataContainer> DataContainer = GetShapeDataContainer(Source))
+	{
+		return DataContainer->RecordedDebugDrawBoxes;
+	}
+
+	return TConstArrayView<TSharedPtr<FChaosVDDebugDrawBoxDataWrapper>>();
+}
+
+TConstArrayView<TSharedPtr<FChaosVDDebugDrawLineDataWrapper>> UChaosVDGenericDebugDrawDataComponent::GetDebugDrawLinesDataView( EChaosVDDrawDataContainerSource Source) const
+{
+	if (TSharedPtr<FChaosVDDebugShapeDataContainer> DataContainer = GetShapeDataContainer(Source))
+	{
+		return DataContainer->RecordedDebugDrawLines;
+	}
+
+	return TConstArrayView<TSharedPtr<FChaosVDDebugDrawLineDataWrapper>>();
+}
+
+TConstArrayView<TSharedPtr<FChaosVDDebugDrawSphereDataWrapper>> UChaosVDGenericDebugDrawDataComponent::GetDebugDrawSpheresDataView( EChaosVDDrawDataContainerSource Source) const
+{
+	if (TSharedPtr<FChaosVDDebugShapeDataContainer> DataContainer = GetShapeDataContainer(Source))
+	{
+		return DataContainer->RecordedDebugDrawSpheres;
+	}
+
+	return TConstArrayView<TSharedPtr<FChaosVDDebugDrawSphereDataWrapper>>();
+}
+
+TConstArrayView<TSharedPtr<FChaosVDDebugDrawImplicitObjectDataWrapper>> UChaosVDGenericDebugDrawDataComponent::GetDebugDrawImplicitObjectsDataView( EChaosVDDrawDataContainerSource Source) const
+{
+	if (TSharedPtr<FChaosVDDebugShapeDataContainer> DataContainer = GetShapeDataContainer(Source))
+	{
+		return DataContainer->RecordedDebugDrawImplicitObjects;
+	}
+
+	return TConstArrayView<TSharedPtr<FChaosVDDebugDrawImplicitObjectDataWrapper>>();
+}
+
+TSharedPtr<FChaosVDDebugShapeDataContainer> UChaosVDGenericDebugDrawDataComponent::GetShapeDataContainer(EChaosVDDrawDataContainerSource Source) const
+{
+	switch (Source)
+	{
+		case EChaosVDDrawDataContainerSource::SolverFrame:
+			return CurrentSolverFrameDebugDrawData;
+		case EChaosVDDrawDataContainerSource::SolverStage:
+			return CurrentSolverStageDebugDrawData;
+		case EChaosVDDrawDataContainerSource::GameFrame:
+			return CurrentGameFrameDebugDrawData;
+		default:
+			return nullptr;
+	}
 }
diff --git a/Engine/Plugins/ChaosVD/Source/ChaosVDBuiltInExtensions/Private/GenericDebugDraw/Components/ChaosVDGenericDebugDrawDataComponent.h b/Engine/Plugins/ChaosVD/Source/ChaosVDBuiltInExtensions/Private/GenericDebugDraw/Components/ChaosVDGenericDebugDrawDataComponent.h
--- a/Engine/Plugins/ChaosVD/Source/ChaosVDBuiltInExtensions/Private/GenericDebugDraw/Components/ChaosVDGenericDebugDrawDataComponent.h
+++ b/Engine/Plugins/ChaosVD/Source/ChaosVDBuiltInExtensions/Private/GenericDebugDraw/Components/ChaosVDGenericDebugDrawDataComponent.h
@@ -10,6 +10,13 @@
 class UChaosVDGenericDebugDrawDataComponent;
 struct FChaosVDGameFrameData;
 
+enum class EChaosVDDrawDataContainerSource : uint8
+{
+	SolverFrame,
+	SolverStage,
+	GameFrame
+};
+
 UCLASS()
 class UChaosVDGenericDebugDrawDataComponent : public UChaosVDSolverDataComponent
 {
@@ -19,35 +26,22 @@ public:
 
 	virtual void UpdateFromNewGameFrameData(const FChaosVDGameFrameData& InGameFrameData) override;
 
+	virtual void UpdateFromSolverFrameData(const FChaosVDSolverFrameData& InSolverFrameData) override;
+
+	virtual void UpdateFromNewSolverStageData(const FChaosVDSolverFrameData& InSolverFrameData, const FChaosVDFrameStageData& InSolverFrameStageData) override;
+
 	virtual void ClearData() override;
 
-	TConstArrayView<TSharedPtr<FChaosVDDebugDrawBoxDataWrapper>> GetDebugDrawBoxesDataView() const { return DebugDrawBoxes; };
-	TConstArrayView<TSharedPtr<FChaosVDDebugDrawLineDataWrapper>> GetDebugDrawLinesDataView() const { return DebugDrawLines; };
-	TConstArrayView<TSharedPtr<FChaosVDDebugDrawSphereDataWrapper>> GetDebugDrawSpheresDataView() const { return DebugDrawSpheres; };
-	TConstArrayView<TSharedPtr<FChaosVDDebugDrawImplicitObjectDataWrapper>> GetDebugDrawImplicitObjectsDataView() const { return DebugDrawImplicitObjects; };
+	TConstArrayView<TSharedPtr<FChaosVDDebugDrawBoxDataWrapper>> GetDebugDrawBoxesDataView(EChaosVDDrawDataContainerSource Source) const;
+	TConstArrayView<TSharedPtr<FChaosVDDebugDrawLineDataWrapper>> GetDebugDrawLinesDataView(EChaosVDDrawDataContainerSource Source) const;
+	TConstArrayView<TSharedPtr<FChaosVDDebugDrawSphereDataWrapper>> GetDebugDrawSpheresDataView(EChaosVDDrawDataContainerSource Source) const;
+	TConstArrayView<TSharedPtr<FChaosVDDebugDrawImplicitObjectDataWrapper>> GetDebugDrawImplicitObjectsDataView(EChaosVDDrawDataContainerSource Source) const;
+
+	TSharedPtr<FChaosVDDebugShapeDataContainer> GetShapeDataContainer(EChaosVDDrawDataContainerSource Source) const;
 
 private:
 
-	template<typename RecordedDataType>
-	void CopyDataFromSourceMap(const TMap<int32,TArray<TSharedPtr<RecordedDataType>>>& CopyFrom, TArray<TSharedPtr<RecordedDataType>>& CopyTo, int32 SolverID);
-	
-	TArray<TSharedPtr<FChaosVDDebugDrawBoxDataWrapper>> DebugDrawBoxes;
-	TArray<TSharedPtr<FChaosVDDebugDrawLineDataWrapper>> DebugDrawLines;
-	TArray<TSharedPtr<FChaosVDDebugDrawSphereDataWrapper>> DebugDrawSpheres;
-	TArray<TSharedPtr<FChaosVDDebugDrawImplicitObjectDataWrapper>> DebugDrawImplicitObjects;
+	TSharedPtr<FChaosVDDebugShapeDataContainer> CurrentSolverStageDebugDrawData;
+	TSharedPtr<FChaosVDDebugShapeDataContainer> CurrentSolverFrameDebugDrawData;
+	TSharedPtr<FChaosVDDebugShapeDataContainer> CurrentGameFrameDebugDrawData;
 };
-
-template <typename RecordedDataType>
-void UChaosVDGenericDebugDrawDataComponent::CopyDataFromSourceMap(const TMap<int32, TArray<TSharedPtr<RecordedDataType>>>& CopyFrom, TArray<TSharedPtr<RecordedDataType>>& CopyTo, int32 SolverID)
-{
-	if (const TArray<TSharedPtr<RecordedDataType>>* RecordedData = CopyFrom.Find(SolverID))
-	{
-		CopyTo.Reset(RecordedData->Num());
-		CopyTo = *RecordedData;
-	}
-	else
-	{
-		CopyTo.Reset();	
-	}
-}
-
diff --git a/Engine/Plugins/ChaosVD/Source/ChaosVDBuiltInExtensions/Private/GenericDebugDraw/Trace/DataProcessors/ChaosVDDebugDrawBoxDataProcessor.cpp b/Engine/Plugins/ChaosVD/Source/ChaosVDBuiltInExtensions/Private/GenericDebugDraw/Trace/DataProcessors/ChaosVDDebugDrawBoxDataProcessor.cpp
--- a/Engine/Plugins/ChaosVD/Source/ChaosVDBuiltInExtensions/Private/GenericDebugDraw/Trace/DataProcessors/ChaosVDDebugDrawBoxDataProcessor.cpp
+++ b/Engine/Plugins/ChaosVD/Source/ChaosVDBuiltInExtensions/Private/GenericDebugDraw/Trace/DataProcessors/ChaosVDDebugDrawBoxDataProcessor.cpp
@@ -1,6 +1,8 @@
 // Copyright Epic Games, Inc. All Rights Reserved.
 
 #include "ChaosVDDebugDrawBoxDataProcessor.h"
+
+#include "ChaosVDDebugDrawDataProcessorUtils.h"
 #include "DataWrappers/ChaosVDDebugShapeDataWrapper.h"
 #include "Visualizers/ChaosVDDebugDrawUtils.h"
 
@@ -24,15 +26,15 @@ bool FChaosVDDebugDrawBoxDataProcessor::ProcessRawData(const TArray<uint8>& InDa
 	
 	if (bSuccess)
 	{
-		if (TSharedPtr<FChaosVDGameFrameData> CurrentFrameData = ProviderSharedPtr->GetCurrentGameFrame().Pin())
-		{
-			DebugDrawData->SolverID = ProviderSharedPtr->GetRemappedSolverID(DebugDrawData->SolverID);
+		using namespace Chaos::VisualDebugger::Utils;
 
-			if (TSharedPtr<FChaosVDDebugShapeDataContainer> DebugDrawShapeData = CurrentFrameData->GetCustomDataHandler().GetOrAddDefaultData<FChaosVDDebugShapeDataContainer>())
-			{
-				DebugDrawShapeData->RecordedDebugDrawBoxesBySolverID.FindOrAdd(DebugDrawData->SolverID).Add(DebugDrawData);
-				CurrentFrameData->MarkDirty();
-			}
+		TSharedRef<FChaosVDTraceProvider> ProviderAsSharedRef = ProviderSharedPtr.ToSharedRef();
+		TSharedRef<FChaosVDDebugDrawBoxDataWrapper> DebugDrawDataAsSharedRef = DebugDrawData.ToSharedRef();
+		RemapDebugDrawShapeDataSolverID(DebugDrawDataAsSharedRef, ProviderAsSharedRef);
+	
+		if (TSharedPtr<FChaosVDDebugShapeDataContainer> DebugDrawShapeData = GetShapeDataContainer(DebugDrawDataAsSharedRef, ProviderAsSharedRef, EShapeDataContainerAccessorFlags::MarkFrameDirty))
+		{
+			DebugDrawShapeData->RecordedDebugDrawBoxes.Add(DebugDrawData);
 		}
 	}
 
diff --git a/Engine/Plugins/ChaosVD/Source/ChaosVDBuiltInExtensions/Private/GenericDebugDraw/Trace/DataProcessors/ChaosVDDebugDrawDataProcessorUtils.h b/Engine/Plugins/ChaosVD/Source/ChaosVDBuiltInExtensions/Private/GenericDebugDraw/Trace/DataProcessors/ChaosVDDebugDrawDataProcessorUtils.h
new file mode 100644
--- /dev/null
+++ b/Engine/Plugins/ChaosVD/Source/ChaosVDBuiltInExtensions/Private/GenericDebugDraw/Trace/DataProcessors/ChaosVDDebugDrawDataProcessorUtils.h
@@ -0,0 +1,84 @@
+// Copyright Epic Games, Inc. All Rights Reserved.
+
+#pragma once
+#include "ChaosVDRecording.h"
+#include "Templates/SharedPointer.h"
+#include "Trace/ChaosVDTraceProvider.h"
+
+namespace Chaos::VisualDebugger::Utils
+{
+	enum class EShapeDataContainerAccessorFlags : uint8
+	{
+		None = 0,
+		/** If the container lives in a game thread frame, it will mark it as dirty */
+		MarkFrameDirty = 1 << 0
+	};
+	ENUM_CLASS_FLAGS(EShapeDataContainerAccessorFlags)
+
+	template<typename DataType>
+	TSharedPtr<FChaosVDDebugShapeDataContainer> GetShapeDataContainer(const TSharedRef<DataType>& InData, const TSharedRef<FChaosVDTraceProvider>& TraceProvider, EShapeDataContainerAccessorFlags Flags)
+	{
+		EChaosVDParticleContext TracedThreadContext = InData->ThreadContext;
+		if (InData->SolverID == INDEX_NONE)
+		{
+			TracedThreadContext = EChaosVDParticleContext::GameThread;
+		}
+
+		switch (TracedThreadContext)
+		{
+			case EChaosVDParticleContext::GameThread:
+				{
+					if (TSharedPtr<FChaosVDGameFrameData> CurrentFrameData = TraceProvider->GetCurrentGameFrame().Pin())
+					{
+						if (TSharedPtr<FChaosVDMultiSolverDebugShapeDataContainer> MultiSolverData = CurrentFrameData->GetCustomDataHandler().GetOrAddDefaultData<FChaosVDMultiSolverDebugShapeDataContainer>())
+						{
+							TSharedPtr<FChaosVDDebugShapeDataContainer>& SolverData = MultiSolverData->DataBySolverID.FindOrAdd(InData->SolverID);
+
+							if (!SolverData)
+							{
+								SolverData = MakeShared<FChaosVDDebugShapeDataContainer>();
+							}
+
+							if (EnumHasAnyFlags(Flags, EShapeDataContainerAccessorFlags::MarkFrameDirty))
+							{
+								CurrentFrameData->MarkDirty();
+							}
+
+							return SolverData;
+						}
+
+						return nullptr;
+					}
+					break;
+				}
+			case EChaosVDParticleContext::PhysicsThread:
+				{
+					EChaosVDSolverStageAccessorFlags StageAccessorFlags = EChaosVDSolverStageAccessorFlags::None;
+
+					if (FChaosVDFrameStageData* CurrentSolverStage = TraceProvider->GetCurrentSolverStageDataForCurrentFrame(InData->SolverID, StageAccessorFlags))
+					{
+						return CurrentSolverStage->GetCustomDataHandler().GetOrAddDefaultData<FChaosVDDebugShapeDataContainer>();
+					}
+
+					if (FChaosVDSolverFrameData* CurrentSolverFrameData = TraceProvider->GetCurrentSolverFrame(InData->SolverID))
+					{
+						return CurrentSolverFrameData->GetCustomData().GetOrAddDefaultData<FChaosVDDebugShapeDataContainer>();
+					}
+				}
+			case EChaosVDParticleContext::Invalid:
+			default:
+				return nullptr;
+		}
+
+		return nullptr;
+	}
+
+	template<typename DataType>
+	void RemapDebugDrawShapeDataSolverID(const TSharedRef<DataType>& InData, const TSharedRef<FChaosVDTraceProvider>& TraceProvider)
+	{
+		if (InData->SolverID != INDEX_NONE)
+		{
+			InData->SolverID = TraceProvider->GetRemappedSolverID(InData->SolverID);
+		}
+	}
+}
diff --git a/Engine/Plugins/ChaosVD/Source/ChaosVDBuiltInExtensions/Private/GenericDebugDraw/Trace/DataProcessors/ChaosVDDebugDrawImplicitObjectDataProcessor.cpp b/Engine/Plugins/ChaosVD/Source/ChaosVDBuiltInExtensions/Private/GenericDebugDraw/Trace/DataProcessors/ChaosVDDebugDrawImplicitObjectDataProcessor.cpp
--- a/Engine/Plugins/ChaosVD/Source/ChaosVDBuiltInExtensions/Private/GenericDebugDraw/Trace/DataProcessors/ChaosVDDebugDrawImplicitObjectDataProcessor.cpp
+++ b/Engine/Plugins/ChaosVD/Source/ChaosVDBuiltInExtensions/Private/GenericDebugDraw/Trace/DataProcessors/ChaosVDDebugDrawImplicitObjectDataProcessor.cpp
@@ -1,10 +1,11 @@
 // Copyright Epic Games, Inc. All Rights Reserved.
 
 #include "ChaosVDDebugDrawImplicitObjectDataProcessor.h"
+
+#include "ChaosVDDebugDrawDataProcessorUtils.h"
 #include "DataWrappers/ChaosVDDebugShapeDataWrapper.h"
 #include "Visualizers/ChaosVDDebugDrawUtils.h"
 
-
 FChaosVDDebugDrawImplicitObjectDataProcessor::FChaosVDDebugDrawImplicitObjectDataProcessor() : FChaosVDDataProcessorBase(FChaosVDDebugDrawImplicitObjectDataWrapper::WrapperTypeName)
 {
 }
@@ -21,18 +22,18 @@ bool FChaosVDDebugDrawImplicitObjectDataProcessor::ProcessRawData(const TArray<u
 
 	const TSharedPtr<FChaosVDDebugDrawImplicitObjectDataWrapper> DebugDrawData = MakeShared<FChaosVDDebugDrawImplicitObjectDataWrapper>();
 	const bool bSuccess = Chaos::VisualDebugger::ReadDataFromBuffer(InData, *DebugDrawData, ProviderSharedPtr.ToSharedRef());
-	
+
 	if (bSuccess)
 	{
-		if (TSharedPtr<FChaosVDGameFrameData> CurrentFrameData = ProviderSharedPtr->GetCurrentGameFrame().Pin())
-		{
-			DebugDrawData->SolverID = ProviderSharedPtr->GetRemappedSolverID(DebugDrawData->SolverID);
+		using namespace Chaos::VisualDebugger::Utils;
 
-			if (TSharedPtr<FChaosVDDebugShapeDataContainer> DebugDrawShapeData = CurrentFrameData->GetCustomDataHandler().GetOrAddDefaultData<FChaosVDDebugShapeDataContainer>())
-			{	
-				DebugDrawShapeData->RecordedDebugDrawImplicitObjectsBySolverID.FindOrAdd(DebugDrawData->SolverID).Add(DebugDrawData);
-				CurrentFrameData->MarkDirty();
-			}
+		TSharedRef<FChaosVDTraceProvider> ProviderAsSharedRef = ProviderSharedPtr.ToSharedRef();
+		TSharedRef<FChaosVDDebugDrawImplicitObjectDataWrapper> DebugDrawDataAsSharedRef = DebugDrawData.ToSharedRef();
+		RemapDebugDrawShapeDataSolverID(DebugDrawDataAsSharedRef, ProviderAsSharedRef);
+	
+		if (TSharedPtr<FChaosVDDebugShapeDataContainer> DebugDrawShapeData = GetShapeDataContainer(DebugDrawDataAsSharedRef, ProviderAsSharedRef, EShapeDataContainerAccessorFlags::MarkFrameDirty))
+		{
+			DebugDrawShapeData->RecordedDebugDrawImplicitObjects.Add(DebugDrawData);
 		}
 	}
 
diff --git a/Engine/Plugins/ChaosVD/Source/ChaosVDBuiltInExtensions/Private/GenericDebugDraw/Trace/DataProcessors/ChaosVDDebugDrawLineDataProcessor.cpp b/Engine/Plugins/ChaosVD/Source/ChaosVDBuiltInExtensions/Private/GenericDebugDraw/Trace/DataProcessors/ChaosVDDebugDrawLineDataProcessor.cpp
--- a/Engine/Plugins/ChaosVD/Source/ChaosVDBuiltInExtensions/Private/GenericDebugDraw/Trace/DataProcessors/ChaosVDDebugDrawLineDataProcessor.cpp
+++ b/Engine/Plugins/ChaosVD/Source/ChaosVDBuiltInExtensions/Private/GenericDebugDraw/Trace/DataProcessors/ChaosVDDebugDrawLineDataProcessor.cpp
@@ -1,6 +1,8 @@
 // Copyright Epic Games, Inc. All Rights Reserved.
 
 #include "ChaosVDDebugDrawLineDataProcessor.h"
+
+#include "ChaosVDDebugDrawDataProcessorUtils.h"
 #include "DataWrappers/ChaosVDDebugShapeDataWrapper.h"
 #include "Visualizers/ChaosVDDebugDrawUtils.h"
 
@@ -24,15 +26,15 @@ bool FChaosVDDebugDrawLineDataProcessor::ProcessRawData(const TArray<uint8>& InD
 	
 	if (bSuccess)
 	{
-		if (TSharedPtr<FChaosVDGameFrameData> CurrentFrameData = ProviderSharedPtr->GetCurrentGameFrame().Pin())
-		{
-			DebugDrawData->SolverID = ProviderSharedPtr->GetRemappedSolverID(DebugDrawData->SolverID);
+		using namespace Chaos::VisualDebugger::Utils;
 
-			if (TSharedPtr<FChaosVDDebugShapeDataContainer> DebugDrawShapeData = CurrentFrameData->GetCustomDataHandler().GetOrAddDefaultData<FChaosVDDebugShapeDataContainer>())
-			{
-				DebugDrawShapeData->RecordedDebugDrawLinesBySolverID.FindOrAdd(DebugDrawData->SolverID).Add(DebugDrawData);
-				CurrentFrameData->MarkDirty();
-			}
+		TSharedRef<FChaosVDTraceProvider> ProviderAsSharedRef = ProviderSharedPtr.ToSharedRef();
+		TSharedRef<FChaosVDDebugDrawLineDataWrapper> DebugDrawDataAsSharedRef = DebugDrawData.ToSharedRef();
+		RemapDebugDrawShapeDataSolverID(DebugDrawDataAsSharedRef, ProviderAsSharedRef);
+	
+		if (TSharedPtr<FChaosVDDebugShapeDataContainer> DebugDrawShapeData = GetShapeDataContainer(DebugDrawDataAsSharedRef, ProviderAsSharedRef, EShapeDataContainerAccessorFlags::MarkFrameDirty))
+		{
+			DebugDrawShapeData->RecordedDebugDrawLines.Add(DebugDrawData);
 		}
 	}
 
diff --git a/Engine/Plugins/ChaosVD/Source/ChaosVDBuiltInExtensions/Private/GenericDebugDraw/Trace/DataProcessors/ChaosVDDebugDrawSphereDataProcessor.cpp b/Engine/Plugins/ChaosVD/Source/ChaosVDBuiltInExtensions/Private/GenericDebugDraw/Trace/DataProcessors/ChaosVDDebugDrawSphereDataProcessor.cpp
--- a/Engine/Plugins/ChaosVD/Source/ChaosVDBuiltInExtensions/Private/GenericDebugDraw/Trace/DataProcessors/ChaosVDDebugDrawSphereDataProcessor.cpp
+++ b/Engine/Plugins/ChaosVD/Source/ChaosVDBuiltInExtensions/Private/GenericDebugDraw/Trace/DataProcessors/ChaosVDDebugDrawSphereDataProcessor.cpp
@@ -1,6 +1,8 @@
 // Copyright Epic Games, Inc. All Rights Reserved.
 
 #include "ChaosVDDebugDrawSphereDataProcessor.h"
+
+#include "ChaosVDDebugDrawDataProcessorUtils.h"
 #include "DataWrappers/ChaosVDDebugShapeDataWrapper.h"
 #include "Visualizers/ChaosVDDebugDrawUtils.h"
 
@@ -24,15 +26,15 @@ bool FChaosVDDebugDrawSphereDataProcessor::ProcessRawData(const TArray<uint8>& I
 	
 	if (bSuccess)
 	{
-		if (TSharedPtr<FChaosVDGameFrameData> CurrentFrameData = ProviderSharedPtr->GetCurrentGameFrame().Pin())
-		{
-			DebugDrawData->SolverID = ProviderSharedPtr->GetRemappedSolverID(DebugDrawData->SolverID);
+		using namespace Chaos::VisualDebugger::Utils;
+
+		TSharedRef<FChaosVDTraceProvider> ProviderAsSharedRef = ProviderSharedPtr.ToSharedRef();
+		TSharedRef<FChaosVDDebugDrawSphereDataWrapper> DebugDrawDataAsSharedRef = DebugDrawData.ToSharedRef();
+		RemapDebugDrawShapeDataSolverID(DebugDrawDataAsSharedRef, ProviderAsSharedRef);
 	
-			if (TSharedPtr<FChaosVDDebugShapeDataContainer> DebugDrawShapeData = CurrentFrameData->GetCustomDataHandler().GetOrAddDefaultData<FChaosVDDebugShapeDataContainer>())
-			{
-				DebugDrawShapeData->RecordedDebugDrawSpheresBySolverID.FindOrAdd(DebugDrawData->SolverID).Add(DebugDrawData);
-				CurrentFrameData->MarkDirty();
-			}
+		if (TSharedPtr<FChaosVDDebugShapeDataContainer> DebugDrawShapeData = GetShapeDataContainer(DebugDrawDataAsSharedRef, ProviderAsSharedRef, EShapeDataContainerAccessorFlags::MarkFrameDirty))
+		{
+			DebugDrawShapeData->RecordedDebugDrawSpheres.Add(DebugDrawData);
 		}
 	}
 
diff --git a/Engine/Plugins/ChaosVD/Source/ChaosVDBuiltInExtensions/Private/GenericDebugDraw/Visualizers/ChaosVDGenericDebugDrawDataComponentVisualizer.cpp b/Engine/Plugins/ChaosVD/Source/ChaosVDBuiltInExtensions/Private/GenericDebugDraw/Visualizers/ChaosVDGenericDebugDrawDataComponentVisualizer.cpp
--- a/Engine/Plugins/ChaosVD/Source/ChaosVDBuiltInExtensions/Private/GenericDebugDraw/Visualizers/ChaosVDGenericDebugDrawDataComponentVisualizer.cpp
+++ b/Engine/Plugins/ChaosVD/Source/ChaosVDBuiltInExtensions/Private/GenericDebugDraw/Visualizers/ChaosVDGenericDebugDrawDataComponentVisualizer.cpp
@@ -30,6 +30,17 @@ void FChaosVDGenericDebugDrawDataComponentVisualizer::RegisterVisualizerMenus()
 	CreateGenericVisualizerMenu<UChaosVDGenericDebugDrawSettings, EChaosVDGenericDebugDrawVisualizationFlags>(FName("ChaosVDViewportToolbarBase.Show"), MenuSection, MenuSectionLabel, FlagsMenuLabel, FlagsMenuTooltip, FlagsMenuIcon, SettingsMenuLabel, SettingsMenuTooltip);
 }
 
+void FChaosVDGenericDebugDrawDataComponentVisualizer::DrawData(const FSceneView* View, FPrimitiveDrawInterface* PDI, const FChaosVDGenericDebugDrawDataVisualizationSettings& InVisualizationContext)
+{
+	DrawBoxes(PDI, View, InVisualizationContext);
+
+	DrawLines(PDI, View, InVisualizationContext);
+
+	DrawSpheres(PDI, View, InVisualizationContext);
+
+	DrawImplicitObjects(PDI, View, InVisualizationContext);
+}
+
 void FChaosVDGenericDebugDrawDataComponentVisualizer::DrawVisualization(const UActorComponent* Component, const FSceneView* View, FPrimitiveDrawInterface* PDI)
 {
 	const UChaosVDGenericDebugDrawDataComponent* DebugDrawDataComponent = Cast<UChaosVDGenericDebugDrawDataComponent>(Component);
@@ -72,20 +83,21 @@ void FChaosVDGenericDebugDrawDataComponentVisualizer::DrawVisualization(const UA
 		return;
 	}
 
-	DrawBoxes(PDI, View, VisualizationContext);
-
-	DrawLines(PDI, View, VisualizationContext);
+	VisualizationContext.DataSource = EChaosVDDrawDataContainerSource::GameFrame;
+	DrawData(View, PDI, VisualizationContext);
 
-	DrawSpheres(PDI, View, VisualizationContext);
+	VisualizationContext.DataSource = EChaosVDDrawDataContainerSource::SolverFrame;
+	DrawData(View, PDI, VisualizationContext);
 
-	DrawImplicitObjects(PDI, View, VisualizationContext);
+	VisualizationContext.DataSource = EChaosVDDrawDataContainerSource::SolverStage;
+	DrawData(View, PDI, VisualizationContext);
 }
 
 void FChaosVDGenericDebugDrawDataComponentVisualizer::DrawBoxes(FPrimitiveDrawInterface* PDI, const FSceneView* View, const FChaosVDGenericDebugDrawDataVisualizationSettings& InVisualizationContext)
 {
 	if (InVisualizationContext.IsVisualizationFlagEnabled(EChaosVDGenericDebugDrawVisualizationFlags::DrawBoxes))
 	{
-		TConstArrayView<TSharedPtr<FChaosVDDebugDrawBoxDataWrapper>> DebugDrawBoxesView = InVisualizationContext.DataComponent->GetDebugDrawBoxesDataView();
+		TConstArrayView<TSharedPtr<FChaosVDDebugDrawBoxDataWrapper>> DebugDrawBoxesView = InVisualizationContext.DataComponent->GetDebugDrawBoxesDataView(InVisualizationContext.DataSource);
 		for (const TSharedPtr<FChaosVDDebugDrawBoxDataWrapper>& DebugDrawBox : DebugDrawBoxesView)
 		{
 			if (DebugDrawBox)
@@ -111,7 +123,7 @@ void FChaosVDGenericDebugDrawDataComponentVisualizer::DrawLines(FPrimitiveDrawIn
 {
 	if (InVisualizationContext.IsVisualizationFlagEnabled(EChaosVDGenericDebugDrawVisualizationFlags::DrawLines))
 	{
-		TConstArrayView<TSharedPtr<FChaosVDDebugDrawLineDataWrapper>> DebugDrawLinesView = InVisualizationContext.DataComponent->GetDebugDrawLinesDataView();
+		TConstArrayView<TSharedPtr<FChaosVDDebugDrawLineDataWrapper>> DebugDrawLinesView = InVisualizationContext.DataComponent->GetDebugDrawLinesDataView(InVisualizationContext.DataSource);
 		for (const TSharedPtr<FChaosVDDebugDrawLineDataWrapper>& DebugDrawLine : DebugDrawLinesView)
 		{
 			if (DebugDrawLine)
@@ -140,7 +152,7 @@ void FChaosVDGenericDebugDrawDataComponentVisualizer::DrawSpheres(FPrimitiveDraw
 {
 	if (InVisualizationContext.IsVisualizationFlagEnabled(EChaosVDGenericDebugDrawVisualizationFlags::DrawSpheres))
 	{
-		TConstArrayView<TSharedPtr<FChaosVDDebugDrawSphereDataWrapper>> DebugDrawSpheres = InVisualizationContext.DataComponent->GetDebugDrawSpheresDataView();
+		TConstArrayView<TSharedPtr<FChaosVDDebugDrawSphereDataWrapper>> DebugDrawSpheres = InVisualizationContext.DataComponent->GetDebugDrawSpheresDataView(InVisualizationContext.DataSource);
 		for (const TSharedPtr<FChaosVDDebugDrawSphereDataWrapper>& DebugDrawSphere : DebugDrawSpheres)
 		{
 			if (DebugDrawSphere)
@@ -172,7 +184,7 @@ void FChaosVDGenericDebugDrawDataComponentVisualizer::DrawImplicitObjects(FPrimi
 	
 	if (InVisualizationContext.IsVisualizationFlagEnabled(EChaosVDGenericDebugDrawVisualizationFlags::DrawImplicitObjects))
 	{
-		TConstArrayView<TSharedPtr<FChaosVDDebugDrawImplicitObjectDataWrapper>> DebugDrawImplicitObjects = InVisualizationContext.DataComponent->GetDebugDrawImplicitObjectsDataView();
+		TConstArrayView<TSharedPtr<FChaosVDDebugDrawImplicitObjectDataWrapper>> DebugDrawImplicitObjects = InVisualizationContext.DataComponent->GetDebugDrawImplicitObjectsDataView(InVisualizationContext.DataSource);
 		for (const TSharedPtr<FChaosVDDebugDrawImplicitObjectDataWrapper>& DebugDrawImplicitObjectData : DebugDrawImplicitObjects)
 		{
 			if (DebugDrawImplicitObjectData)
diff --git a/Engine/Plugins/ChaosVD/Source/ChaosVDBuiltInExtensions/Private/GenericDebugDraw/Visualizers/ChaosVDGenericDebugDrawDataComponentVisualizer.h b/Engine/Plugins/ChaosVD/Source/ChaosVDBuiltInExtensions/Private/GenericDebugDraw/Visualizers/ChaosVDGenericDebugDrawDataComponentVisualizer.h
--- a/Engine/Plugins/ChaosVD/Source/ChaosVDBuiltInExtensions/Private/GenericDebugDraw/Visualizers/ChaosVDGenericDebugDrawDataComponentVisualizer.h
+++ b/Engine/Plugins/ChaosVD/Source/ChaosVDBuiltInExtensions/Private/GenericDebugDraw/Visualizers/ChaosVDGenericDebugDrawDataComponentVisualizer.h
@@ -1,9 +1,11 @@
 // Copyright Epic Games, Inc. All Rights Reserved.
 
 #pragma once
+#include "GenericDebugDraw/Components/ChaosVDGenericDebugDrawDataComponent.h"
 #include "Visualizers/ChaosVDComponentVisualizerBase.h"
 #include "GenericDebugDraw/Settings/ChaosVDGenericDebugDrawSettings.h"
 
+enum class EChaosVDDrawDataContainerSource : uint8;
 struct FChaosVDDebugDrawImplicitObjectDataWrapper;
 class UChaosVDGenericDebugDrawDataComponent;
 class UChaosVDGenericDebugDrawSettings;
@@ -17,6 +19,8 @@ struct FChaosVDGenericDebugDrawDataVisualizationSettings : public FChaosVDVisual
 
 	bool bShowDebugText = false;
 
+	EChaosVDDrawDataContainerSource DataSource = EChaosVDDrawDataContainerSource::GameFrame;
+
 	const UChaosVDGenericDebugDrawDataComponent* DataComponent = nullptr;
 
 	bool IsVisualizationFlagEnabled(EChaosVDGenericDebugDrawVisualizationFlags Flag) const
@@ -37,6 +41,8 @@ public:
 	virtual bool CanHandleClick(const HChaosVDComponentVisProxy& VisProxy) override { return false;}
 
 	virtual void DrawVisualization(const UActorComponent* Component, const FSceneView* View, FPrimitiveDrawInterface* PDI) override;
+	
+	void DrawData(const FSceneView* View, FPrimitiveDrawInterface* PDI, const FChaosVDGenericDebugDrawDataVisualizationSettings& InVisualizationContext);
 
 	void DrawBoxes(FPrimitiveDrawInterface* PDI, const FSceneView* View, const FChaosVDGenericDebugDrawDataVisualizationSettings& InVisualizationContext);
     void DrawLines(FPrimitiveDrawInterface* PDI, const FSceneView* View, const FChaosVDGenericDebugDrawDataVisualizationSettings& InVisualizationContext);
diff --git a/Engine/Source/Runtime/Core/Public/UObject/FortniteMainBranchObjectVersions.inl b/Engine/Source/Runtime/Core/Public/UObject/FortniteMainBranchObjectVersions.inl
--- a/Engine/Source/Runtime/Core/Public/UObject/FortniteMainBranchObjectVersions.inl
+++ b/Engine/Source/Runtime/Core/Public/UObject/FortniteMainBranchObjectVersions.inl
@@ -628,5 +628,8 @@ FORTNITE_MAIN_VERSION(ParticleInflatedBoundsInChaosVisualDebugger, 204)
 // Migrate properties from FLandscapeLayer to ULandscapeEditLayer
 FORTNITE_MAIN_VERSION(MigrateLandscapeEditLayerProperties, 205)
 
+// Added more context data to CVD's traced shapes so we can play it back at the solver stage level (not just game thread frames) 
+FORTNITE_MAIN_VERSION(ThreadContextDataInChaosVisualDebuggerDebugDrawData, 206)
+
 // -----<new versions can be added above this line>-------------------------------------------------
 #undef FORTNITE_MAIN_VERSION
\ No newline at end of file
diff --git a/Engine/Source/Runtime/Experimental/Chaos/Private/ChaosVisualDebugger/ChaosVisualDebuggerTrace.cpp b/Engine/Source/Runtime/Experimental/Chaos/Private/ChaosVisualDebugger/ChaosVisualDebuggerTrace.cpp
--- a/Engine/Source/Runtime/Experimental/Chaos/Private/ChaosVisualDebugger/ChaosVisualDebuggerTrace.cpp
+++ b/Engine/Source/Runtime/Experimental/Chaos/Private/ChaosVisualDebugger/ChaosVisualDebuggerTrace.cpp
@@ -947,6 +947,24 @@ void FChaosVisualDebuggerTrace::TraceNetworkTickOffset(int32 TickOffset, int32 S
 	};
 }
 
+bool FChaosVisualDebuggerTrace::CanTraceDebugDrawShape(int32& OutSolverID)
+{
+	using namespace Chaos::VisualDebugger::Utils;
+	
+	if (const FChaosVDContext* CVDContextData = FChaosVDThreadContext::Get().GetCurrentContext(EChaosVDContextType::Solver))
+	{
+		if (!IsContextEnabledAndValid(CVDContextData))
+		{
+			return false;
+		}
+	
+		OutSolverID = OutSolverID == INDEX_NONE ? CVDContextData->Id : OutSolverID;
+
+		return true;
+	}
+	return true;
+}
+
 void FChaosVisualDebuggerTrace::TraceDebugDrawBox(const FBox& InBox, FName Tag, FColor Color, int32 SolverID)
 {
 	if (!IsTracing())
@@ -954,11 +972,18 @@ void FChaosVisualDebuggerTrace::TraceDebugDrawBox(const FBox& InBox, FName Tag,
 		return;
 	}
 
+	// Generic Debug Draw might not have a context if they are recorded from the game thread in game code, in that case we allow the trace anyway
+	if (!CanTraceDebugDrawShape(SolverID))
+	{
+		return;
+	}
+
 	FChaosVDDebugDrawBoxDataWrapper DataWrapper;
 	DataWrapper.SolverID = SolverID;
 	DataWrapper.Tag = Tag;
 	DataWrapper.Color = Color;
 	DataWrapper.Box = InBox;
+	DataWrapper.ThreadContext = IsInGameThread() ? EChaosVDParticleContext::GameThread : EChaosVDParticleContext::PhysicsThread;
 
 	DataWrapper.MarkAsValid();
 
@@ -975,12 +1000,19 @@ void FChaosVisualDebuggerTrace::TraceDebugDrawLine(const FVector& InStartLocatio
 		return;
 	}
 
+	// Generic Debug Draw might not have a context if they are recorded from the game thread in game code, in that case we allow the trace anyway
+	if (!CanTraceDebugDrawShape(SolverID))
+	{
+		return;
+	}
+
 	FChaosVDDebugDrawLineDataWrapper DataWrapper;
 	DataWrapper.SolverID = SolverID;
 	DataWrapper.Tag = Tag;
 	DataWrapper.Color = Color;
 	DataWrapper.StartLocation = InStartLocation;
 	DataWrapper.EndLocation = InEndLocation;
+	DataWrapper.ThreadContext = IsInGameThread() ? EChaosVDParticleContext::GameThread : EChaosVDParticleContext::PhysicsThread;
 
 	DataWrapper.MarkAsValid();
 
@@ -996,6 +1028,12 @@ void FChaosVisualDebuggerTrace::TraceDebugDrawVector(const FVector& InStartLocat
 	{
 		return;
 	}
+	
+	// Generic Debug Draw might not have a context if they are recorded from the game thread in game code, in that case we allow the trace anyway
+	if (!CanTraceDebugDrawShape(SolverID))
+	{
+		return;
+	}
 
 	FChaosVDDebugDrawLineDataWrapper DataWrapper;
 	DataWrapper.SolverID = SolverID;
@@ -1004,6 +1042,7 @@ void FChaosVisualDebuggerTrace::TraceDebugDrawVector(const FVector& InStartLocat
 	DataWrapper.StartLocation = InStartLocation;
 	DataWrapper.EndLocation = InStartLocation + InVector;
 	DataWrapper.bIsArrow = true;
+	DataWrapper.ThreadContext = IsInGameThread() ? EChaosVDParticleContext::GameThread : EChaosVDParticleContext::PhysicsThread;
 
 	DataWrapper.MarkAsValid();
 
@@ -1020,12 +1059,19 @@ void FChaosVisualDebuggerTrace::TraceDebugDrawSphere(const FVector& Center, floa
 		return;
 	}
 
+	// Generic Debug Draw might not have a context if they are recorded from the game thread in game code, in that case we allow the trace anyway
+	if (!CanTraceDebugDrawShape(SolverID))
+	{
+		return;
+	}
+
 	FChaosVDDebugDrawSphereDataWrapper DataWrapper;
 	DataWrapper.SolverID = SolverID;
 	DataWrapper.Tag = Tag;
 	DataWrapper.Color = Color;
 	DataWrapper.Origin = Center;
 	DataWrapper.Radius = Radius;
+	DataWrapper.ThreadContext = IsInGameThread() ? EChaosVDParticleContext::GameThread : EChaosVDParticleContext::PhysicsThread;
 
 	DataWrapper.MarkAsValid();
 
@@ -1041,12 +1087,19 @@ void FChaosVisualDebuggerTrace::TraceDebugDrawImplicitObject(const Chaos::FImpli
 	{
 		return;
 	}
+	
+	// Generic Debug Draw might not have a context if they are recorded from the game thread in game code, in that case we allow the trace anyway
+	if (!CanTraceDebugDrawShape(SolverID))
+	{
+		return;
+	}
 
 	FChaosVDDebugDrawImplicitObjectDataWrapper DataWrapper;
 	DataWrapper.SolverID = SolverID;
 	DataWrapper.Tag = Tag;
 	DataWrapper.Color = Color;
 	DataWrapper.ParentTransform = InParentTransform;
+	DataWrapper.ThreadContext = IsInGameThread() ? EChaosVDParticleContext::GameThread : EChaosVDParticleContext::PhysicsThread;
 
 	const uint32 GeometryHash = GeometryTracerObject.GetGeometryHashForImplicit(Implicit);
 	TraceImplicitObject({ GeometryHash, const_cast<Chaos::FImplicitObject*>(Implicit) });
diff --git a/Engine/Source/Runtime/Experimental/Chaos/Public/ChaosVisualDebugger/ChaosVisualDebuggerTrace.h b/Engine/Source/Runtime/Experimental/Chaos/Public/ChaosVisualDebugger/ChaosVisualDebuggerTrace.h
--- a/Engine/Source/Runtime/Experimental/Chaos/Public/ChaosVisualDebugger/ChaosVisualDebuggerTrace.h
+++ b/Engine/Source/Runtime/Experimental/Chaos/Public/ChaosVisualDebugger/ChaosVisualDebuggerTrace.h
@@ -417,6 +417,14 @@ private:
 	 * @param DebugNameID Name to use as ID
 	 */
 	static void TraceNonSolverTransform(const FTransform& InTransform, FStringView DebugNameID);
+
+	/**
+	 * Evaluated the current thread context to determine if we can trace the requested shape. It also updates the solver ID passed as reference
+	 * if it was invalid, but we have a valid cvd context
+	 * @param OutSolverID Solver ID to update if needed using the CVD context data
+	 * @return True if we can trace shapes
+	 */
+	static bool CanTraceDebugDrawShape(int32& OutSolverID);
 	
 	/** Resets the state of the CVD Tracer */
 	static void Reset();
diff --git a/Engine/Source/Runtime/Experimental/ChaosVDData/Public/ChaosVDRecording.h b/Engine/Source/Runtime/Experimental/ChaosVDData/Public/ChaosVDRecording.h
--- a/Engine/Source/Runtime/Experimental/ChaosVDData/Public/ChaosVDRecording.h
+++ b/Engine/Source/Runtime/Experimental/ChaosVDData/Public/ChaosVDRecording.h
@@ -29,62 +29,6 @@
 
 DECLARE_MULTICAST_DELEGATE_TwoParams(FChaosVDGeometryDataLoaded, const Chaos::FConstImplicitObjectPtr&, const uint32 GeometryID)
 
-/** Set of flags used to define characteristics of a loaded solver stage */
-enum class EChaosVDSolverStageFlags : uint8
-{
-	None = 0,
-	/** Set if the solver stage is open and can take new data */
-	Open = 1 << 0,
-	/** Set if the solver stage was explicitly recorded - If not set, this stage was created on the fly during load */
-	ExplicitStage = 1 << 1,
-};
-
-ENUM_CLASS_FLAGS(EChaosVDSolverStageFlags)
-
-
-struct FChaosVDFrameStageData
-{
-	FString StepName;
-	TArray<TSharedPtr<FChaosVDParticleDataWrapper>> RecordedParticlesData;
-	TArray<TSharedPtr<FChaosVDParticlePairMidPhase>> RecordedMidPhases;
-	TArray<TSharedPtr<FChaosVDJointConstraint>> RecordedJointConstraints;
-	TArray<FChaosVDConstraint> RecordedConstraints;
-	TMap<int32, TArray<FChaosVDConstraint>> RecordedConstraintsByParticleID;
-	TMap<int32, TArray<TSharedPtr<FChaosVDParticlePairMidPhase>>> RecordedMidPhasesByParticleID;
-	TSet<int32> ParticlesDestroyedIDs;
-
-	// Used for de-duplication during trace analysis
-	TMap<int32, int32> CurrentRecordedParticlesIndexes;
-
-	EChaosVDSolverStageFlags StageFlags = EChaosVDSolverStageFlags::None;
-};
-
-UE_DEPRECATED(5.6, "FChaosVDStepData is deprecated and will be removed. Please use FChaosVDFrameStageData instead.")
-typedef FChaosVDFrameStageData FChaosVDStepData;
-
-struct FChaosVDTrackedLocation
-{
-	FString DebugName;
-	FVector Location;
-};
-
-struct FChaosVDTrackedTransform
-{
-	FString DebugName;
-	FTransform Transform;
-};
-
-enum class EChaosVDNetworkSyncDataRequirements
-{
-	None = 0,
-	InternalFrameNumber  = 1 << 0,
-	NetworkTickOffset = 1 << 1,
-
-	All = InternalFrameNumber | NetworkTickOffset
-};
-
-ENUM_CLASS_FLAGS(EChaosVDNetworkSyncDataRequirements)
-
 /** Handle to user-defined data types in a CVD recorded frame.
  * These handles share ownership of the data
  */
@@ -180,6 +124,74 @@
 	TMap<FName, FChaosVDCustomUserDataHandle> CustomDataHandlesByType;
 };
 
+/** Set of flags used to define characteristics of a loaded solver stage */
+enum class EChaosVDSolverStageFlags : uint8
+{
+	None = 0,
+	/** Set if the solver stage is open and can take new data */
+	Open = 1 << 0,
+	/** Set if the solver stage was explicitly recorded - If not set, this stage was created on the fly during load */
+	ExplicitStage = 1 << 1,
+};
+
+ENUM_CLASS_FLAGS(EChaosVDSolverStageFlags)
+
+struct FChaosVDFrameStageData
+{
+	FString StepName;
+	TArray<TSharedPtr<FChaosVDParticleDataWrapper>> RecordedParticlesData;
+	TArray<TSharedPtr<FChaosVDParticlePairMidPhase>> RecordedMidPhases;
+	TArray<TSharedPtr<FChaosVDJointConstraint>> RecordedJointConstraints;
+	TArray<FChaosVDConstraint> RecordedConstraints;
+	TMap<int32, TArray<FChaosVDConstraint>> RecordedConstraintsByParticleID;
+	TMap<int32, TArray<TSharedPtr<FChaosVDParticlePairMidPhase>>> RecordedMidPhasesByParticleID;
+	TSet<int32> ParticlesDestroyedIDs;
+
+	// Used for de-duplication during trace analysis
+	TMap<int32, int32> CurrentRecordedParticlesIndexes;
+
+	EChaosVDSolverStageFlags StageFlags = EChaosVDSolverStageFlags::None;
+
+	FChaosVDCustomFrameData& GetCustomDataHandler()
+	{
+		return CustomData;
+	}
+	
+	const FChaosVDCustomFrameData& GetCustomDataHandler() const
+	{
+		return CustomData;
+	}
+
+private:
+	FChaosVDCustomFrameData CustomData;
+};
+
+UE_DEPRECATED(5.6, "FChaosVDStepData is deprecated and will be removed. Please use FChaosVDFrameStageData instead.")
+typedef FChaosVDFrameStageData FChaosVDStepData;
+
+struct FChaosVDTrackedLocation
+{
+	FString DebugName;
+	FVector Location;
+};
+
+struct FChaosVDTrackedTransform
+{
+	FString DebugName;
+	FTransform Transform;
+};
+
+enum class EChaosVDNetworkSyncDataRequirements
+{
+	None = 0,
+	InternalFrameNumber  = 1 << 0,
+	NetworkTickOffset = 1 << 1,
+
+	All = InternalFrameNumber | NetworkTickOffset
+};
+
+ENUM_CLASS_FLAGS(EChaosVDNetworkSyncDataRequirements)
+
 UE_DEPRECATED(5.6, "This FChaosVDStepsContainer is deprecated and it will be removed in the future. Use FChaosVDStagesContainer instead.")
 typedef TArray<FChaosVDFrameStageData, TInlineAllocator<16>> FChaosVDStepsContainer;
 
diff --git a/Engine/Source/Runtime/Experimental/ChaosVisualDebugger/Private/DataWrappers/ChaosVDDebugShapeDataWrapper.cpp b/Engine/Source/Runtime/Experimental/ChaosVisualDebugger/Private/DataWrappers/ChaosVDDebugShapeDataWrapper.cpp
--- a/Engine/Source/Runtime/Experimental/ChaosVisualDebugger/Private/DataWrappers/ChaosVDDebugShapeDataWrapper.cpp
+++ b/Engine/Source/Runtime/Experimental/ChaosVisualDebugger/Private/DataWrappers/ChaosVDDebugShapeDataWrapper.cpp
@@ -2,11 +2,19 @@
 
 #include "DataWrappers/ChaosVDDebugShapeDataWrapper.h"
 
+#include "UObject/FortniteMainBranchObjectVersion.h"
+
 void FChaosVDDebugDrawShapeBase::SerializeBase_Internal(FArchive& Ar)
 {
 	Ar << SolverID;
 	Ar << Tag;
 	Ar << Color;
+
+	Ar.UsingCustomVersion(FFortniteMainBranchObjectVersion::GUID);
+	if (Ar.CustomVer(FFortniteMainBranchObjectVersion::GUID) >= FFortniteMainBranchObjectVersion::ThreadContextDataInChaosVisualDebuggerDebugDrawData)
+	{
+		Ar << ThreadContext;
+	}
 }
 
 bool FChaosVDDebugDrawBoxDataWrapper::Serialize(FArchive& Ar)
diff --git a/Engine/Source/Runtime/Experimental/ChaosVisualDebugger/Public/DataWrappers/ChaosVDDebugShapeDataWrapper.h b/Engine/Source/Runtime/Experimental/ChaosVisualDebugger/Public/DataWrappers/ChaosVDDebugShapeDataWrapper.h
--- a/Engine/Source/Runtime/Experimental/ChaosVisualDebugger/Public/DataWrappers/ChaosVDDebugShapeDataWrapper.h
+++ b/Engine/Source/Runtime/Experimental/ChaosVisualDebugger/Public/DataWrappers/ChaosVDDebugShapeDataWrapper.h
@@ -11,10 +11,18 @@ struct FChaosVDDebugShapeDataContainer
 {
 	GENERATED_BODY()
 
-	TMap<int32,TArray<TSharedPtr<FChaosVDDebugDrawBoxDataWrapper>>> RecordedDebugDrawBoxesBySolverID;
-	TMap<int32,TArray<TSharedPtr<FChaosVDDebugDrawLineDataWrapper>>> RecordedDebugDrawLinesBySolverID;
-	TMap<int32,TArray<TSharedPtr<FChaosVDDebugDrawSphereDataWrapper>>> RecordedDebugDrawSpheresBySolverID;
-	TMap<int32,TArray<TSharedPtr<FChaosVDDebugDrawImplicitObjectDataWrapper>>> RecordedDebugDrawImplicitObjectsBySolverID;
+	TArray<TSharedPtr<FChaosVDDebugDrawBoxDataWrapper>> RecordedDebugDrawBoxes;
+	TArray<TSharedPtr<FChaosVDDebugDrawLineDataWrapper>> RecordedDebugDrawLines;
+	TArray<TSharedPtr<FChaosVDDebugDrawSphereDataWrapper>> RecordedDebugDrawSpheres;
+	TArray<TSharedPtr<FChaosVDDebugDrawImplicitObjectDataWrapper>> RecordedDebugDrawImplicitObjects;
+};
+
+USTRUCT()
+struct FChaosVDMultiSolverDebugShapeDataContainer
+{
+	GENERATED_BODY()
+
+	TMap<int32, TSharedPtr<FChaosVDDebugShapeDataContainer>> DataBySolverID;
 };
 
 USTRUCT()
@@ -31,6 +39,9 @@ struct FChaosVDDebugDrawShapeBase : public FChaosVDWrapperDataBase
 	UPROPERTY()
 	FColor Color = FColor::Blue;
 
+	UPROPERTY()
+	EChaosVDParticleContext ThreadContext = EChaosVDParticleContext::GameThread;
+
 protected:
 	CHAOSVDRUNTIME_API void SerializeBase_Internal(FArchive& Ar);
 };
