[ChaosVD] Fix an issue where the scene outliner might end up showing particles in the incorrect sub folder or with an incorrect particle parent.

CVD uses the solver's particle id as a way to keep track of particles during playback. Particles ID can be re-used in the solver, therefore CVD has a system to detect when that happen and update the visualization accordingly.

There was a bug where we were not properly detecting when we should re-evaluate the particle parent data when the particle data itself changed.

diff --git a/Engine/Plugins/ChaosVD/Source/ChaosVD/Private/ChaosVDSceneParticle.cpp b/Engine/Plugins/ChaosVD/Source/ChaosVD/Private/ChaosVDSceneParticle.cpp
--- a/Engine/Plugins/ChaosVD/Source/ChaosVD/Private/ChaosVDSceneParticle.cpp
+++ b/Engine/Plugins/ChaosVD/Source/ChaosVD/Private/ChaosVDSceneParticle.cpp
@@ -126,17 +126,19 @@ int32 FChaosVDSceneParticle::GetStreamingID() const
 
 void FChaosVDSceneParticle::UpdateParent(const TSharedPtr<FChaosVDParticleDataWrapper>& InRecordedData)
 {
-	TSharedPtr<FChaosVDScene> ScenePtr = SceneWeakPtr.Pin();
-
-	if (ScenePtr && InRecordedData->ParticleCluster.HasValidData())
+	if (TSharedPtr<FChaosVDScene> ScenePtr = SceneWeakPtr.Pin())
 	{
-		if (TSharedPtr<FChaosVDSceneParticle> ParentParticle = ScenePtr->GetParticleInstance(InRecordedData->SolverID, InRecordedData->ParticleCluster.ParentParticleID))
+		TSharedPtr<FChaosVDSceneParticle> ParentClusterParticle = InRecordedData->ParticleCluster.HasValidData() ?
+																	ScenePtr->GetParticleInstance(InRecordedData->SolverID, InRecordedData->ParticleCluster.ParentParticleID)
+																	: nullptr;
+		if (ParentClusterParticle)
 		{
-			SetParent(ParentParticle);
+			SetParent(ParentClusterParticle);
 		}
-		else if (AChaosVDSolverInfoActor* SolverData = ScenePtr->GetSolverInfoActor(InRecordedData->SolverID))
+		else
 		{
-			if (UChaosVDParticleDataComponent* ParticleDataComponent = SolverData->GetParticleDataComponent())
+			AChaosVDSolverInfoActor* SolverData = ScenePtr->GetSolverInfoActor(InRecordedData->SolverID);
+			if (UChaosVDParticleDataComponent* ParticleDataComponent = SolverData ? SolverData->GetParticleDataComponent() : nullptr)
 			{
 				SetParent(ParticleDataComponent->GetParticleContainerByType(InRecordedData->Type));
 			}
@@ -145,10 +147,6 @@ void FChaosVDSceneParticle::UpdateParent(const TSharedPtr<FChaosVDParticleDataWr
 				SetParent(nullptr);
 			}
 		}
-		else
-		{
-			SetParent(nullptr);
-		}
 	}
 
 	EnumAddFlags(DirtyFlags, EChaosVDSceneParticleDirtyFlags::TEDS);
@@ -162,10 +160,11 @@ void FChaosVDSceneParticle::PreUpdateFromRecordedParticleData(const TSharedPtr<F
 	{
 		return;	
 	}
-	
+
 	if (!ParticleDataPtr
 		|| ParticleDataPtr->ParticleCluster.HasValidData() != InRecordedData->ParticleCluster.HasValidData()
-		|| ParticleDataPtr->ParticleCluster.ParentParticleID != InRecordedData->ParticleCluster.ParentParticleID)
+		|| ParticleDataPtr->ParticleCluster.ParentParticleID != InRecordedData->ParticleCluster.ParentParticleID
+		|| ParticleDataPtr->Type != InRecordedData->Type)
 	{
 		EnumAddFlags(DirtyFlags, EChaosVDSceneParticleDirtyFlags::Parent);
 	}
