[ChaosVD] Added support to combine multiple loaded trace files into one for easy sharing.

This CL:
- Adds a new button in the UI that gets enabled when multiple CVD sessions are loaded. If the button is clicked save file dialog opens and allows you to choose the path and name for the combined file.
- Created a new file handle type that allows other systems (like UE Trace), work with the combined CVD file as it wasn't combined.
- Implemented a minimal "cvdmulti" format, which is just the raw data of multiple utrace files combined into one + a file table to be able to address the data of each individual file.
- Fixed a bug where live recordings started from the UI in PIE cannot be stopped because the stop button stays disabled

diff --git a/Engine/Plugins/ChaosVD/Source/ChaosVD/Private/ChaosVDEngine.cpp b/Engine/Plugins/ChaosVD/Source/ChaosVD/Private/ChaosVDEngine.cpp
--- a/Engine/Plugins/ChaosVD/Source/ChaosVD/Private/ChaosVDEngine.cpp
+++ b/Engine/Plugins/ChaosVD/Source/ChaosVD/Private/ChaosVDEngine.cpp
@@ -4,12 +4,12 @@
 
 #include "ChaosVDModule.h"
 #include "ChaosVDPlaybackController.h"
-#include "ChaosVDRuntimeModule.h"
 #include "ChaosVDScene.h"
 #include "ChaosVDSettingsManager.h"
 #include "Chaos/ChaosVDEngineEditorBridge.h"
 #include "ChaosVisualDebugger/ChaosVDOptionalDataChannel.h"
 #include "Settings/ChaosVDMiscSettings.h"
+#include "Trace/ChaosVDCombinedTraceFile.h"
 #include "Trace/ChaosVDTraceManager.h"
 
 void FChaosVDEngine::Initialize()
@@ -194,13 +194,13 @@ bool FChaosVDEngine::HasAnyLiveSessionActive() const
 	return PlaybackController && PlaybackController->IsPlayingLiveSession();
 }
 
-void FChaosVDEngine::LoadRecording(const FString& FilePath, EChaosVDLoadRecordedDataMode LoadingMode)
+void FChaosVDEngine::LoadRecording_Internal(const TFunction<FString(const TSharedPtr<FChaosVDRecording>&)>& LoadCallback, EChaosVDLoadRecordedDataMode LoadingMode)
 {
 	FChaosVDTraceSessionDescriptor NewSessionFromFileDescriptor;
 
+	TSharedPtr<FChaosVDRecording> ExistingRecordingInstance = nullptr;
 	if (LoadingMode == EChaosVDLoadRecordedDataMode::MultiSource && !CurrentSessionDescriptors.IsEmpty())
 	{
-		TSharedPtr<FChaosVDRecording> ExistingRecordingInstance = nullptr;
 		if (const TSharedPtr<const TraceServices::IAnalysisSession> TraceSession = FChaosVDModule::Get().GetTraceManager()->GetSession(CurrentSessionDescriptors[0].SessionName))
 		{
 			if (const FChaosVDTraceProvider* ChaosVDProvider = TraceSession->ReadProvider<FChaosVDTraceProvider>(FChaosVDTraceProvider::ProviderName))
@@ -208,23 +208,68 @@ void FChaosVDEngine::LoadRecording(const FString& FilePath, EChaosVDLoadRecorded
 				ExistingRecordingInstance = ChaosVDProvider->GetRecordingForSession();
 			}
 		}
-
-		NewSessionFromFileDescriptor.SessionName = FChaosVDModule::Get().GetTraceManager()->LoadTraceFile(FilePath, ExistingRecordingInstance);
-	}
-	else
-	{
-		NewSessionFromFileDescriptor.SessionName = FChaosVDModule::Get().GetTraceManager()->LoadTraceFile(FilePath);
 	}
 	
+	NewSessionFromFileDescriptor.SessionName = LoadCallback(ExistingRecordingInstance);
+	
 	NewSessionFromFileDescriptor.bIsLiveSession = false;
 
 	if (NewSessionFromFileDescriptor.IsValid())
 	{
-		UpdateRecentFilesList(FilePath);
 		OpenSession(NewSessionFromFileDescriptor, LoadingMode);
 	}
 }
 
+void FChaosVDEngine::LoadRecording(const FString& FilePath, EChaosVDLoadRecordedDataMode LoadingMode)
+{
+	if (FilePath.EndsWith(TEXT("cvdmulti")))
+	{
+		LoadCombinedMultiRecording(FilePath);
+		return;
+	}
+
+	auto LoadSingleFileFromTraceManager = [&FilePath](const TSharedPtr<FChaosVDRecording>& ExistingRecordingInstance)
+	{
+		return FChaosVDModule::Get().GetTraceManager()->LoadTraceFile(FilePath, ExistingRecordingInstance);
+	};
+	
+	LoadRecording_Internal(LoadSingleFileFromTraceManager, LoadingMode);
+	UpdateRecentFilesList(FilePath);
+}
+
+void FChaosVDEngine::LoadCombinedMultiRecording(const FString& FilePath)
+{
+	if (!ensure(FilePath.EndsWith(TEXT("cvdmulti"))))
+	{
+		return;
+	}
+
+	using namespace Chaos::VisualDebugger;
+
+	TArray<TUniquePtr<IFileHandle>> ExtractedHandles = CombinedTraceFile::GetInnerFileHandles(FilePath);
+	
+	if (!ensure(!ExtractedHandles.IsEmpty()))
+	{
+		return;
+	}
+
+	int32 CurrentFileIndex = 0;
+	for (TUniquePtr<IFileHandle>& Handle : ExtractedHandles)
+	{
+		const FStringFormatOrderedArguments Args {FilePath, FString::FromInt(CurrentFileIndex)};
+		FString SessionName = FString::Format(TEXT("{0}-{1}"), Args);
+		auto LoadSingleFileHandleFromTraceManager = [&Handle, &SessionName](const TSharedPtr<FChaosVDRecording>& ExistingRecordingInstance)
+		{
+			return FChaosVDModule::Get().GetTraceManager()->LoadTraceFile(MoveTemp(Handle), SessionName, ExistingRecordingInstance);
+		};
+
+		LoadRecording_Internal(LoadSingleFileHandleFromTraceManager, EChaosVDLoadRecordedDataMode::MultiSource);
+		CurrentFileIndex++;
+	}
+	
+	UpdateRecentFilesList(FilePath);
+}
+
 bool FChaosVDEngine::ConnectToLiveSession(uint32 SessionID, const FString& InSessionAddress, EChaosVDLoadRecordedDataMode LoadingMode)
 {
 	FChaosVDTraceSessionDescriptor NewSessionFromFileDescriptor;
@@ -267,6 +312,72 @@ void FChaosVDEngine::OpenSession(const FChaosVDTraceSessionDescriptor& SessionDe
 	OnSessionOpened().Broadcast(SessionDescriptor);
 }
 
+bool FChaosVDEngine::SaveOpenSessionToCombinedFile(const FString& InTargetFilePath)
+{
+	if (!ensure(CurrentSessionDescriptors.Num() > 0))
+	{
+		return false;
+	}
+
+	IPlatformFile& FileSystem = IPlatformFile::GetPlatformPhysical();
+
+	TArray<TUniquePtr<IFileHandle>> FilePathsToCombine;
+	FilePathsToCombine.Reserve(CurrentSessionDescriptors.Num());
+
+	for (FChaosVDTraceSessionDescriptor& SessionDescriptor : CurrentSessionDescriptors)
+	{
+		if (SessionDescriptor.SessionName.EndsWith(TEXT("cvdmulti")))
+		{
+			UE_LOG(LogChaosVDEditor, Error, TEXT("Failed to combine files, because one of the files is a multi cvd file and that is not supported yet | [%s]"), *SessionDescriptor.SessionName)
+			return false;
+		}
+	
+		IFileHandle* ContainerFileHandle = FileSystem.OpenRead(*SessionDescriptor.SessionName);
+		if (!ContainerFileHandle)
+		{
+			UE_LOG(LogChaosVDEditor, Error, TEXT("Failed to combine files, because one fo the files cannot be open | [%s]"), *SessionDescriptor.SessionName)
+			return false;
+		}
+
+		FilePathsToCombine.Emplace(TUniquePtr<IFileHandle>(ContainerFileHandle));
+	}
+
+	FString FinalFilePath = InTargetFilePath;
+
+	if (FinalFilePath.IsEmpty())
+	{
+		FString PathPart;
+		FString FilenamePart;
+		FString ExtensionPart;
+		FPaths::Split(CurrentSessionDescriptors[0].SessionName, PathPart, FilenamePart, ExtensionPart);
+
+		FStringFormatOrderedArguments NameArgs { FilenamePart, CurrentSessionDescriptors.Num(),FDateTime::Now().ToString(TEXT("%Y%m%d_%H%M%S")) };
+		FinalFilePath = FPaths::Combine(PathPart, FString::Format(TEXT("ChaosVD-{0}-Combined-{1}-Sessions-{2}.cvdmulti"), NameArgs));
+	}
+	
+	using namespace Chaos::VisualDebugger::CombinedTraceFile;
+
+	return ensure(CombineFiles(FilePathsToCombine, FinalFilePath));
+}
+
+bool FChaosVDEngine::CanCombineOpenSessions() const
+{
+	if (CurrentSessionDescriptors.Num() > 1)
+	{
+		for (const FChaosVDTraceSessionDescriptor& Session : CurrentSessionDescriptors)
+		{
+			if (Session.SessionName.EndsWith(TEXT("cvdmulti")))
+			{
+				return false;
+			}
+		}
+
+		return true;
+	}
+
+	return false;
+}
+
 bool FChaosVDEngine::Tick(float DeltaTime)
 {
 	return true;
diff --git a/Engine/Plugins/ChaosVD/Source/ChaosVD/Private/ChaosVDEngine.h b/Engine/Plugins/ChaosVD/Source/ChaosVD/Private/ChaosVDEngine.h
--- a/Engine/Plugins/ChaosVD/Source/ChaosVD/Private/ChaosVDEngine.h
+++ b/Engine/Plugins/ChaosVD/Source/ChaosVD/Private/ChaosVDEngine.h
@@ -71,9 +71,15 @@ public:
 	bool HasAnyLiveSessionActive() const;
 
 	void LoadRecording(const FString& FilePath, EChaosVDLoadRecordedDataMode LoadingMode = EChaosVDLoadRecordedDataMode::MultiSource);
+	void LoadCombinedMultiRecording(const FString& FilePath);
+
 	bool ConnectToLiveSession(uint32 SessionID, const FString& InSessionAddress, EChaosVDLoadRecordedDataMode LoadingMode = EChaosVDLoadRecordedDataMode::MultiSource);
 
 	void OpenSession(const FChaosVDTraceSessionDescriptor& SessionDescriptor, EChaosVDLoadRecordedDataMode LoadingMode = EChaosVDLoadRecordedDataMode::MultiSource);
+
+	bool SaveOpenSessionToCombinedFile(const FString& InTargetFilePath = FString());
+
+	bool CanCombineOpenSessions() const;
 	
 	FSessionStateChangedDelegate& OnSessionOpened()
 	{
@@ -86,6 +92,9 @@ public:
 	}
 
 private:
+
+	void LoadRecording_Internal(const TFunction<FString(const TSharedPtr<FChaosVDRecording>&)>&, EChaosVDLoadRecordedDataMode LoadingMode = EChaosVDLoadRecordedDataMode::MultiSource);
+
 	void RestoreDataChannelsEnabledStateFromSave();
 	void UpdateSavedDataChannelsEnabledState(TWeakPtr<Chaos::VisualDebugger::FChaosVDOptionalDataChannel> DataChannelChanged);
 
diff --git a/Engine/Plugins/ChaosVD/Source/ChaosVD/Private/Trace/ChaosVDCombinedTraceFile.cpp b/Engine/Plugins/ChaosVD/Source/ChaosVD/Private/Trace/ChaosVDCombinedTraceFile.cpp
new file mode 100644
--- /dev/null
+++ b/Engine/Plugins/ChaosVD/Source/ChaosVD/Private/Trace/ChaosVDCombinedTraceFile.cpp
@@ -0,0 +1,239 @@
+// Copyright Epic Games, Inc. All Rights Reserved.
+
+#include "ChaosVDCombinedTraceFile.h"
+
+#include "ChaosVDModule.h"
+#include "ChaosVisualDebugger/ChaosVDMemWriterReader.h"
+#include "HAL/FileManager.h"
+#include "HAL/FileManagerGeneric.h"
+#include "Internationalization/Internationalization.h"
+#include "Misc/ScopedSlowTask.h"
+#include "Serialization/MemoryReader.h"
+#include "Serialization/MemoryWriter.h"
+
+namespace Chaos::VisualDebugger::CombinedTraceFile
+{
+	int64 FInnerFileHandle::Tell()
+	{
+		return ContainerFileHandle->Tell() - DataOffset;
+	}
+
+	bool FInnerFileHandle::Seek(int64 NewPosition)
+	{
+		if (NewPosition >= DataOffset && NewPosition <= (DataOffset + DataSize))
+		{
+			return ensure(ContainerFileHandle->Seek(DataOffset + NewPosition));
+		}
+	
+		return ensure(false);
+	}
+
+	bool FInnerFileHandle::SeekFromEnd(int64 NewPositionRelativeToEnd)
+	{
+		check(NewPositionRelativeToEnd <= 0);
+		const int64 EndOfFilePos = DataOffset + DataSize;
+
+		const int64 NewPosition = EndOfFilePos + NewPositionRelativeToEnd;
+		if (NewPosition >= DataOffset && NewPosition <= (DataOffset + DataSize))
+		{
+			return ensure(ContainerFileHandle->Seek(NewPosition));
+		}
+
+		return ensure(false);
+	}
+
+	bool FInnerFileHandle::Read(uint8* Destination, int64 BytesToRead)
+	{
+		return ensure(ContainerFileHandle->Read(Destination, BytesToRead));
+	}
+
+	bool FInnerFileHandle::ReadAt(uint8* Destination, int64 BytesToRead, int64 Offset)
+	{
+		return ensure(ContainerFileHandle->ReadAt(Destination, BytesToRead, Offset + DataOffset));
+	}
+
+	bool FInnerFileHandle::Write(const uint8* Source, int64 BytesToWrite)
+	{
+		return ensureMsgf(false, TEXT("Not Supported"));
+	}
+
+	bool FInnerFileHandle::Flush(const bool bFullFlush)
+	{
+		return ensureMsgf(false, TEXT("Not Supported"));
+	}
+
+	bool FInnerFileHandle::Truncate(int64 NewSize)
+	{
+		return ensureMsgf(false, TEXT("Not Supported"));
+	}
+
+	int64 FInnerFileHandle::Size()
+	{
+		return DataSize;
+	}
+
+	FInnerFileHandle::FInnerFileHandle(const FString& InContainerFilePath, int64 DataOffset, int64 DataSize)
+	: ContainerFileHandle(nullptr)
+	, DataOffset(DataOffset)
+	, DataSize(DataSize)
+	{
+		IPlatformFile& FileSystem = IPlatformFile::GetPlatformPhysical();
+		ContainerFileHandle = TUniquePtr<IFileHandle>(FileSystem.OpenRead(*InContainerFilePath));
+		if (ContainerFileHandle)
+		{
+			ContainerFileHandle->Seek(DataOffset);
+		}
+	}
+
+	bool CombineFiles(const TArray<TUniquePtr<IFileHandle>>& InFileHandlesToCombine, const FString& InCombinedFilePathName)
+	{
+		TUniquePtr<IFileHandle> CombinedFileHandle = nullptr;
+
+		if (InFileHandlesToCombine.Num() == 0)
+		{
+			UE_LOG(LogChaosVDEditor, Error, TEXT("Failed to combine files because an empty file handles array was provided."));
+			return false;
+		}
+
+		IPlatformFile& FileSystem = IPlatformFile::GetPlatformPhysical();
+		
+		if (FileSystem.FileExists(*InCombinedFilePathName))
+		{
+			UE_LOG(LogChaosVDEditor, Error, TEXT("Failed to combine files because an the target file name already exists."));
+			return false;
+		}
+
+		IFileManager& FileManager = IFileManager::Get();
+
+		FArchive* FileWriter = FileManager.CreateFileWriter(*InCombinedFilePathName);
+
+		if (!FileWriter)
+		{
+			UE_LOG(LogChaosVDEditor, Error, TEXT("Failed to combine files because there was an error while creating the target file."))
+			return false;
+		}
+		
+		constexpr float AmountOfWork = 1.0f;
+		const float CombiningFilesPercentagePerElement = 1.0f / static_cast<float>(InFileHandlesToCombine.Num());
+
+		FScopedSlowTask CombiningOpenRecordingsTask(AmountOfWork, NSLOCTEXT("ChaosVisualDebugger", "CombiningFilesMessage", "Combining Open Recordings ..."));
+		CombiningOpenRecordingsTask.MakeDialog();
+
+		FChaosVDArchiveHeader::Current().Serialize(*FileWriter);
+
+		int64 FileTableHeaderPos = FileWriter->Tell();
+		{
+			FFileTableHeader EmptyHeaderPlaceholder;
+			*FileWriter << EmptyHeaderPlaceholder;
+		}
+
+		FFileTable FileTable;
+		for (const TUniquePtr<IFileHandle>& FileHandle : InFileHandlesToCombine)
+		{
+			FScopedSlowTask ProcessingFileForCombiningTask(AmountOfWork, NSLOCTEXT("ChaosVisualDebugger", "ProcessingFileForCombiningTaskMessage", "Processing file ..."));
+			CombiningOpenRecordingsTask.MakeDialog();
+
+			FFileEntry FileEntry;
+			FileEntry.Size = FileHandle->Size();
+			FileEntry.StartPos = FileWriter->Tell();
+
+			FileTable.Files.Emplace(FileEntry);
+
+			constexpr int32 ReportProgressAfterReadSize = PLATFORM_FILE_WRITER_BUFFER_SIZE * 10;
+			constexpr float ProcessingFileMinProgressAmount = 1.0f / static_cast<float>(ReportProgressAfterReadSize);
+
+			constexpr int32 MaxChunkSize = PLATFORM_FILE_WRITER_BUFFER_SIZE;
+			TArray<uint8, TInlineAllocator<MaxChunkSize>> ChunkToCopy;
+			int64 DataRemaining = FileEntry.Size;
+
+			FileHandle->Seek(0);
+
+			while (DataRemaining > 0)
+			{
+				int32 DataSizeToRead = MaxChunkSize;
+				if (MaxChunkSize > DataRemaining)
+				{
+					DataSizeToRead = static_cast<int32>(DataRemaining);
+				}
+
+				ChunkToCopy.SetNum(DataSizeToRead);
+
+				FileHandle->Read(ChunkToCopy.GetData(), DataSizeToRead);
+				FileWriter->Serialize(ChunkToCopy.GetData(), DataSizeToRead);
+
+				DataRemaining -= DataSizeToRead;
+				
+				const int64 DataProcessedSoFar = FileEntry.Size - DataRemaining;
+				if (DataProcessedSoFar % ReportProgressAfterReadSize == 0)
+				{
+					ProcessingFileForCombiningTask.EnterProgressFrame(ProcessingFileMinProgressAmount);
+				}
+			}
+			
+			CombiningOpenRecordingsTask.EnterProgressFrame(CombiningFilesPercentagePerElement);
+		}
+
+		FFileTableHeader FileHeader;
+		FileHeader.FileTablePos = FileWriter->Tell();
+
+		*FileWriter << FileTable;
+
+		// Now that we have the final file table location, we can go back and re-serialize
+		// our header with the correct position data
+		FileWriter->Seek(FileTableHeaderPos);
+		*FileWriter << FileHeader;
+
+		FileWriter->Close();
+		
+		return true;
+	}
+
+	TArray<TUniquePtr<IFileHandle>> GetInnerFileHandles(const FString& InContainerFilePath)
+	{
+		TArray<TUniquePtr<IFileHandle>> FileHandles;
+		IPlatformFile& FileSystem = IPlatformFile::GetPlatformPhysical();
+		
+		IFileHandle* ContainerFileHandle = FileSystem.OpenRead(*InContainerFilePath);
+		if (!ensure(ContainerFileHandle))
+		{
+			UE_LOG(LogChaosVDEditor, Error, TEXT("Failed to read combined trace file. There was an error trying to open the file for read | File path [%s]."), *InContainerFilePath);
+			return FileHandles;
+		}
+
+		FArchiveFileReaderGeneric FileReader(ContainerFileHandle, *InContainerFilePath, ContainerFileHandle->Size());
+
+		FChaosVDArchiveHeader::Current().Serialize(FileReader);
+
+		FFileTableHeader FileTableHeader;
+		FileReader << FileTableHeader;
+
+		if (FileTableHeader.FileTablePos == 0)
+		{
+			UE_LOG(LogChaosVDEditor, Error, TEXT("Failed to read combined trace file. The serialized file table header is not valid | File path [%s]."), *InContainerFilePath);
+			return FileHandles;
+		}
+
+		FileReader.Seek(FileTableHeader.FileTablePos);
+
+		FFileTable FileTable;
+		FileReader << FileTable;
+
+		for (const FFileEntry& FileEntry : FileTable.Files)
+		{
+			TUniquePtr<FInnerFileHandle> PackedDataFileHandle = MakeUnique<FInnerFileHandle>(InContainerFilePath, FileEntry.StartPos, FileEntry.Size);
+			if (ensure(PackedDataFileHandle->IsValidHandle()))
+			{
+				FileHandles.Emplace(MoveTemp(PackedDataFileHandle));
+			}
+			else
+			{
+				UE_LOG(LogChaosVDEditor, Error, TEXT("Failed to create a inner file handle for. there was an error trying to open the file for read | Combined file path [%s]."), *InContainerFilePath);
+			}
+		}
+
+		/** Each InnerFileHandle opens its own file table to the container file, so it is fine to close this handle */
+		FileReader.Close();
+		
+		return MoveTemp(FileHandles);
+	}
+}
diff --git a/Engine/Plugins/ChaosVD/Source/ChaosVD/Private/Trace/ChaosVDCombinedTraceFile.h b/Engine/Plugins/ChaosVD/Source/ChaosVD/Private/Trace/ChaosVDCombinedTraceFile.h
new file mode 100644
--- /dev/null
+++ b/Engine/Plugins/ChaosVD/Source/ChaosVD/Private/Trace/ChaosVDCombinedTraceFile.h
@@ -0,0 +1,105 @@
+// Copyright Epic Games, Inc. All Rights Reserved.
+
+#pragma once
+#include "HAL/Platform.h"
+#include "Containers/UnrealString.h"
+#include "GenericPlatform/GenericPlatformFile.h"
+#include "Templates/SharedPointer.h"
+
+namespace Chaos::VisualDebugger::CombinedTraceFile
+{
+	/** File handle that allows other systems, like UE Trace, to read data from a combined file, as it was a separated file */
+	class FInnerFileHandle : public IFileHandle
+	{
+	public:
+		virtual int64 Tell() override;
+		virtual bool Seek(int64 NewPosition) override;
+		virtual bool SeekFromEnd(int64 NewPositionRelativeToEnd = 0) override;
+		virtual bool Read(uint8* Destination, int64 BytesToRead) override;
+		virtual bool ReadAt(uint8* Destination, int64 BytesToRead, int64 Offset) override;
+		virtual bool Write(const uint8* Source, int64 BytesToWrite) override;
+		virtual bool Flush(const bool bFullFlush = false) override;
+		virtual bool Truncate(int64 NewSize) override;
+		virtual int64 Size() override;
+
+		FInnerFileHandle(const FString& InContainerFilePath, int64 DataOffset, int64 DataSize);
+
+		bool IsValidHandle() const
+		{
+			return ContainerFileHandle.IsValid();
+		}
+
+	private:
+		TUniquePtr<IFileHandle> ContainerFileHandle;
+		int64 DataOffset;
+		int64 DataSize;
+	};
+
+	struct FFileEntry
+	{
+		int64 StartPos = 0;
+		int64 Size = 0;
+		
+		void Serialize(FArchive& Ar)
+		{
+			Ar << StartPos; 
+			Ar << Size; 
+		}
+	};
+
+	struct FFileTable
+	{
+		TArray<FFileEntry> Files;
+		
+		void Serialize(FArchive& Ar)
+		{
+			Ar << Files; 
+		}
+	};
+
+	struct FFileTableHeader
+	{
+		int64 FileTablePos = 0;
+
+		void Serialize(FArchive& Ar)
+		{
+			Ar << FileTablePos;
+		}		
+	};
+
+	inline FArchive& operator<<(FArchive& Ar, FFileTableHeader& Data)
+	{
+		Data.Serialize(Ar);
+		return Ar;
+	}
+
+	inline FArchive& operator<<(FArchive& Ar, FFileTable& Data)
+	{
+		Data.Serialize(Ar);
+		return Ar;
+	}
+
+	inline FArchive& operator<<(FArchive& Ar, FFileEntry& Data)
+	{
+		Data.Serialize(Ar);
+		return Ar;
+	}
+
+	/**
+	 * Takes an array of file handles and creates a new file that contains all their data, but that can be accessive as indivitual files later on
+	 * @param InFileHandlesToCombine File handles to get the data to combine from
+	 * @param InCombinedFilePathName Desired file path name for the generated file
+	 * @return true if successful
+	 */
+	bool CombineFiles(const TArray<TUniquePtr<IFileHandle>>& InFileHandlesToCombine, const FString& InCombinedFilePathName);
+
+	/**
+	 * Take a file path to a combined CVD recording file, and returns file handles for the individual files inside it
+	 * @param InContainerFilePath File path name to the combined CVD file
+	 * @return Unique file handles for each one of the files that were combined
+	 */
+	TArray<TUniquePtr<IFileHandle>> GetInnerFileHandles(const FString& InContainerFilePath);
+
+}
+
+
diff --git a/Engine/Plugins/ChaosVD/Source/ChaosVD/Private/Trace/ChaosVDTraceManager.cpp b/Engine/Plugins/ChaosVD/Source/ChaosVD/Private/Trace/ChaosVDTraceManager.cpp
--- a/Engine/Plugins/ChaosVD/Source/ChaosVD/Private/Trace/ChaosVDTraceManager.cpp
+++ b/Engine/Plugins/ChaosVD/Source/ChaosVD/Private/Trace/ChaosVDTraceManager.cpp
@@ -5,6 +5,7 @@
 #include "ChaosVDEngine.h"
 #include "ChaosVDModule.h"
 #include "Features/IModularFeatures.h"
+#include "GenericPlatform/GenericPlatformFile.h"
 #include "Modules/ModuleManager.h"
 #include "Misc/Paths.h"
 #include "ProfilingDebugging/TraceAuxiliary.h"
@@ -48,6 +49,24 @@ FString FChaosVDTraceManager::LoadTraceFile(const FString& InTraceFilename, cons
 	return FString();
 }
 
+FString FChaosVDTraceManager::LoadTraceFile(TUniquePtr<IFileHandle>&& InFileHandle, const FString& InTraceSessionName, const TSharedPtr<FChaosVDRecording>& ExistingRecordingPtr)
+{
+	ITraceServicesModule& TraceServicesModule = FModuleManager::LoadModuleChecked<ITraceServicesModule>("TraceServices");
+	if (const TSharedPtr<TraceServices::IAnalysisService> TraceAnalysisService = TraceServicesModule.GetAnalysisService())
+	{
+		SetPendingExternalRecordingToProcess(ExistingRecordingPtr);
+
+		if (const TSharedPtr<const TraceServices::IAnalysisSession> NewSession = TraceAnalysisService->StartAnalysis(~0, *InTraceSessionName, CreateFileDataStream(MoveTemp(InFileHandle))))
+		{
+			AnalysisSessionByName.Add(InTraceSessionName, NewSession.ToSharedRef());
+
+			return NewSession->GetName();
+		}
+	}
+
+	return FString();
+}
+
 const UE::Trace::FStoreClient::FSessionInfo* FChaosVDTraceManager::GetTraceSessionInfo(FStringView InSessionHost, FGuid TraceGuid)
 {
 	if (InSessionHost.IsEmpty())
@@ -68,6 +87,43 @@ const UE::Trace::FStoreClient::FSessionInfo* FChaosVDTraceManager::GetTraceSessi
 	return StoreClient->GetSessionInfoByGuid(TraceGuid);
 }
 
+TUniquePtr<UE::Trace::IInDataStream> FChaosVDTraceManager::CreateFileDataStream(TUniquePtr<IFileHandle>&& InFileHandle)
+{
+	check(InFileHandle);
+
+	struct FFileDataStream : public UE::Trace::FFileDataStream
+	{
+		virtual int32 Read(void* Data, uint32 Size) override
+		{
+			if (Remaining <= 0)
+			{
+				return 0;
+			}
+			if (Size > Remaining)
+			{
+				Size = static_cast<uint32>(Remaining);
+			}
+			Remaining -= Size;
+			if (!Handle->Read(static_cast<uint8*>(Data), Size))
+			{
+				return 0;
+			}
+			return Size;
+		}
+
+		TUniquePtr<IFileHandle> Handle;
+		uint64 Remaining;
+	};
+
+	FFileDataStream* FileStream = new FFileDataStream();
+	FileStream->Handle = MoveTemp(InFileHandle);
+	FileStream->Remaining = FileStream->Handle->Size();
+
+	TUniquePtr<UE::Trace::IInDataStream> DataStream(FileStream);
+
+	return DataStream;
+}
+
 void FChaosVDTraceManager::SetPendingExternalRecordingToProcess(const TSharedPtr<FChaosVDRecording>& ExistingRecordingPtr)
 {
 	TWeakPtr<FChaosVDRecording>& PendingRecording = FChaosVDTraceManagerThreadContext::Get().PendingExternalRecordingWeakPtr;
diff --git a/Engine/Plugins/ChaosVD/Source/ChaosVD/Private/Trace/ChaosVDTraceManager.h b/Engine/Plugins/ChaosVD/Source/ChaosVD/Private/Trace/ChaosVDTraceManager.h
--- a/Engine/Plugins/ChaosVD/Source/ChaosVD/Private/Trace/ChaosVDTraceManager.h
+++ b/Engine/Plugins/ChaosVD/Source/ChaosVD/Private/Trace/ChaosVDTraceManager.h
@@ -46,6 +46,8 @@ public:
 	 * @param InTraceFilename File Name including (Path Included) of the Trace file to load
 	 */
 	FString LoadTraceFile(const FString& InTraceFilename, const TSharedPtr<FChaosVDRecording>& ExistingRecordingPtr = nullptr);
+
+	FString LoadTraceFile(TUniquePtr<IFileHandle>&& InFileHandle, const FString& InTraceSessionName, const TSharedPtr<FChaosVDRecording>& ExistingRecordingPtr);
 	
 	/**
 	 * Connects to a live Trace Session and starts analyzing it.
@@ -73,6 +75,8 @@ public:
 
 private:
 
+	TUniquePtr<UE::Trace::IInDataStream> CreateFileDataStream(TUniquePtr<IFileHandle>&& InFileHandle);
+
 	/** Temporary workaround method to set an existing recording structure in CVD's trace provider before the trace analysis starts
 	 * in the Trace analysis thread.*/
 	void SetPendingExternalRecordingToProcess(const TSharedPtr<FChaosVDRecording>& ExistingRecordingPtr);
diff --git a/Engine/Plugins/ChaosVD/Source/ChaosVD/Private/Widgets/SChaosVDMainTab.cpp b/Engine/Plugins/ChaosVD/Source/ChaosVD/Private/Widgets/SChaosVDMainTab.cpp
--- a/Engine/Plugins/ChaosVD/Source/ChaosVD/Private/Widgets/SChaosVDMainTab.cpp
+++ b/Engine/Plugins/ChaosVD/Source/ChaosVD/Private/Widgets/SChaosVDMainTab.cpp
@@ -410,6 +410,39 @@ void SChaosVDMainTab::ResetLayout()
 	FChaosVDModule::Get().ReloadInstanceUI(GetChaosVDEngineInstance()->GetInstanceGuid());
 }
 
+void SChaosVDMainTab::CombineOpenSessions()
+{
+	TArray<FString> OutSelectedFilenames;
+	if (IDesktopPlatform* DesktopPlatform = FDesktopPlatformModule::Get())
+	{
+		FString ExtensionStr;
+		ExtensionStr += TEXT("CVD Multi Session|*.cvdmulti");
+	
+		DesktopPlatform->SaveFileDialog(
+			FSlateApplication::Get().FindBestParentWindowHandleForDialogs(AsShared()),
+			LOCTEXT("SaveDialogTitle", "Save Combined Chaos Visual Debug File").ToString(),
+			*FPaths::ProfilingDir(),
+			TEXT(""),
+			*ExtensionStr,
+			EFileDialogFlags::None,
+			OutSelectedFilenames
+		);
+	}
+
+	if (!OutSelectedFilenames.IsEmpty())
+	{
+		const FString& TargetFilePath = OutSelectedFilenames[0];
+		if (GetChaosVDEngineInstance()->SaveOpenSessionToCombinedFile(TargetFilePath))
+		{
+			FPlatformProcess::ExploreFolder(*TargetFilePath);
+		}
+		else
+		{
+			FMessageDialog::Open(EAppMsgCategory::Error, EAppMsgType::Ok, LOCTEXT("FailedCombineFilesMessageTitle", "Failed to combine open recordings into a single file.\n\n See Logs for mor info."), LOCTEXT("FailedCombineFilesMessageTitle", "Failed combine files"));
+		}
+	}
+}
+
 void SChaosVDMainTab::GenerateMainWindowMenu()
 {
 	FMenuBarBuilder MenuBarBuilder = FMenuBarBuilder(TSharedPtr<FUICommandList >());
@@ -422,6 +455,7 @@ void SChaosVDMainTab::GenerateMainWindowMenu()
 		}),
 		"File"
 	);
+
 	MenuBarBuilder.AddPullDownMenu(
 		LOCTEXT("WindowMenuLabel", "Window"),
 		FText::GetEmpty(),
@@ -558,7 +592,7 @@ TSharedRef<SWidget> SChaosVDMainTab::GenerateSettingsMenuWidget()
 
 void SChaosVDMainTab::LoadCVDFile(const FString& InFilename, EChaosVDLoadRecordedDataMode LoadingMode)
 {
-	if (ensure(InFilename.EndsWith(TEXT("utrace"))))
+	if (ensure(InFilename.EndsWith(TEXT("utrace")) || InFilename.EndsWith("cvdmulti")))
 	{
 		GetChaosVDEngineInstance()->LoadRecording(InFilename, LoadingMode);
 	}
@@ -691,6 +725,7 @@ void SChaosVDMainTab::BrowseChaosVDRecordingFromFolder(FStringView FolderPath, E
 	{
 		FString ExtensionStr;
 		ExtensionStr += TEXT("Unreal Trace|*.utrace|");
+		ExtensionStr += TEXT("CVD Multi Session|*.cvdmulti");
 		//TODO: Re-enable this when we add "Clips" support as these will use our own format
 		//ExtensionStr += TEXT("Chaos Visual Debugger|*.cvd");
 	
@@ -706,7 +741,9 @@ void SChaosVDMainTab::BrowseChaosVDRecordingFromFolder(FStringView FolderPath, E
 	}
 
 	// Ideally, we should not need to do this, but the UI to support multi source is not robust yet, The UI improvement task to avoid this scenario is planned as UE-197418
-	if (const bool bShouldOverrideMode = OutOpenFilenames.Num() > 1 && LoadingMode == EChaosVDLoadRecordedDataMode::SingleSource)
+	const bool bHasMultiFileData = OutOpenFilenames.Num() > 1 || (OutOpenFilenames.IsValidIndex(0) && OutOpenFilenames[0].EndsWith(TEXT("cvdmulti")));
+
+	if (const bool bShouldOverrideMode = bHasMultiFileData && LoadingMode == EChaosVDLoadRecordedDataMode::SingleSource)
 	{
 		GetChaosVDEngineInstance()->CloseActiveTraceSessions();
 		LoadingMode = EChaosVDLoadRecordedDataMode::MultiSource;
@@ -887,6 +924,51 @@ void SChaosVDMainTab::RegisterMainTabMenu()
 				false
 			));
 	}));
+
+	Section.AddDynamicEntry("CombineFiles", FNewToolMenuSectionDelegate::CreateLambda([](FToolMenuSection& InSection)
+	{
+		const UChaosVDMainToolbarMenuContext* Context = InSection.FindContext<UChaosVDMainToolbarMenuContext>();
+		TSharedPtr<SChaosVDMainTab> MainTabPtr = Context->MainTab.Pin();
+		if (!MainTabPtr)
+		{
+			return;
+		}
+
+		TSharedRef<SButton> CombineFilesButton = MainTabPtr->CreateSimpleButton(
+			[](){ return LOCTEXT("CombineFilesLabel", "Combine"); },
+			[](){ return LOCTEXT("CombineFilesLabelDesc", "Click here to combine multiple open recordings into a single file to make sharing easier."); },
+			FAppStyle::Get().GetBrush("MainFrame.ZipUpProject"),
+			Context, FOnClicked::CreateLambda([WeakTab = MainTabPtr.ToWeakPtr()]()
+			{
+				if (TSharedPtr<SChaosVDMainTab> TabPtr = WeakTab.Pin())
+				{
+					TabPtr->CombineOpenSessions();
+				}
+
+				return FReply::Handled();
+			}));
+
+		TAttribute<bool> IsCombineButtonEnabled;
+		IsCombineButtonEnabled.BindSPLambda(MainTabPtr.ToSharedRef(), [WeakTab = MainTabPtr.ToWeakPtr()]()
+		{
+			if (TSharedPtr<SChaosVDMainTab> TabPtr = WeakTab.Pin())
+			{
+				return TabPtr->GetChaosVDEngineInstance()->CanCombineOpenSessions();
+			}
+			return false;
+		});
+
+		CombineFilesButton->SetEnabled(IsCombineButtonEnabled);
+
+		InSection.AddEntry(
+			FToolMenuEntry::InitWidget(
+				"CombineFilesLabelButton",
+				CombineFilesButton,
+				FText::GetEmpty(),
+				true,
+				false
+			));
+	}));
 	
 	Section.AddSeparator(NAME_None);
 
diff --git a/Engine/Plugins/ChaosVD/Source/ChaosVD/Private/Widgets/SChaosVDRecordingControls.cpp b/Engine/Plugins/ChaosVD/Source/ChaosVD/Private/Widgets/SChaosVDRecordingControls.cpp
--- a/Engine/Plugins/ChaosVD/Source/ChaosVD/Private/Widgets/SChaosVDRecordingControls.cpp
+++ b/Engine/Plugins/ChaosVD/Source/ChaosVD/Private/Widgets/SChaosVDRecordingControls.cpp
@@ -855,6 +855,8 @@ bool SChaosVDRecordingControls::IsRecordingToggleButtonEnabled(EChaosVDRecording
 
 				return CanSelectMultiSessionTarget(SessionInfo.ToSharedRef());
 			}
+
+			return true;
 		}
 
 		return false;
diff --git a/Engine/Plugins/ChaosVD/Source/ChaosVD/Public/Widgets/SChaosVDMainTab.h b/Engine/Plugins/ChaosVD/Source/ChaosVD/Public/Widgets/SChaosVDMainTab.h
--- a/Engine/Plugins/ChaosVD/Source/ChaosVD/Public/Widgets/SChaosVDMainTab.h
+++ b/Engine/Plugins/ChaosVD/Source/ChaosVD/Public/Widgets/SChaosVDMainTab.h
@@ -156,6 +156,8 @@ private:
 
 	void ResetLayout();
 
+	void CombineOpenSessions();
+
 	void GenerateMainWindowMenu();
 	void GenerateRecentFilesMenu(FMenuBuilder& MenuBuilder);
 
