[ChaosVD] Marking settings properties as config so they can be saved and restore to defaults, which were missing.

diff --git a/Engine/Plugins/ChaosVD/Source/ChaosVD/Private/Settings/ChaosVDCharacterConstraintsVisualizationSettings.h b/Engine/Plugins/ChaosVD/Source/ChaosVD/Private/Settings/ChaosVDCharacterConstraintsVisualizationSettings.h
--- a/Engine/Plugins/ChaosVD/Source/ChaosVD/Private/Settings/ChaosVDCharacterConstraintsVisualizationSettings.h
+++ b/Engine/Plugins/ChaosVD/Source/ChaosVD/Private/Settings/ChaosVDCharacterConstraintsVisualizationSettings.h
@@ -41,7 +41,7 @@ class UChaosVDCharacterConstraintsVisualizationSettings : public UChaosVDVisuali
 public:
 
 	/** If true, any debug draw text available will be drawn */
-	UPROPERTY(EditAnywhere, Category=DebugDraw)
+	UPROPERTY(config, EditAnywhere, Category=DebugDraw)
 	bool bShowDebugText = false;
 	
 	/** Color used for torque vector */
diff --git a/Engine/Plugins/ChaosVD/Source/ChaosVD/Private/Settings/ChaosVDCollisionVisualizationSettings.h b/Engine/Plugins/ChaosVD/Source/ChaosVD/Private/Settings/ChaosVDCollisionVisualizationSettings.h
--- a/Engine/Plugins/ChaosVD/Source/ChaosVD/Private/Settings/ChaosVDCollisionVisualizationSettings.h
+++ b/Engine/Plugins/ChaosVD/Source/ChaosVD/Private/Settings/ChaosVDCollisionVisualizationSettings.h
@@ -29,7 +29,7 @@ class UChaosVDCollisionDataVisualizationSettings : public UChaosVDVisualizationS
 public:
 
 	/** If true, any available debug text will be drawn */
-	UPROPERTY(EditAnywhere, Category=DebugDraw)
+	UPROPERTY(EditAnywhere, config, Category=DebugDraw)
 	bool bShowDebugText = false;
 
 	/** The depth priority used for while drawing contact data. Can be World or Foreground (with this one the shapes will be drawn on top of the geometry and be always visible) */
diff --git a/Engine/Plugins/ChaosVD/Source/ChaosVD/Private/Settings/ChaosVDJointConstraintVisualizationSettings.h b/Engine/Plugins/ChaosVD/Source/ChaosVD/Private/Settings/ChaosVDJointConstraintVisualizationSettings.h
--- a/Engine/Plugins/ChaosVD/Source/ChaosVD/Private/Settings/ChaosVDJointConstraintVisualizationSettings.h
+++ b/Engine/Plugins/ChaosVD/Source/ChaosVD/Private/Settings/ChaosVDJointConstraintVisualizationSettings.h
@@ -37,11 +37,11 @@ class UChaosVDJointConstraintsVisualizationSettings : public UChaosVDVisualizati
 public:
 
 	/** If true, any debug draw text available will be drawn */
-	UPROPERTY(EditAnywhere, Category=DebugDraw)
+	UPROPERTY(EditAnywhere, config, Category=DebugDraw)
 	bool bShowDebugText = false;
 
 	/** The depth priority used for while drawing data. Can be World or Foreground (with this one the shapes will be drawn on top of the geometry and be always visible) */
-	UPROPERTY(EditAnywhere, Category=DebugDraw)
+	UPROPERTY(EditAnywhere, config, Category=DebugDraw)
 	TEnumAsByte<ESceneDepthPriorityGroup> DepthPriority = ESceneDepthPriorityGroup::SDPG_Foreground;
 
 	/** Scale to apply to the Linear Impulse vector before draw it. */
diff --git a/Engine/Plugins/ChaosVD/Source/ChaosVD/Private/Settings/ChaosVDParticleVisualizationSettings.h b/Engine/Plugins/ChaosVD/Source/ChaosVD/Private/Settings/ChaosVDParticleVisualizationSettings.h
--- a/Engine/Plugins/ChaosVD/Source/ChaosVD/Private/Settings/ChaosVDParticleVisualizationSettings.h
+++ b/Engine/Plugins/ChaosVD/Source/ChaosVD/Private/Settings/ChaosVDParticleVisualizationSettings.h
@@ -237,7 +237,7 @@ class UChaosVDParticleVisualizationDebugDrawSettings : public UChaosVDVisualizat
 	GENERATED_BODY()
 public:
 
-	UPROPERTY(EditAnywhere, Category=DebugDraw)
+	UPROPERTY(config, EditAnywhere, Category=DebugDraw)
 	bool bShowDebugText = false;
 
 	/** The depth priority used for while drawing contact data. Can be World or Foreground (with this one the shapes will be drawn on top of the geometry and be always visible) */
diff --git a/Engine/Plugins/ChaosVD/Source/ChaosVD/Private/Settings/ChaosVDSceneQueryVisualizationSettings.h b/Engine/Plugins/ChaosVD/Source/ChaosVD/Private/Settings/ChaosVDSceneQueryVisualizationSettings.h
--- a/Engine/Plugins/ChaosVD/Source/ChaosVD/Private/Settings/ChaosVDSceneQueryVisualizationSettings.h
+++ b/Engine/Plugins/ChaosVD/Source/ChaosVD/Private/Settings/ChaosVDSceneQueryVisualizationSettings.h
@@ -44,7 +44,7 @@ public:
 	static EChaosVDSceneQueryVisualizationFlags GetDataVisualizationFlags();
 
 	/** If true, any debug draw text available will be drawn */
-	UPROPERTY(EditAnywhere, Category=DebugDraw)
+	UPROPERTY(config, EditAnywhere, Category=DebugDraw)
 	bool bShowText = false;
 
 	/** The depth priority used for while drawing. Can be World or Foreground (with this one the shapes representing the query
diff --git a/Engine/Plugins/ChaosVD/Source/ChaosVDBuiltInExtensions/Private/GenericDebugDraw/Settings/ChaosVDGenericDebugDrawSettings.h b/Engine/Plugins/ChaosVD/Source/ChaosVDBuiltInExtensions/Private/GenericDebugDraw/Settings/ChaosVDGenericDebugDrawSettings.h
--- a/Engine/Plugins/ChaosVD/Source/ChaosVDBuiltInExtensions/Private/GenericDebugDraw/Settings/ChaosVDGenericDebugDrawSettings.h
+++ b/Engine/Plugins/ChaosVD/Source/ChaosVDBuiltInExtensions/Private/GenericDebugDraw/Settings/ChaosVDGenericDebugDrawSettings.h
@@ -26,15 +26,15 @@ class UChaosVDGenericDebugDrawSettings : public UChaosVDVisualizationSettingsObj
 public:
 
 	/** If true, any debug draw text available will be drawn */
-	UPROPERTY(EditAnywhere, Category=DebugDraw)
+	UPROPERTY(config, EditAnywhere, Category=DebugDraw)
 	bool bShowDebugText = false;
 
 	/** The depth priority used for while drawing data. Can be World or Foreground (with this one the shapes will be drawn on top of the geometry and be always visible) */
-	UPROPERTY(EditAnywhere, Category=DebugDraw)
+	UPROPERTY(config, EditAnywhere, Category=DebugDraw)
 	TEnumAsByte<ESceneDepthPriorityGroup> DepthPriority = ESceneDepthPriorityGroup::SDPG_Foreground;
 
 	/** Thickness to apply to any debug draw shape controlled by this setting. */
-	UPROPERTY(EditAnywhere, config, Category=DebugDraw)
+	UPROPERTY(config, EditAnywhere, Category=DebugDraw)
 	float BaseThickness = 2.0f;
 
 	static void SetDataVisualizationFlags(EChaosVDGenericDebugDrawVisualizationFlags NewFlags);
