[ChaosVD] Fix issues with the empty frame trimming logic

CVD has logic to attempt to trim out frames that doesn't have data that can be visualized, but it has some issues like incorrectly allowing empty frames at the start of the recording.

This CL upates the system to only apply the trim logic at the start of the recording, and be more more agressive on trimming data out. Now we only accept frames that have solver data, regardless if there was game thread data like scene queries.

I also added an option disable the trim feature in the settings page.

Lastly, this CL also exposes in the settings UI the number of frames we want to queue before we push them into the CVD recording, and increased the default value to 10.

This settings aloows to "smooth out" playback in live sessions, as by default CVD will unload any visualized solver if a game thread frame is played, but no data is available.

In that scenario, if the solver data is lagging behind because we have async physics enabled and the GT frame time is consirably highier thant the PT frame time, the visualization would flicker.

Increasing the number of frames queued mitigates that issue. A better solution would be to keep track of all GT and PT frame times, and dynamically adjust the queue size, but that is a larger change than it is safe for UE 5.6, and the recent multi-source feature adds to the complexity.

diff --git a/Engine/Plugins/ChaosVD/Source/ChaosVD/Private/ChaosVDPlaybackController.cpp b/Engine/Plugins/ChaosVD/Source/ChaosVD/Private/ChaosVDPlaybackController.cpp
--- a/Engine/Plugins/ChaosVD/Source/ChaosVD/Private/ChaosVDPlaybackController.cpp
+++ b/Engine/Plugins/ChaosVD/Source/ChaosVD/Private/ChaosVDPlaybackController.cpp
@@ -230,7 +230,7 @@ void FChaosVDPlaybackController::GoToRecordedSolverStage_AssumesLocked(const int
 				// TODO: Implement a less expensive way of handle these cases.
 				// We should keep the previous state of each loaded particle so if when going back they are not in the new delta we are evaluating, (and were not destroyed)
 				// we can just re-apply that last known state.
-				const bool bNeedsToPlayFromKeyframe = FrameDiff < 0 || FMath::Abs(FrameDiff) > FrameDriftTolerance;
+				const bool bNeedsToPlayFromKeyframe = !bPlayedFirstFrame || FrameDiff < 0 || FMath::Abs(FrameDiff) > FrameDriftTolerance;
 
 				if (bNeedsToPlayFromKeyframe || CurrentTrackInfo->CurrentFrame == 0)
 				{
diff --git a/Engine/Plugins/ChaosVD/Source/ChaosVD/Private/Settings/ChaosVDGeneralSettings.h b/Engine/Plugins/ChaosVD/Source/ChaosVD/Private/Settings/ChaosVDGeneralSettings.h
--- a/Engine/Plugins/ChaosVD/Source/ChaosVD/Private/Settings/ChaosVDGeneralSettings.h
+++ b/Engine/Plugins/ChaosVD/Source/ChaosVD/Private/Settings/ChaosVDGeneralSettings.h
@@ -14,6 +14,14 @@ class CHAOSVD_API UChaosVDGeneralSettings : public UChaosVDSettingsObjectBase
 	GENERATED_BODY()
 
 public:
+	/** If true, CVD will only load frames that have solver data in them - Only takes effect before loading a file */
+	UPROPERTY(Config, EditAnywhere, Category = "Session Data Loading")
+	bool bTrimEmptyFrames = true;
+
+	/** How many Game thread frames CVD should queue internally before making them available in the visualization and timeline controls - Only takes effect before loading a file */
+	UPROPERTY(Config, EditAnywhere, Category = "Session Data Loading")
+	int32 MaxGameThreadFramesToQueueNum = 10;
+
 	/** If true, CVD will only load collision geometry that is visible */
 	UPROPERTY(Config, EditAnywhere, Category = "Performance | Streaming")
 	bool bStreamingSystemEnabled = true;
diff --git a/Engine/Plugins/ChaosVD/Source/ChaosVD/Private/Trace/ChaosVDTraceProvider.cpp b/Engine/Plugins/ChaosVD/Source/ChaosVD/Private/Trace/ChaosVDTraceProvider.cpp
--- a/Engine/Plugins/ChaosVD/Source/ChaosVD/Private/Trace/ChaosVDTraceProvider.cpp
+++ b/Engine/Plugins/ChaosVD/Source/ChaosVD/Private/Trace/ChaosVDTraceProvider.cpp
@@ -4,12 +4,14 @@
 
 #include "ChaosVDModule.h"
 #include "ChaosVDRecording.h"
+#include "ChaosVDSettingsManager.h"
 #include "ChaosVisualDebugger/ChaosVisualDebuggerTrace.h"
 #include "Compression/OodleDataCompressionUtil.h"
 #include "Containers/Ticker.h"
 #include "DataProcessors/ChaosVDCollisionChannelsInfoDataProcessor.h"
 #include "ExtensionsSystem/ChaosVDExtensionsManager.h"
 #include "Misc/MessageDialog.h"
+#include "Settings/ChaosVDGeneralSettings.h"
 #include "Trace/DataProcessors/ChaosVDCharacterGroundConstraintDataProcessor.h"
 #include "Trace/DataProcessors/ChaosVDConstraintDataProcessor.h"
 #include "Trace/DataProcessors/ChaosVDJointConstraintDataProcessor.h"
@@ -22,15 +24,6 @@
 #include "Trace/DataProcessors/ChaosVDArchiveHeaderProcessor.h"
 #include "Trace/DataProcessors/ChaosVDDataProcessorBase.h"
 
-namespace Chaos::VisualDebugger::Cvars
-{
-	static int32 MaxFramesToQueue = 2;
-	static FAutoConsoleVariableRef CVarChaosVDMaxFramesToQueue(
-		TEXT("p.Chaos.VD.Tool.MaxFramesToQueue"),
-		MaxFramesToQueue,
-		TEXT("The max number of game frames we will queue before processing them while the trace files is being loaded."));
-}
-
 #define LOCTEXT_NAMESPACE "ChaosVisualDebugger"
 
 FName FChaosVDTraceProvider::ProviderName("ChaosVDProvider");
@@ -42,6 +35,12 @@ FChaosVDTraceProvider::FChaosVDTraceProvider(TraceServices::IAnalysisSession& In
 
 	// Start with a default header data as a fallback
 	HeaderData = FChaosVDArchiveHeader::Current();
+	
+	if (UChaosVDGeneralSettings* Settings = FChaosVDSettingsManager::Get().GetSettingsObject<UChaosVDGeneralSettings>())
+	{
+		bShouldTrimOutStartEmptyFrames = Settings->bTrimEmptyFrames;
+		MaxGameFramesToQueueNum = Settings->MaxGameThreadFramesToQueueNum;
+	}
 }
 
 void FChaosVDTraceProvider::CreateRecordingInstanceForSession(const FString& InSessionName)
@@ -172,9 +171,7 @@ void FChaosVDTraceProvider::CommitProcessedGameFramesToRecording()
 	// These are generated even when no solvers are available (specially in PIE), so we need to discard any game frame that will not resolve to a solver frame
 	// Physics Frames and GT frames lifetimes might not align with async physics enabled, so to make sure we have all the solver data for that time range, we queue a handful of game frames before processing them.
 
-	// TODO: This empty frame logic trimming is no longer 100% correct since we added support for async physics and support for recording GT specific data.
-	// We need to find a better way to trim out empty data. Jira for Tracking UE-227365
-	if (CurrentGameFrameQueueSize > Chaos::VisualDebugger::Cvars::MaxFramesToQueue)
+	if (CurrentGameFrameQueueSize > MaxGameFramesToQueueNum)
 	{
 		TSharedPtr<FChaosVDGameFrameData> ProcessedGameFrameData;
 		DeQueueGameFrameForProcessing(ProcessedGameFrameData);
@@ -194,10 +191,16 @@ void FChaosVDTraceProvider::CommitProcessedGameFramesToRecording()
 			{
 				GetAvailablePendingSolverIDsAtGameFrame(ProcessedGameFrameData.ToSharedRef(), SolverIDs);
 			}
+	
+			const bool bHasAnySolverData = !SolverIDs.IsEmpty();
+			const bool bHasAnyGameFrame = InternalRecording->GetAvailableGameFramesNumber() > 0;
+			
+			bool bHasRelevantCVDData = true;
+			if (bShouldTrimOutStartEmptyFrames)
+			{
+				bHasRelevantCVDData = bHasAnyGameFrame ? true : bHasAnySolverData;
+			}
 
-			// Any frame that is marked as dirty, or has any kind of solver data is considered relevant. Otherwise, we can trim it from the recording
-			// by just ignoring it
-			const bool bHasRelevantCVDData = !SolverIDs.IsEmpty() || ProcessedGameFrameData->IsDirty();
 			if (bHasRelevantCVDData)
 			{
 				InternalRecording->AddGameFrameData(*ProcessedGameFrameData);	
diff --git a/Engine/Plugins/ChaosVD/Source/ChaosVD/Public/Trace/ChaosVDTraceProvider.h b/Engine/Plugins/ChaosVD/Source/ChaosVD/Public/Trace/ChaosVDTraceProvider.h
--- a/Engine/Plugins/ChaosVD/Source/ChaosVD/Public/Trace/ChaosVDTraceProvider.h
+++ b/Engine/Plugins/ChaosVD/Source/ChaosVD/Public/Trace/ChaosVDTraceProvider.h
@@ -212,7 +212,10 @@ private:
 	TSet<FString> MissingDataProcessors;
 	TSet<FString> TypesFailedToSerialize;
 
-	int64 DataProcessedSoFarCounter = 0; 
+	int64 DataProcessedSoFarCounter = 0;
+
+	bool bShouldTrimOutStartEmptyFrames = false;
+	int32 MaxGameFramesToQueueNum = 10;
 
 	friend class FChaosVDTraceAnalyzer;
 
