[ChaosVD] Fix localization warning in SChaosVDMainTab

#rnx

diff --git a/Engine/Plugins/ChaosVD/Source/ChaosVD/Private/Widgets/SChaosVDMainTab.cpp b/Engine/Plugins/ChaosVD/Source/ChaosVD/Private/Widgets/SChaosVDMainTab.cpp
--- a/Engine/Plugins/ChaosVD/Source/ChaosVD/Private/Widgets/SChaosVDMainTab.cpp
+++ b/Engine/Plugins/ChaosVD/Source/ChaosVD/Private/Widgets/SChaosVDMainTab.cpp
@@ -438,7 +438,7 @@ void SChaosVDMainTab::CombineOpenSessions()
 		}
 		else
 		{
-			FMessageDialog::Open(EAppMsgCategory::Error, EAppMsgType::Ok, LOCTEXT("FailedCombineFilesMessageTitle", "Failed to combine open recordings into a single file.\n\n See Logs for mor info."), LOCTEXT("FailedCombineFilesMessageTitle", "Failed combine files"));
+			FMessageDialog::Open(EAppMsgCategory::Error, EAppMsgType::Ok, LOCTEXT("FailedCombineFilesMessage", "Failed to combine open recordings into a single file.\n\n See Logs for mor info."), LOCTEXT("FailedCombineFilesMessageTitle", "Failed combine files"));
 		}
 	}
 }
