[ChaosVD] Fix for an issue where all simple shapes of type Box or Sphere aren't visible.

The problem was caused because there is a system to flag auto-generated geometry as garbage, but since a few months ago, CVD also use assets as geometry. These asset doesn't need to be manually flagged as garbage, and doing so causes issues the one this CL addresses

diff --git a/Engine/Plugins/ChaosVD/Source/ChaosVD/Private/ChaosVDGeometryBuilder.cpp b/Engine/Plugins/ChaosVD/Source/ChaosVD/Private/ChaosVDGeometryBuilder.cpp
--- a/Engine/Plugins/ChaosVD/Source/ChaosVD/Private/ChaosVDGeometryBuilder.cpp
+++ b/Engine/Plugins/ChaosVD/Source/ChaosVD/Private/ChaosVDGeometryBuilder.cpp
@@ -319,10 +319,11 @@ void FChaosVDGeometryBuilder::DeInitialize()
 
 	for (const TPair<uint32, TObjectPtr<UStaticMesh>>& StaticMeshByKey : StaticMeshCacheMap)
 	{
-		if (StaticMeshByKey.Value)
+		const TObjectPtr<UStaticMesh>& StaticMesh = StaticMeshByKey.Value;
+		if (StaticMesh && !StaticMesh->IsAsset())
 		{
-			StaticMeshByKey.Value->ClearFlags(RF_Standalone);
-			StaticMeshByKey.Value->MarkAsGarbage();
+			StaticMesh->ClearFlags(RF_Standalone);
+			StaticMesh->MarkAsGarbage();
 		}
 		CleaningGeometrySlowTask.EnterProgressFrame(PercentagePerElement);
 	}
