[ChaosVD] Fix for an issue where you cannot select a solver track using the key shortcut CTRL + [Track Number] where track number is a numeric key.

It was working for numpad keys, but not for aplhanumeric keys.

diff --git a/Engine/Plugins/ChaosVD/Source/ChaosVD/Private/Widgets/SChaosVDMainTab.cpp b/Engine/Plugins/ChaosVD/Source/ChaosVD/Private/Widgets/SChaosVDMainTab.cpp
--- a/Engine/Plugins/ChaosVD/Source/ChaosVD/Private/Widgets/SChaosVDMainTab.cpp
+++ b/Engine/Plugins/ChaosVD/Source/ChaosVD/Private/Widgets/SChaosVDMainTab.cpp
@@ -637,7 +637,7 @@ void SChaosVDMainTab::ProccessKeyEventForPlaybackTrackSelector(const FKeyEvent&
 		
 		if (KeyCode >= AlphaNumKeyCodeLowerBound && KeyCode <= AlphaNumKeyCodeUpperBound)
 		{
-			TrackSlotIndex = KeyCode - AlphaNumKeyCodeUpperBound;
+			TrackSlotIndex = KeyCode - AlphaNumKeyCodeLowerBound;
 		}
 
 		if (KeyCode >= NumPadNumberKeyCodeLowerBound && KeyCode <= NumPadNumberKeyCodeUpperBound)
