[ChaosVD] Downgrade ensure to verbose log in FChaosVDTraceProvider::GetCurrentSolverFrame

Since it implementation all edge cases has been addressed, and there are a few scenarios where the condition evaluated can be false but is expected and the system already handles it.

One of these cases is a particle that was recorded before a solver frame start even was recorded. This happens if the recording is started mid physics thread frame. We could fix that situation in the particle data recording code, but this could happen with any type of data, including extensions. As the system is already implemented to just ignore that data, handling this at the recording time is not worth it.

Therefore is safe to downgrade the ensure to a verbose log. I am also adding additional verbose logging so in cases where we suspect an issue might be related to this solver id remap system, we can enable verbose logging an know exactly what happened before and after.

Lastly, this CL also fixes an issue where verbose logging is compiled out for the LogChaosVDEditor category

diff --git a/Engine/Plugins/ChaosVD/Source/ChaosVD/Private/Trace/ChaosVDTraceProvider.cpp b/Engine/Plugins/ChaosVD/Source/ChaosVD/Private/Trace/ChaosVDTraceProvider.cpp
--- a/Engine/Plugins/ChaosVD/Source/ChaosVD/Private/Trace/ChaosVDTraceProvider.cpp
+++ b/Engine/Plugins/ChaosVD/Source/ChaosVD/Private/Trace/ChaosVDTraceProvider.cpp
@@ -141,6 +141,8 @@ int32 FChaosVDTraceProvider::RemapSolverID(int32 SolverID)
 		InternalRecording->ReserveSolverID_AssumesLocked(RemappedSolverID);
 	}
 
+	UE_LOG(LogChaosVDEditor, Verbose, TEXT("[%hs] Remapped solver id from [%d] to [%d]."), __func__, SolverID, RemappedSolverID);
+
 	RemappedSolversIDs.Emplace(SolverID, RemappedSolverID);
 
 	return RemappedSolverID;
@@ -156,6 +158,7 @@ int32 FChaosVDTraceProvider::GetRemappedSolverID(int32 SolverID)
 	}
 	else
 	{
+		UE_LOG(LogChaosVDEditor, Verbose, TEXT("[%hs] Failed to get remapped solver id [%d]. Data that references the invalid solver id will be ignored."), __func__, SolverID);
 		RemappedSolverID = INDEX_NONE; 
 	}
 
@@ -229,8 +232,9 @@ FChaosVDSolverFrameData* FChaosVDTraceProvider::GetCurrentSolverFrame(const int3
 	// middle of a frame and therefore the solver hasn't been open in CVD yet.
 	bool bIsInvalidSolverID = InSolverID == INDEX_NONE && (!CurrentSolverFramesByID.IsEmpty() && RemappedGameThreadTrackID != INDEX_NONE);
 
-	if (!ensure(!bIsInvalidSolverID))
+	if (bIsInvalidSolverID)
 	{
+		UE_LOG(LogChaosVDEditor, Verbose, TEXT("[%hs] was called with an invalid solver ID. Data that references the invalid solver id will be ignored."), __func__);
 		return nullptr;
 	}
 
diff --git a/Engine/Plugins/ChaosVD/Source/ChaosVD/Public/ChaosVDModule.h b/Engine/Plugins/ChaosVD/Source/ChaosVD/Public/ChaosVDModule.h
--- a/Engine/Plugins/ChaosVD/Source/ChaosVD/Public/ChaosVDModule.h
+++ b/Engine/Plugins/ChaosVD/Source/ChaosVD/Public/ChaosVDModule.h
@@ -17,7 +17,7 @@ class FSpawnTabArgs;
 struct FGuid;
 
 /** Main module class for the Chaos Visual Debugger editor */
-DECLARE_LOG_CATEGORY_EXTERN(LogChaosVDEditor, Log, Log);
+DECLARE_LOG_CATEGORY_EXTERN(LogChaosVDEditor, Log, VeryVerbose);
 class FChaosVDModule : public IModuleInterface
 {
 public:
