[ChaosVD] Fix for an issue where the reset settings button is not working for some settings category.

Two settings objects were missing the correct class metadata specifiers we need.

diff --git a/Engine/Plugins/ChaosVD/Source/ChaosVD/Private/Settings/ChaosVDGeneralSettings.h b/Engine/Plugins/ChaosVD/Source/ChaosVD/Private/Settings/ChaosVDGeneralSettings.h
--- a/Engine/Plugins/ChaosVD/Source/ChaosVD/Private/Settings/ChaosVDGeneralSettings.h
+++ b/Engine/Plugins/ChaosVD/Source/ChaosVD/Private/Settings/ChaosVDGeneralSettings.h
@@ -8,7 +8,7 @@
 /**
  * General settings that controls how CVD behaves
  */
-UCLASS()
+UCLASS(config=ChaosVD, PerObjectConfig)
 class CHAOSVD_API UChaosVDGeneralSettings : public UChaosVDSettingsObjectBase
 {
 	GENERATED_BODY()
diff --git a/Engine/Plugins/ChaosVD/Source/ChaosVDBuiltInExtensions/Private/AccelerationStructures/Settings/ChaosVDAccelerationStructureVisualizationSettings.h b/Engine/Plugins/ChaosVD/Source/ChaosVDBuiltInExtensions/Private/AccelerationStructures/Settings/ChaosVDAccelerationStructureVisualizationSettings.h
--- a/Engine/Plugins/ChaosVD/Source/ChaosVDBuiltInExtensions/Private/AccelerationStructures/Settings/ChaosVDAccelerationStructureVisualizationSettings.h
+++ b/Engine/Plugins/ChaosVD/Source/ChaosVDBuiltInExtensions/Private/AccelerationStructures/Settings/ChaosVDAccelerationStructureVisualizationSettings.h
@@ -36,7 +36,7 @@ ENUM_CLASS_FLAGS(EChaosVDAccelerationStructureDataVisualizationFlags);
 /**
  * Settings object that stores the values that control how acceleration structures are debug drawn
  */
-UCLASS()
+UCLASS(config=ChaosVD, PerObjectConfig)
 class UChaosVDAccelerationStructureVisualizationSettings : public UChaosVDVisualizationSettingsObjectBase
 {
 	GENERATED_BODY()
