[ChaosVD] Fix playback controls gets disabled when using network tick sync mode + other edge cases where the UI can get soft-locked

- Instead of disabling the main game frames playback timeline, when we are in network tick sync mode (which is incompatible). Now CVD's playback code will perform a time-stamp based sync to a compatible track (if any), and perform the network sync mode between the remaining tracks. This removes situations where due to edge cases all playback controls can get disabled, and removed the need to add extra UI feedback explaining you can't use the game frames timeline.
- Fixed an issue where the fallback to use a time-stamp sync mode when two tracks we try to sync together don't have valid net tick data, was not working as intended.
- Allow the next and prev buttons to be enabled when an active live session is paused
- Made it so playback doesn't stop after changing sync modes

diff --git a/Engine/Plugins/ChaosVD/Source/ChaosVD/Private/ChaosVDPlaybackController.cpp b/Engine/Plugins/ChaosVD/Source/ChaosVD/Private/ChaosVDPlaybackController.cpp
--- a/Engine/Plugins/ChaosVD/Source/ChaosVD/Private/ChaosVDPlaybackController.cpp
+++ b/Engine/Plugins/ChaosVD/Source/ChaosVD/Private/ChaosVDPlaybackController.cpp
@@ -73,12 +73,12 @@ bool FChaosVDPlaybackController::LoadChaosVDRecordingFromTraceSession(const FCha
 
 	LoadedRecording->SetIsLive(InSessionDescriptor.bIsLiveSession);
 
-	HandleCurrentRecordingUpdated();
-	
 	if (TSharedPtr<FChaosVDScene> ScenePtr = SceneToControl.Pin())
 	{
 		ScenePtr->SetLoadedRecording(LoadedRecording);
 	}
+
+	HandleCurrentRecordingUpdated();
 	
 	bHasPendingGTUpdateBroadcast = true;
 
@@ -502,14 +502,29 @@ int32 FChaosVDPlaybackController::ConvertCurrentFrameToOtherTrackFrame_AssumesLo
 									{
 										if (InFromTrack->bIsServer) // Server --> Client track
 										{
-											ToFrame = FromSolverFrameData->InternalFrameNumber - ToSolverFrameData->GetClampedNetworkTickOffset();
+											if (ToSolverFrameData->NetworkTickOffset != INDEX_NONE)
+											{
+												ToFrame = FromSolverFrameData->InternalFrameNumber - ToSolverFrameData->NetworkTickOffset;
+												
+											}
+											else
+											{
+												ToFrame = INDEX_NONE;
+											}
 										}
 										else // Client --> Client Track or Client --> Server Track
 										{
 											// This works for Client --> Server conversion because in that case we want to add the frame offset.
 											// As the tick offset in server tracks is 0,  the following calculation will return a negative offset, which it will result in the intended addition in the last caluclation
-											int32 FrameOffset = ToSolverFrameData->GetClampedNetworkTickOffset() - FromSolverFrameData->GetClampedNetworkTickOffset();
-											ToFrame = FromSolverFrameData->InternalFrameNumber - FrameOffset;
+											if (ToSolverFrameData->NetworkTickOffset != INDEX_NONE || FromSolverFrameData->NetworkTickOffset != INDEX_NONE)
+											{
+												int32 FrameOffset = ToSolverFrameData->NetworkTickOffset - FromSolverFrameData->NetworkTickOffset != INDEX_NONE;
+												ToFrame = FromSolverFrameData->InternalFrameNumber - FrameOffset;
+											}
+											else
+											{
+												ToFrame = INDEX_NONE;
+											}
 										}
 									}
 									else
@@ -521,7 +536,17 @@ int32 FChaosVDPlaybackController::ConvertCurrentFrameToOtherTrackFrame_AssumesLo
 									}
 								}
 
-								return ensure(ToFrame != INDEX_NONE) ? LoadedRecording->GetLowestSolverFrameNumberAtNetworkFrameNumber_AssumesLocked(InToTrack->TrackID, ToFrame) : INDEX_NONE;
+								// If we failed to re-map the track based on the network tick, just fallback to the timestamp sync logic.
+								// The UI shows a warning on all the conditions this can happen, which is when any of the tracks does not have net tick offset data
+
+								if (ToFrame != INDEX_NONE)
+								{
+									return LoadedRecording->GetLowestSolverFrameNumberAtNetworkFrameNumber_AssumesLocked(InToTrack->TrackID, ToFrame);
+								}
+								if (FromSolverFrameData)
+								{
+									return LoadedRecording->GetLowestSolverFrameNumberAtCycle_AssumesLocked(InToTrack->TrackID, FromSolverFrameData->FrameCycle);
+								}
 							}
 						default:
 							break;
@@ -931,11 +956,17 @@ void FChaosVDPlaybackController::HandleFramePlaybackControlInput(EChaosVDPlaybac
 		return;
 	}
 
+	if (InTrackInfoRef->TrackID != GameTrackID && !IsCompatibleWithSyncMode(InTrackInfoRef, CurrentSyncMode))
+	{
+		UE_LOG(LogChaosVDEditor, Warning, TEXT("[%s] Attempted to perform a playback action with an incompatible sync mode | Ignoring... "), ANSI_TO_TCHAR(__FUNCTION__));
+		return;
+	}
+
 	if (InTrackInfoRef != ActiveTrack)
 	{
 		if (ButtonID == EChaosVDPlaybackButtonsID::Play && ActiveTrack->bIsPlaying)
 		{
-			UE_LOG(LogChaosVDEditor, Warning, TEXT("[%s] Attempted to play when thera was another track playing. This should not happen, this is probably a widget with controls not properly disabled."), ANSI_TO_TCHAR(__FUNCTION__));
+			UE_LOG(LogChaosVDEditor, Warning, TEXT("[%s] Attempted to play when there was another track playing. This should not happen, this is probably a widget with controls not properly disabled."), ANSI_TO_TCHAR(__FUNCTION__));
 			ActiveTrack->bIsPlaying = false;
 		}
 	}
@@ -1094,6 +1125,28 @@ void FChaosVDPlaybackController::SyncTracks_AssumesLocked(const TSharedRef<const
 	{
 	case EChaosVDTrackType::Game:
 		{
+			// Game track types are old and will be deprecated in the future (UE-277464)
+			// These only support the default time-stamp sync mode.
+			// So if we are in an incompatible sync mode, instead of fully disabling its use, we first try to find a solver track that is compatible, advance that using the timestamp base mode
+			// and then sync the rest using the desired sync mode.
+			// This is ok because the actual data we visualize is on solver tracks, except scene queries (hence why we plan to deprecate game tracks in the future)
+			// This a common case in PIE, when the recording is started before starting PIE
+			if (!IsCompatibleWithSyncMode(FromTrack, CurrentSyncMode))
+			{
+				TSharedPtr<const FChaosVDTrackInfo> TrackToReSyncFrom = GetTrackToReSyncFromOnModeChange_AssumesLocked(CurrentSyncMode);
+
+				const bool bCanUseNewTrackForResync = TrackToReSyncFrom && TrackToReSyncFrom->TrackID != GameTrackID;
+				if (bCanUseNewTrackForResync)
+				{
+					const int32 SolverFrameNumber = LoadedRecording->GetLowestSolverFrameNumberGameFrame_AssumesLocked(TrackToReSyncFrom->TrackID, FromTrack->CurrentFrame);
+					const int32 StageNumber = GetTrackLastStageAtFrame_AssumesLocked(EChaosVDTrackType::Solver, TrackToReSyncFrom->TrackID, SolverFrameNumber);
+
+					GoToTrackFrame_AssumesLocked(PlaybackSelfInstigatorID, EChaosVDTrackType::Solver, TrackToReSyncFrom->TrackID, SolverFrameNumber, StageNumber);
+					
+					return SyncTracks_AssumesLocked(TrackToReSyncFrom.ToSharedRef(), TrackSyncMode);
+				}
+			}
+	
 			for (const TSharedPtr<const FChaosVDTrackInfo>& SolverTrack : AvailableSolverTracks)
 			{
 				if (!SolverTrack->bTrackSyncEnabled)
@@ -1159,35 +1212,75 @@ void FChaosVDPlaybackController::SetScrubbingTimeline(bool bNewIsScrubbingTimeli
 }
 
 TSharedPtr<const FChaosVDTrackInfo> FChaosVDPlaybackController::GetTrackToReSyncFromOnModeChange(EChaosVDSyncTimelinesMode SyncMode)
+{
+	if (!LoadedRecording)
+	{
+		return nullptr;
+	}
+
+	FReadScopeLock ReadLock(LoadedRecording->GetRecordingDataLock());
+	return GetTrackToReSyncFromOnModeChange_AssumesLocked(SyncMode);
+}
+
+TSharedPtr<const FChaosVDTrackInfo> FChaosVDPlaybackController::GetTrackToReSyncFromOnModeChange_AssumesLocked(EChaosVDSyncTimelinesMode SyncMode)
 {
 	switch (SyncMode)
 	{
-		case EChaosVDSyncTimelinesMode::RecordedTimestamp:
+	case EChaosVDSyncTimelinesMode::RecordedTimestamp:
+		{
+			return ActiveTrack->IsValidTrack() ? ActiveTrack : GetTrackInfo(EChaosVDTrackType::Game, GameTrackID);
+		}
+	case EChaosVDSyncTimelinesMode::NetworkTick:
+		{
+			// Network ticks can only sync from other solver tracks
+			if (ActiveTrack->IsValidTrack() && ActiveTrack->TrackType != EChaosVDTrackType::Game)
 			{
-				return ActiveTrack->IsValidTrack() ? ActiveTrack : GetTrackInfo(EChaosVDTrackType::Game, GameTrackID);
+				return ActiveTrack;
 			}
-		case EChaosVDSyncTimelinesMode::NetworkTick:
+
+			TSharedPtr<const FChaosVDTrackInfo> GameThreadTrack = GetTrackInfo(EChaosVDTrackType::Game, GameTrackID);
+			if (GameThreadTrack && GameThreadTrack->IsValidTrack())
 			{
-				// Network ticks can only sync from other solver tracks
+				TArray<TSharedPtr<const FChaosVDTrackInfo>> AvailableSolverTracks;
+				GetAvailableTrackInfosAtTrackFrame_AssumesLocked(EChaosVDTrackType::Solver, GameThreadTrack.ToSharedRef(), AvailableSolverTracks);
 
-				if (ActiveTrack->IsValidTrack() && ActiveTrack->TrackType != EChaosVDTrackType::Game)
+				for (const TSharedPtr<const FChaosVDTrackInfo>& Track : AvailableSolverTracks)
 				{
-					return ActiveTrack;
+					if (Track->bHasNetworkSyncData)
+					{
+						return Track;
+					}
 				}
 
-				TSharedPtr<const FChaosVDTrackInfo> GameThreadTrack = GetTrackInfo(EChaosVDTrackType::Game, GameTrackID);
-				if (GameThreadTrack && GameThreadTrack->IsValidTrack())
-				{
-					TArray<TSharedPtr<const FChaosVDTrackInfo>> AvailableSolverTracks;
-					GetAvailableTrackInfosAtTrackFrame(EChaosVDTrackType::Solver, GameThreadTrack.ToSharedRef(), AvailableSolverTracks);
+				return nullptr;
+			}
 
-					return AvailableSolverTracks.IsEmpty() ? nullptr : AvailableSolverTracks[0];
-				}
+			return nullptr;
+		}
+	default:
+		return nullptr;
+	}
+}
 
-				return nullptr;
+bool FChaosVDPlaybackController::IsCompatibleWithSyncMode(const TSharedRef<const FChaosVDTrackInfo>& InTrackRef, EChaosVDSyncTimelinesMode SyncMode)
+{
+	switch (SyncMode)
+	{
+		case EChaosVDSyncTimelinesMode::RecordedTimestamp:
+			return true;
+		case EChaosVDSyncTimelinesMode::NetworkTick:
+		{
+			if (InTrackRef->TrackID == GameTrackID)
+			{
+				return false;
 			}
+			return InTrackRef->bHasNetworkSyncData;
+		}
+		case EChaosVDSyncTimelinesMode::Manual:
+			return true;
+		case EChaosVDSyncTimelinesMode::None:
 		default:
-			return nullptr;
+			return false;
 	}
 }
 
@@ -1201,6 +1294,7 @@ void FChaosVDPlaybackController::SetTimelineSyncMode(EChaosVDSyncTimelinesMode S
 
 		if (TrackToReSyncFrom && TrackToReSyncFrom->IsValidTrack())
 		{
+			bool bWasIsPlaying = IsPlaying();
 			const TSharedRef<const FChaosVDTrackInfo> TrackToReSyncFromRef = TrackToReSyncFrom.ToSharedRef();
 
 			TrySetActiveTrack(TrackToReSyncFromRef);
@@ -1210,6 +1304,11 @@ void FChaosVDPlaybackController::SetTimelineSyncMode(EChaosVDSyncTimelinesMode S
 				SyncTracks(TrackToReSyncFromRef, CurrentSyncMode);
 			}
 
+			if (bWasIsPlaying)
+			{
+				HandleFramePlaybackControlInput(EChaosVDPlaybackButtonsID::Play, TrackToReSyncFromRef, PlaybackSelfInstigatorID);
+			}
+
 			OnDataUpdated().Broadcast(AsWeak());
 		}
 		else
diff --git a/Engine/Plugins/ChaosVD/Source/ChaosVD/Private/Widgets/SChaosVDGameFramesPlaybackControls.cpp b/Engine/Plugins/ChaosVD/Source/ChaosVD/Private/Widgets/SChaosVDGameFramesPlaybackControls.cpp
--- a/Engine/Plugins/ChaosVD/Source/ChaosVD/Private/Widgets/SChaosVDGameFramesPlaybackControls.cpp
+++ b/Engine/Plugins/ChaosVD/Source/ChaosVD/Private/Widgets/SChaosVDGameFramesPlaybackControls.cpp
@@ -88,33 +88,9 @@ bool SChaosVDGameFramesPlaybackControls::CanPlayback() const
 	{
 		if (PlaybackControllerPtr->IsRecordingLoaded() && GameTrackInfoRef->TrackType != EChaosVDTrackType::Invalid)
 		{
-			// GameFrames playback controls can only be enabled on the Recorded timestamp sync mode
-			bool bIsCompatibleSyncMode = true;
-
-			if (PlaybackControllerPtr->GetTimelineSyncMode() == EChaosVDSyncTimelinesMode::NetworkTick)
-			{
-				static TArray<TSharedPtr<const FChaosVDTrackInfo>> OutTrackInfo;
-				OutTrackInfo.Reset();
-
-				bool bHasTracksWithNetworkSyncData = false;
-				PlaybackControllerPtr->GetAvailableTrackInfosAtTrackFrame(EChaosVDTrackType::Solver, GameTrackInfoRef, OutTrackInfo);
-				for (const TSharedPtr<const FChaosVDTrackInfo>& TrackInfo : OutTrackInfo)
-				{
-					if (TrackInfo && TrackInfo->bHasNetworkSyncData)
-					{
-						bHasTracksWithNetworkSyncData = true;
-						break;
-					}
-				}
-
-				// If we are on Network tick sync mode, but we don't have ant solver track that supports that mode loaded yet (or ever if this was an old file or recorded a standalone game)
-				// allow the game track playback control to work. The solver track widgets themselves will show the corresponding compatibility warning
-				bIsCompatibleSyncMode &= !bHasTracksWithNetworkSyncData;
-			}
-
 			TSharedRef<FChaosVDTrackInfo> ActivePlaybackTrack = PlaybackControllerPtr->GetActiveTrackInfo();
 			bool bActiveTrackIsThisTrack = FChaosVDTrackInfo::AreSameTrack(ActivePlaybackTrack, GameTrackInfoRef);
-			bCanControlPlayback = bIsCompatibleSyncMode && (bActiveTrackIsThisTrack || (!bActiveTrackIsThisTrack && !ActivePlaybackTrack->bIsPlaying));
+			bCanControlPlayback = bActiveTrackIsThisTrack || (!bActiveTrackIsThisTrack && !ActivePlaybackTrack->bIsPlaying);
 		}
 	}
 
@@ -149,7 +125,9 @@ EChaosVDTimelineElementIDFlags SChaosVDGameFramesPlaybackControls::GetElementEna
 		if (ControllerSharedPtr->IsPlayingLiveSession())
 		{
 			constexpr EChaosVDTimelineElementIDFlags PlaybackElementDisabledDuringLiveSession = EChaosVDTimelineElementIDFlags::Stop | EChaosVDTimelineElementIDFlags::Next | EChaosVDTimelineElementIDFlags::Prev;
-			EnumRemoveFlags(EnabledButtonFlags, PlaybackElementDisabledDuringLiveSession);
+			constexpr EChaosVDTimelineElementIDFlags PlaybackElementDisabledDuringPausedLiveSession = EChaosVDTimelineElementIDFlags::Stop;
+
+			EnumRemoveFlags(EnabledButtonFlags, IsPlaying() ? PlaybackElementDisabledDuringLiveSession : PlaybackElementDisabledDuringPausedLiveSession);
 			return EnabledButtonFlags;
 		}
 		else
diff --git a/Engine/Plugins/ChaosVD/Source/ChaosVD/Private/Widgets/SChaosVDSolverPlaybackControls.cpp b/Engine/Plugins/ChaosVD/Source/ChaosVD/Private/Widgets/SChaosVDSolverPlaybackControls.cpp
--- a/Engine/Plugins/ChaosVD/Source/ChaosVD/Private/Widgets/SChaosVDSolverPlaybackControls.cpp
+++ b/Engine/Plugins/ChaosVD/Source/ChaosVD/Private/Widgets/SChaosVDSolverPlaybackControls.cpp
@@ -164,20 +164,19 @@ bool SChaosVDSolverPlaybackControls::CanPlayback() const
 
 	if (const TSharedPtr<FChaosVDPlaybackController> CurrentPlaybackControllerPtr = PlaybackController.Pin())
 	{
-
-		bCanControlPlayback = !CurrentPlaybackControllerPtr->IsPlayingLiveSession();
+		TSharedRef<FChaosVDTrackInfo> ActivePlaybackTrack = CurrentPlaybackControllerPtr->GetActiveTrackInfo();
+		bCanControlPlayback = CurrentPlaybackControllerPtr->IsPlayingLiveSession() ? !ActivePlaybackTrack->bIsPlaying : true;
 
 		// When it is not a live session, the Game Frames timeline follows the same rule as other timelines. The controls are locked unless we are who started a Play action
 		if (bCanControlPlayback)
 		{
-			bool bIsCompatibleSyncMode = CurrentPlaybackControllerPtr->GetTimelineSyncMode() == EChaosVDSyncTimelinesMode::NetworkTick ? SolverTrackInfoRef->bHasNetworkSyncData : true;
+			const bool bIsCompatibleSyncMode = CurrentPlaybackControllerPtr->IsCompatibleWithSyncMode(SolverTrackInfoRef, CurrentPlaybackControllerPtr->GetTimelineSyncMode());
 
-			if(!bIsCompatibleSyncMode)
+			if (!bIsCompatibleSyncMode)
 			{
 				return false;
 			}
 
-			TSharedRef<FChaosVDTrackInfo> ActivePlaybackTrack = CurrentPlaybackControllerPtr->GetActiveTrackInfo();
 			bool bActiveTrackIsThisTrack = FChaosVDTrackInfo::AreSameTrack(ActivePlaybackTrack, SolverTrackInfoRef);
 			bCanControlPlayback =  bActiveTrackIsThisTrack || (!bActiveTrackIsThisTrack && !ActivePlaybackTrack->bIsPlaying);
 		}
diff --git a/Engine/Plugins/ChaosVD/Source/ChaosVD/Public/ChaosVDPlaybackController.h b/Engine/Plugins/ChaosVD/Source/ChaosVD/Public/ChaosVDPlaybackController.h
--- a/Engine/Plugins/ChaosVD/Source/ChaosVD/Public/ChaosVDPlaybackController.h
+++ b/Engine/Plugins/ChaosVD/Source/ChaosVD/Public/ChaosVDPlaybackController.h
@@ -404,6 +404,15 @@ public:
 	 * @param SyncMode New timeline Sync Mode
 	 * */
 	TSharedPtr<const FChaosVDTrackInfo> GetTrackToReSyncFromOnModeChange(EChaosVDSyncTimelinesMode SyncMode);
+	TSharedPtr<const FChaosVDTrackInfo> GetTrackToReSyncFromOnModeChange_AssumesLocked(EChaosVDSyncTimelinesMode SyncMode);
+
+	/**
+	 * Checks if the provided track is compatible with a specific sync mode
+	 * @param InTrackRef Track to check against
+	 * @param SyncMode Timeline sync mode to verify if the provided track is compatible with
+	 * @return true if the provided track is compatible with the provided sync mode
+	 */
+	bool IsCompatibleWithSyncMode(const TSharedRef<const FChaosVDTrackInfo>& InTrackRef, EChaosVDSyncTimelinesMode SyncMode);
 
 	/** Sets a new sync mode to be used between tracks */
 	void SetTimelineSyncMode(EChaosVDSyncTimelinesMode SyncMode);
