[ChaosVD] Added support to open files by dragging and drop them into the viewport.

It supports:
- Normal CVD files (utrace)
- Combined CVD files (cvdmulti)
- Dragging and dropping multiple files at once. This will load them in multi-source mode, combining the data on the fly. Same as it works when opening multiple files from the Open File menu

diff --git a/Engine/Plugins/ChaosVD/Source/ChaosVD/Private/Widgets/SChaosVDMainTab.cpp b/Engine/Plugins/ChaosVD/Source/ChaosVD/Private/Widgets/SChaosVDMainTab.cpp
--- a/Engine/Plugins/ChaosVD/Source/ChaosVD/Private/Widgets/SChaosVDMainTab.cpp
+++ b/Engine/Plugins/ChaosVD/Source/ChaosVD/Private/Widgets/SChaosVDMainTab.cpp
@@ -592,7 +592,7 @@ TSharedRef<SWidget> SChaosVDMainTab::GenerateSettingsMenuWidget()
 
 void SChaosVDMainTab::LoadCVDFile(const FString& InFilename, EChaosVDLoadRecordedDataMode LoadingMode)
 {
-	if (ensure(InFilename.EndsWith(TEXT("utrace")) || InFilename.EndsWith("cvdmulti")))
+	if (ensure(IsSupportedFile(InFilename)))
 	{
 		GetChaosVDEngineInstance()->LoadRecording(InFilename, LoadingMode);
 	}
@@ -602,6 +602,25 @@ void SChaosVDMainTab::LoadCVDFile(const FString& InFilename, EChaosVDLoadRecorde
 	}
 }
 
+void SChaosVDMainTab::LoadCVDFiles(TConstArrayView<FString> InFilenames, EChaosVDLoadRecordedDataMode LoadingMode)
+{
+	// Ideally, we should not need to do this, but the UI to support multi source is not robust yet, The UI improvement task to avoid this scenario is planned as UE-197418
+	const bool bHasMultiFileData = InFilenames.Num() > 1 || (InFilenames.IsValidIndex(0) && InFilenames[0].EndsWith(TEXT("cvdmulti")));
+
+	if (bHasMultiFileData && LoadingMode == EChaosVDLoadRecordedDataMode::SingleSource)
+	{
+		GetChaosVDEngineInstance()->CloseActiveTraceSessions();
+		LoadingMode = EChaosVDLoadRecordedDataMode::MultiSource;
+	
+		UE_LOG(LogChaosVDEditor, Warning, TEXT("[%s] Single source mode was selected with multiple files. Overriding mode to Multi Source..."), ANSI_TO_TCHAR(__FUNCTION__));
+	}
+
+	for (const FString& Filename : InFilenames)
+	{
+		LoadCVDFile(Filename, LoadingMode);
+	}
+}
+
 TSharedRef<IDetailsView> SChaosVDMainTab::CreateDetailsView(const FDetailsViewArgs& InDetailsViewArgs)
 {
 	FPropertyEditorModule& PropertyEditorModule = FModuleManager::GetModuleChecked<FPropertyEditorModule>("PropertyEditor");
@@ -740,21 +759,7 @@ void SChaosVDMainTab::BrowseChaosVDRecordingFromFolder(FStringView FolderPath, E
 		);
 	}
 
-	// Ideally, we should not need to do this, but the UI to support multi source is not robust yet, The UI improvement task to avoid this scenario is planned as UE-197418
-	const bool bHasMultiFileData = OutOpenFilenames.Num() > 1 || (OutOpenFilenames.IsValidIndex(0) && OutOpenFilenames[0].EndsWith(TEXT("cvdmulti")));
-
-	if (const bool bShouldOverrideMode = bHasMultiFileData && LoadingMode == EChaosVDLoadRecordedDataMode::SingleSource)
-	{
-		GetChaosVDEngineInstance()->CloseActiveTraceSessions();
-		LoadingMode = EChaosVDLoadRecordedDataMode::MultiSource;
-	
-		UE_LOG(LogChaosVDEditor, Warning, TEXT("[%s] Single source mode was selected with multiple files. Overriding mode to Multi Source..."), ANSI_TO_TCHAR(__FUNCTION__));
-	}
-
-	for (const FString& Filename : OutOpenFilenames)
-	{
-		LoadCVDFile(Filename, LoadingMode);
-	}
+	LoadCVDFiles(OutOpenFilenames, LoadingMode);
 }
 
 bool SChaosVDMainTab::ConnectToLiveSession(int32 SessionID, const FString& InSessionAddress, EChaosVDLoadRecordedDataMode LoadingMode) const
@@ -762,6 +767,11 @@ bool SChaosVDMainTab::ConnectToLiveSession(int32 SessionID, const FString& InSes
 	return GetChaosVDEngineInstance()->ConnectToLiveSession(SessionID, InSessionAddress, LoadingMode);
 }
 
+bool SChaosVDMainTab::IsSupportedFile(const FString& InFilename)
+{
+	return InFilename.EndsWith(TEXT(".utrace")) || InFilename.EndsWith(TEXT(".cvdmulti"));
+}
+
 void SChaosVDMainTab::SetUpDisableCPUThrottlingDelegate()
 {
 	if (GEditor)
diff --git a/Engine/Plugins/ChaosVD/Source/ChaosVD/Private/Widgets/SChaosVDPlaybackViewport.cpp b/Engine/Plugins/ChaosVD/Source/ChaosVD/Private/Widgets/SChaosVDPlaybackViewport.cpp
--- a/Engine/Plugins/ChaosVD/Source/ChaosVD/Private/Widgets/SChaosVDPlaybackViewport.cpp
+++ b/Engine/Plugins/ChaosVD/Source/ChaosVD/Private/Widgets/SChaosVDPlaybackViewport.cpp
@@ -4,6 +4,7 @@
 
 #include "ChaosVDCommands.h"
 #include "ChaosVDEditorMode.h"
+#include "ChaosVDEngine.h"
 #include "ChaosVDPlaybackController.h"
 #include "ChaosVDPlaybackViewportClient.h"
 #include "ChaosVDScene.h"
@@ -367,6 +368,57 @@ void SChaosVDPlaybackViewport::OnFocusViewportToSelection()
 	PlaybackViewportClient->FocusOnSelectedObject();
 }
 
+FReply SChaosVDPlaybackViewport::OnDragOver(const FGeometry& MyGeometry, const FDragDropEvent& DragDropEvent)
+{
+	TSharedPtr<SChaosVDMainTab> AsCVDToolkitHost = StaticCastSharedPtr<SChaosVDMainTab>(EditorModeTools->GetToolkitHost());
+	if (!AsCVDToolkitHost)
+	{
+		return SEditorViewport::OnDragOver(MyGeometry, DragDropEvent);
+	}
+	
+	TSharedPtr<FDragDropOperation> Operation = DragDropEvent.GetOperation();
+	if (const bool bValidOperation = Operation.IsValid() && Operation->IsOfType<FExternalDragOperation>())
+	{
+		const TSharedPtr<FExternalDragOperation> AsExternalDragOperation = StaticCastSharedPtr<FExternalDragOperation>(Operation);
+		if (AsExternalDragOperation->HasFiles())
+		{
+			for (const FString& DraggedFile : AsExternalDragOperation->GetFiles())
+			{
+				if (!AsCVDToolkitHost->IsSupportedFile(DraggedFile))
+				{
+					return SEditorViewport::OnDragOver(MyGeometry, DragDropEvent);
+				}
+			}
+
+			return FReply::Handled();
+		}
+	}
+
+	return SEditorViewport::OnDragOver(MyGeometry, DragDropEvent);
+}
+
+FReply SChaosVDPlaybackViewport::OnDrop(const FGeometry& MyGeometry, const FDragDropEvent& DragDropEvent)
+{
+	TSharedPtr<SChaosVDMainTab> AsCVDToolkitHost = StaticCastSharedPtr<SChaosVDMainTab>(EditorModeTools->GetToolkitHost());
+	if (!AsCVDToolkitHost)
+	{
+		return SEditorViewport::OnDrop(MyGeometry, DragDropEvent);
+	}
+
+	TSharedPtr<FDragDropOperation> Operation = DragDropEvent.GetOperation();
+	if (const bool bValidOperation = Operation.IsValid() && Operation->IsOfType<FExternalDragOperation>())
+	{
+		const TSharedPtr<FExternalDragOperation> AsExternalDragOperation = StaticCastSharedPtr<FExternalDragOperation>(Operation);
+		if (AsExternalDragOperation->HasFiles())
+		{
+			AsCVDToolkitHost->LoadCVDFiles(AsExternalDragOperation->GetFiles(), EChaosVDLoadRecordedDataMode::SingleSource);
+			return FReply::Handled();
+		}
+	}
+
+	return SEditorViewport::OnDrop(MyGeometry, DragDropEvent);
+}
+
 EVisibility SChaosVDPlaybackViewport::GetTrackSelectorKeyVisibility() const
 {
 	if (TSharedPtr<SChaosVDMainTab> CVDToolkitHost = StaticCastSharedPtr<SChaosVDMainTab>(EditorModeTools->GetToolkitHost()))
diff --git a/Engine/Plugins/ChaosVD/Source/ChaosVD/Private/Widgets/SChaosVDPlaybackViewport.h b/Engine/Plugins/ChaosVD/Source/ChaosVD/Private/Widgets/SChaosVDPlaybackViewport.h
--- a/Engine/Plugins/ChaosVD/Source/ChaosVD/Private/Widgets/SChaosVDPlaybackViewport.h
+++ b/Engine/Plugins/ChaosVD/Source/ChaosVD/Private/Widgets/SChaosVDPlaybackViewport.h
@@ -67,6 +67,9 @@ public:
 
 protected:
 	
+	virtual FReply OnDragOver(const FGeometry& MyGeometry, const FDragDropEvent& DragDropEvent) override;
+	virtual FReply OnDrop(const FGeometry& MyGeometry, const FDragDropEvent& DragDropEvent) override;
+	
 	EVisibility GetTrackSelectorKeyVisibility() const;
 
 	virtual TSharedRef<FEditorViewportClient> MakeEditorViewportClient() override;
diff --git a/Engine/Plugins/ChaosVD/Source/ChaosVD/Public/Widgets/SChaosVDMainTab.h b/Engine/Plugins/ChaosVD/Source/ChaosVD/Public/Widgets/SChaosVDMainTab.h
--- a/Engine/Plugins/ChaosVD/Source/ChaosVD/Public/Widgets/SChaosVDMainTab.h
+++ b/Engine/Plugins/ChaosVD/Source/ChaosVD/Public/Widgets/SChaosVDMainTab.h
@@ -106,7 +106,26 @@ public:
 	
 	bool ConnectToLiveSession(int32 SessionID, const FString& InSessionAddress, EChaosVDLoadRecordedDataMode LoadingMode) const;
 
+	/**
+	 * Evaluates a filename to determine if it is a supported filetype
+	 * @param InFilename Filename to evaluate
+	 * @return true if the file is supported
+	 */
+	bool IsSupportedFile(const FString& InFilename);
+
+	/**
+	 * Load the provided file into the current CVD instance
+	 * @param InFilename Filename with path to load
+	 * @param LoadingMode Desired loading mode (Single source or multi source)
+	 */
 	void LoadCVDFile(const FString& InFilename, EChaosVDLoadRecordedDataMode LoadingMode);
+	
+	/**
+	 * Load the provided files into the current CVD instance
+	 * @param InFilenames Array view of Filename with paths to load
+	 * @param LoadingMode Desired loading mode (Single source or multi source)
+	 */
+	void LoadCVDFiles(TConstArrayView<FString> InFilenames, EChaosVDLoadRecordedDataMode LoadingMode);
 
 	TSharedRef<IDetailsView> CreateDetailsView(const FDetailsViewArgs& InDetailsViewArgs);
 
