[ChaosVD] Improved auto camera tracking controls for selected objects.

Now we use the editor's camera orbit mode for the auto track object mode, instead of manually updating the view port focus box to achieve the same thing. This now allows you to smoothly change the view rotation of distance while auto-tracking during playback.

Also implemented the key shortcuts to control the tracking distance with the scroll well and respect the camera speed settings.

With these changes, CVD's auto tracking mode behaves like the existing editor's camera modes, and by using the built-in editor implementation.

diff --git a/Engine/Plugins/ChaosVD/Source/ChaosVD/Private/ChaosVDPlaybackViewportClient.cpp b/Engine/Plugins/ChaosVD/Source/ChaosVD/Private/ChaosVDPlaybackViewportClient.cpp
--- a/Engine/Plugins/ChaosVD/Source/ChaosVD/Private/ChaosVDPlaybackViewportClient.cpp
+++ b/Engine/Plugins/ChaosVD/Source/ChaosVD/Private/ChaosVDPlaybackViewportClient.cpp
@@ -155,6 +155,28 @@ void FChaosVDPlaybackViewportClient::ToggleCanSelectTranslucentGeometry()
 	SetCanSelectTranslucentGeometry(!bAllowTranslucentHitProxies);
 }
 
+bool FChaosVDPlaybackViewportClient::Internal_InputKey(const FInputKeyEventArgs& EventArgs)
+{
+	if (bAutoTrackSelectedObject && !IsFlightCameraActive())
+	{
+		constexpr float TrackingViewZoomStepSize = 50.0f;
+		const float CameraSpeed = GetCameraSpeed();
+		const float TargetDistanceChangeStep = TrackingViewZoomStepSize * CameraSpeed;
+		if (EventArgs.Key == EKeys::MouseScrollUp)
+		{
+			TrackingViewDistance -= TargetDistanceChangeStep;
+		}
+		else if (EventArgs.Key == EKeys::MouseScrollDown)
+		{
+			TrackingViewDistance += TargetDistanceChangeStep;
+		}
+
+		TrackingViewDistance = TrackingViewDistance < 0 ? 0 : TrackingViewDistance;
+	}
+
+	return FEditorViewportClient::Internal_InputKey(EventArgs);
+}
+
 void FChaosVDPlaybackViewportClient::HandleFocusRequest(FBox BoxToFocusOn)
 {
 	FocusViewportOnBox(BoxToFocusOn);
@@ -175,17 +197,34 @@ void FChaosVDPlaybackViewportClient::HandleActorMoving(AActor* MovedActor) const
 	}
 }
 
-void FChaosVDPlaybackViewportClient::TrackSelectedObject()
+void FChaosVDPlaybackViewportClient::UpdateObjectTracking()
 {
-	if (!bAutoTrackSelectedObject || !ModeTools.IsValid())
+	if (!bAutoTrackSelectedObject)
 	{
+		ToggleOrbitCamera(false);
 		return;
 	}
 
-	FocusOnSelectedObject();
+	const FBox SelectionBounds = GetSelectionBounds();
+	ToggleOrbitCamera(static_cast<bool>(SelectionBounds.IsValid));
+
+	if (SelectionBounds.IsValid)
+	{
+		const float TargetViewDistance = TrackingViewDistance + static_cast<float>(SelectionBounds.GetExtent().Size());
+		SetViewLocationForOrbiting(SelectionBounds.GetCenter(), TargetViewDistance);
+	}
 }
 
 void FChaosVDPlaybackViewportClient::FocusOnSelectedObject()
+{
+	FBox SelectionBounds = GetSelectionBounds();
+	if (SelectionBounds.IsValid)
+	{
+		FocusViewportOnBox(SelectionBounds.ExpandBy(TrackingViewDistance), true);
+	}
+}
+
+FBox FChaosVDPlaybackViewportClient::GetSelectionBounds() const
 {
 	if (const TSharedPtr<FChaosVDScene> CVDSceneSharedPtr = CVDScene.Pin())
 	{
@@ -201,24 +240,25 @@ void FChaosVDPlaybackViewportClient::FocusOnSelectedObject()
 			using namespace Chaos::VD::TypedElementDataUtil;
 			if (FChaosVDSceneParticle* Particle = GetStructDataFromTypedElementHandle<FChaosVDSceneParticle>(SelectedParticlesHandles[0]))
 			{
-				FBox BoxToTrack = Particle->GetBoundingBox();
-				FocusViewportOnBox(BoxToTrack.ExpandBy(TrackingViewDistance), true);
-			}	
+				return Particle->GetBoundingBox();
+			}
 		}
 	}
+
+	return FBox();
 }
 
 void FChaosVDPlaybackViewportClient::UpdateMouseDelta()
 {
 	// Make sure we get the camera in the correct position before a mouse drag is handled
-	TrackSelectedObject();
+	UpdateObjectTracking();
 
 	FEditorViewportClient::UpdateMouseDelta();
 }
 
 void FChaosVDPlaybackViewportClient::HandleCVDSceneUpdated()
 {
-	TrackSelectedObject();
+	UpdateObjectTracking();
 	Invalidate();
 }
 
diff --git a/Engine/Plugins/ChaosVD/Source/ChaosVD/Private/ChaosVDPlaybackViewportClient.h b/Engine/Plugins/ChaosVD/Source/ChaosVD/Private/ChaosVDPlaybackViewportClient.h
--- a/Engine/Plugins/ChaosVD/Source/ChaosVD/Private/ChaosVDPlaybackViewportClient.h
+++ b/Engine/Plugins/ChaosVD/Source/ChaosVD/Private/ChaosVDPlaybackViewportClient.h
@@ -22,30 +22,52 @@ public:
 
 	void SetScene(TWeakPtr<FChaosVDScene> InScene);
 
-	virtual UWorld* GetWorld() const override { return CVDWorld; };
+	virtual UWorld* GetWorld() const override
+	{
+		return CVDWorld;
+	}
 
 	virtual void Draw(const FSceneView* View, FPrimitiveDrawInterface* PDI) override;
 	virtual void DrawCanvas(FViewport& InViewport, FSceneView& View, FCanvas& Canvas) override;
 
 	void ToggleObjectTrackingIfSelected();
-	bool IsAutoTrackingSelectedObject() const { return bAutoTrackSelectedObject; };
+
+	bool IsAutoTrackingSelectedObject() const
+	{
+		return bAutoTrackSelectedObject;
+	}
+
 	void SetAutoTrackingViewDistance(float NewDistance);
-	float GetAutoTrackingViewDistance() const { return TrackingViewDistance; }
+
+	float GetAutoTrackingViewDistance() const
+	{
+		return TrackingViewDistance;
+	}
 
 	void GoToLocation(const FVector& InLocation);
 
-	void TrackSelectedObject();
+	void UpdateObjectTracking();
 
 	void FocusOnSelectedObject();
+
+	FBox GetSelectionBounds() const;
 	
 	virtual void UpdateMouseDelta() override;
 
 	void HandleCVDSceneUpdated();
 
-	bool GetCanSelectTranslucentGeometry() const { return bAllowTranslucentHitProxies; };
+	bool GetCanSelectTranslucentGeometry() const
+	{
+		return bAllowTranslucentHitProxies;
+	}
+
 	void SetCanSelectTranslucentGeometry(bool bCanSelect);
 	void ToggleCanSelectTranslucentGeometry();
 
+protected:
+
+	virtual bool Internal_InputKey(const FInputKeyEventArgs& EventArgs) override;
+
 private:
 
 	void HandleFocusRequest(FBox BoxToFocusOn);
@@ -57,5 +79,6 @@ private:
 
 	bool bAutoTrackSelectedObject = false;
 	float TrackingViewDistance = 120.0f;
+
 	bool bAllowTranslucentHitProxies = true;
 };
