[ChaosVD] Added a way to disable enum entries in data visualization settings, and updated all existing visualization flag options to use it.

This is to address an issue that causes confusion (with good reason) in first users, where the enable a visualization flag of a specific category, but nothing happens because the general "Enable Draw" flag is not set.

diff --git a/Engine/Plugins/ChaosVD/Source/ChaosVD/Private/Settings/ChaosVDCharacterConstraintsVisualizationSettings.cpp b/Engine/Plugins/ChaosVD/Source/ChaosVD/Private/Settings/ChaosVDCharacterConstraintsVisualizationSettings.cpp
--- a/Engine/Plugins/ChaosVD/Source/ChaosVD/Private/Settings/ChaosVDCharacterConstraintsVisualizationSettings.cpp
+++ b/Engine/Plugins/ChaosVD/Source/ChaosVD/Private/Settings/ChaosVDCharacterConstraintsVisualizationSettings.cpp
@@ -3,6 +3,7 @@
 #include "ChaosVDCharacterConstraintsVisualizationSettings.h"
 
 #include "ChaosVDSettingsManager.h"
+#include "Utils/ChaosVDUserInterfaceUtils.h"
 
 
 void UChaosVDCharacterConstraintsVisualizationSettings::SetDataVisualizationFlags(EChaosVDCharacterGroundConstraintDataVisualizationFlags NewFlags)
@@ -23,3 +24,8 @@ EChaosVDCharacterGroundConstraintDataVisualizationFlags UChaosVDCharacterConstra
 
 	return EChaosVDCharacterGroundConstraintDataVisualizationFlags::None;
 }
+
+bool UChaosVDCharacterConstraintsVisualizationSettings::CanVisualizationFlagBeChangedByUI(uint32 Flag)
+{
+	return Chaos::VisualDebugger::Utils::ShouldVisFlagBeEnabledInUI(Flag, GlobalCharacterGroundConstraintDataVisualizationFlags, EChaosVDCharacterGroundConstraintDataVisualizationFlags::EnableDraw);
+}
diff --git a/Engine/Plugins/ChaosVD/Source/ChaosVD/Private/Settings/ChaosVDCharacterConstraintsVisualizationSettings.h b/Engine/Plugins/ChaosVD/Source/ChaosVD/Private/Settings/ChaosVDCharacterConstraintsVisualizationSettings.h
--- a/Engine/Plugins/ChaosVD/Source/ChaosVD/Private/Settings/ChaosVDCharacterConstraintsVisualizationSettings.h
+++ b/Engine/Plugins/ChaosVD/Source/ChaosVD/Private/Settings/ChaosVDCharacterConstraintsVisualizationSettings.h
@@ -76,6 +76,8 @@ public:
 	static void SetDataVisualizationFlags(EChaosVDCharacterGroundConstraintDataVisualizationFlags NewFlags);
 	static EChaosVDCharacterGroundConstraintDataVisualizationFlags GetDataVisualizationFlags();
 
+	virtual bool CanVisualizationFlagBeChangedByUI(uint32 Flag) override;
+
 private:
 	/** Set of flags to enable/disable visualization of specific character ground constraint data as debug draw */
 	UPROPERTY(config, meta = (Bitmask, BitmaskEnum = "/Script/ChaosVD.EChaosVDCharacterGroundConstraintDataVisualizationFlags"))
diff --git a/Engine/Plugins/ChaosVD/Source/ChaosVD/Private/Settings/ChaosVDCollisionVisualizationSettings.cpp b/Engine/Plugins/ChaosVD/Source/ChaosVD/Private/Settings/ChaosVDCollisionVisualizationSettings.cpp
--- a/Engine/Plugins/ChaosVD/Source/ChaosVD/Private/Settings/ChaosVDCollisionVisualizationSettings.cpp
+++ b/Engine/Plugins/ChaosVD/Source/ChaosVD/Private/Settings/ChaosVDCollisionVisualizationSettings.cpp
@@ -3,6 +3,7 @@
 #include "ChaosVDCollisionVisualizationSettings.h"
 
 #include "ChaosVDSettingsManager.h"
+#include "Utils/ChaosVDUserInterfaceUtils.h"
 
 
 void UChaosVDCollisionDataVisualizationSettings::SetDataVisualizationFlags(EChaosVDCollisionVisualizationFlags NewFlags)
@@ -24,3 +25,8 @@ EChaosVDCollisionVisualizationFlags UChaosVDCollisionDataVisualizationSettings::
 	return EChaosVDCollisionVisualizationFlags::None;
 }
 
+bool UChaosVDCollisionDataVisualizationSettings::CanVisualizationFlagBeChangedByUI(uint32 Flag)
+{
+	return Chaos::VisualDebugger::Utils::ShouldVisFlagBeEnabledInUI(Flag, CollisionDataVisualizationFlags, EChaosVDCollisionVisualizationFlags::EnableDraw);
+}
+
diff --git a/Engine/Plugins/ChaosVD/Source/ChaosVD/Private/Settings/ChaosVDCollisionVisualizationSettings.h b/Engine/Plugins/ChaosVD/Source/ChaosVD/Private/Settings/ChaosVDCollisionVisualizationSettings.h
--- a/Engine/Plugins/ChaosVD/Source/ChaosVD/Private/Settings/ChaosVDCollisionVisualizationSettings.h
+++ b/Engine/Plugins/ChaosVD/Source/ChaosVD/Private/Settings/ChaosVDCollisionVisualizationSettings.h
@@ -47,6 +47,8 @@ public:
 	static void SetDataVisualizationFlags(EChaosVDCollisionVisualizationFlags NewFlags);
 	static EChaosVDCollisionVisualizationFlags GetDataVisualizationFlags();
 
+	virtual bool CanVisualizationFlagBeChangedByUI(uint32 Flag) override;
+
 private:
 	UPROPERTY(config, meta = (Bitmask, BitmaskEnum = "/Script/ChaosVD.EChaosVDCollisionVisualizationFlags"))
 	/** Set of flags to enable/disable visualization of specific collision data as debug draw */
diff --git a/Engine/Plugins/ChaosVD/Source/ChaosVD/Private/Settings/ChaosVDCoreSettings.cpp b/Engine/Plugins/ChaosVD/Source/ChaosVD/Private/Settings/ChaosVDCoreSettings.cpp
--- a/Engine/Plugins/ChaosVD/Source/ChaosVD/Private/Settings/ChaosVDCoreSettings.cpp
+++ b/Engine/Plugins/ChaosVD/Source/ChaosVD/Private/Settings/ChaosVDCoreSettings.cpp
@@ -54,3 +54,8 @@ void UChaosVDVisualizationSettingsObjectBase::BroadcastSettingsChanged()
 		return false;
 	}));
 }
+
+bool UChaosVDVisualizationSettingsObjectBase::CanVisualizationFlagBeChangedByUI(uint32 Flag)
+{
+	return true;
+}
diff --git a/Engine/Plugins/ChaosVD/Source/ChaosVD/Private/Settings/ChaosVDJointConstraintVisualizationSettings.cpp b/Engine/Plugins/ChaosVD/Source/ChaosVD/Private/Settings/ChaosVDJointConstraintVisualizationSettings.cpp
--- a/Engine/Plugins/ChaosVD/Source/ChaosVD/Private/Settings/ChaosVDJointConstraintVisualizationSettings.cpp
+++ b/Engine/Plugins/ChaosVD/Source/ChaosVD/Private/Settings/ChaosVDJointConstraintVisualizationSettings.cpp
@@ -3,6 +3,7 @@
 #include "ChaosVDJointConstraintVisualizationSettings.h"
 
 #include "ChaosVDSettingsManager.h"
+#include "Utils/ChaosVDUserInterfaceUtils.h"
 
 
 void UChaosVDJointConstraintsVisualizationSettings::SetDataVisualizationFlags(EChaosVDJointsDataVisualizationFlags NewFlags)
@@ -24,3 +25,7 @@ EChaosVDJointsDataVisualizationFlags UChaosVDJointConstraintsVisualizationSettin
 	return EChaosVDJointsDataVisualizationFlags::None;
 }
 
+bool UChaosVDJointConstraintsVisualizationSettings::CanVisualizationFlagBeChangedByUI(uint32 Flag)
+{
+	return Chaos::VisualDebugger::Utils::ShouldVisFlagBeEnabledInUI(Flag, GlobalJointsDataVisualizationFlags, EChaosVDJointsDataVisualizationFlags::EnableDraw);
+}
diff --git a/Engine/Plugins/ChaosVD/Source/ChaosVD/Private/Settings/ChaosVDJointConstraintVisualizationSettings.h b/Engine/Plugins/ChaosVD/Source/ChaosVD/Private/Settings/ChaosVDJointConstraintVisualizationSettings.h
--- a/Engine/Plugins/ChaosVD/Source/ChaosVD/Private/Settings/ChaosVDJointConstraintVisualizationSettings.h
+++ b/Engine/Plugins/ChaosVD/Source/ChaosVD/Private/Settings/ChaosVDJointConstraintVisualizationSettings.h
@@ -71,6 +71,8 @@ public:
 	static void SetDataVisualizationFlags(EChaosVDJointsDataVisualizationFlags NewFlags);
 	static EChaosVDJointsDataVisualizationFlags GetDataVisualizationFlags();
 
+	virtual bool CanVisualizationFlagBeChangedByUI(uint32 Flag) override;
+
 private:
 	/** Set of flags to enable/disable visualization of specific joint constraints data as debug draw */
 	UPROPERTY(config, meta = (Bitmask, BitmaskEnum = "/Script/ChaosVD.EChaosVDJointsDataVisualizationFlags"))
diff --git a/Engine/Plugins/ChaosVD/Source/ChaosVD/Private/Settings/ChaosVDParticleVisualizationSettings.cpp b/Engine/Plugins/ChaosVD/Source/ChaosVD/Private/Settings/ChaosVDParticleVisualizationSettings.cpp
--- a/Engine/Plugins/ChaosVD/Source/ChaosVD/Private/Settings/ChaosVDParticleVisualizationSettings.cpp
+++ b/Engine/Plugins/ChaosVD/Source/ChaosVD/Private/Settings/ChaosVDParticleVisualizationSettings.cpp
@@ -4,6 +4,7 @@
 
 #include "ChaosVDSettingsManager.h"
 #include "DataWrappers/ChaosVDParticleDataWrapper.h"
+#include "Utils/ChaosVDUserInterfaceUtils.h"
 
 
 FColor FChaosDebugDrawColorsByState::GetColorFromState(EChaosVDObjectStateType State) const
@@ -104,6 +105,11 @@ EChaosVDParticleDataVisualizationFlags UChaosVDParticleVisualizationDebugDrawSet
 	return EChaosVDParticleDataVisualizationFlags::None;
 }
 
+bool UChaosVDParticleVisualizationDebugDrawSettings::CanVisualizationFlagBeChangedByUI(uint32 Flag)
+{
+	return Chaos::VisualDebugger::Utils::ShouldVisFlagBeEnabledInUI(Flag, ParticleDataVisualizationFlags, EChaosVDParticleDataVisualizationFlags::EnableDraw);
+}
+
 void UChaosVDParticleVisualizationSettings::SetGeometryVisualizationFlags(EChaosVDGeometryVisibilityFlags Flags)
 {
 	if (UChaosVDParticleVisualizationSettings* Settings = FChaosVDSettingsManager::Get().GetSettingsObject<UChaosVDParticleVisualizationSettings>())
diff --git a/Engine/Plugins/ChaosVD/Source/ChaosVD/Private/Settings/ChaosVDParticleVisualizationSettings.h b/Engine/Plugins/ChaosVD/Source/ChaosVD/Private/Settings/ChaosVDParticleVisualizationSettings.h
--- a/Engine/Plugins/ChaosVD/Source/ChaosVD/Private/Settings/ChaosVDParticleVisualizationSettings.h
+++ b/Engine/Plugins/ChaosVD/Source/ChaosVD/Private/Settings/ChaosVDParticleVisualizationSettings.h
@@ -289,6 +289,8 @@ public:
 
 	static EChaosVDParticleDataVisualizationFlags GetDataDebugDrawVisualizationFlags();
 
+	virtual bool CanVisualizationFlagBeChangedByUI(uint32 Flag) override;
+
 private:
 	/** Set of flags to enable/disable visualization of specific particle data as debug draw */
 	UPROPERTY(config, meta = (Bitmask, BitmaskEnum = "/Script/ChaosVD.EChaosVDParticleDataVisualizationFlags"))
diff --git a/Engine/Plugins/ChaosVD/Source/ChaosVD/Private/Settings/ChaosVDSceneQueryVisualizationSettings.cpp b/Engine/Plugins/ChaosVD/Source/ChaosVD/Private/Settings/ChaosVDSceneQueryVisualizationSettings.cpp
--- a/Engine/Plugins/ChaosVD/Source/ChaosVD/Private/Settings/ChaosVDSceneQueryVisualizationSettings.cpp
+++ b/Engine/Plugins/ChaosVD/Source/ChaosVD/Private/Settings/ChaosVDSceneQueryVisualizationSettings.cpp
@@ -3,6 +3,7 @@
 #include "ChaosVDSceneQueryVisualizationSettings.h"
 
 #include "ChaosVDSettingsManager.h"
+#include "Utils/ChaosVDUserInterfaceUtils.h"
 
 
 void UChaosVDSceneQueriesVisualizationSettings::SetDataVisualizationFlags(EChaosVDSceneQueryVisualizationFlags NewFlags)
@@ -23,3 +24,8 @@ EChaosVDSceneQueryVisualizationFlags UChaosVDSceneQueriesVisualizationSettings::
 
 	return EChaosVDSceneQueryVisualizationFlags::None;
 }
+
+bool UChaosVDSceneQueriesVisualizationSettings::CanVisualizationFlagBeChangedByUI(uint32 Flag)
+{
+	return Chaos::VisualDebugger::Utils::ShouldVisFlagBeEnabledInUI(Flag, GlobalSceneQueriesVisualizationFlags, EChaosVDSceneQueryVisualizationFlags::EnableDraw);
+}
diff --git a/Engine/Plugins/ChaosVD/Source/ChaosVD/Private/Settings/ChaosVDSceneQueryVisualizationSettings.h b/Engine/Plugins/ChaosVD/Source/ChaosVD/Private/Settings/ChaosVDSceneQueryVisualizationSettings.h
--- a/Engine/Plugins/ChaosVD/Source/ChaosVD/Private/Settings/ChaosVDSceneQueryVisualizationSettings.h
+++ b/Engine/Plugins/ChaosVD/Source/ChaosVD/Private/Settings/ChaosVDSceneQueryVisualizationSettings.h
@@ -43,6 +43,8 @@ public:
 	static void SetDataVisualizationFlags(EChaosVDSceneQueryVisualizationFlags NewFlags);
 	static EChaosVDSceneQueryVisualizationFlags GetDataVisualizationFlags();
 
+	virtual bool CanVisualizationFlagBeChangedByUI(uint32 Flag) override;
+
 	/** If true, any debug draw text available will be drawn */
 	UPROPERTY(config, EditAnywhere, Category=DebugDraw)
 	bool bShowText = false;
diff --git a/Engine/Plugins/ChaosVD/Source/ChaosVD/Private/Visualizers/ChaosVDParticleDataComponentVisualizer.cpp b/Engine/Plugins/ChaosVD/Source/ChaosVD/Private/Visualizers/ChaosVDParticleDataComponentVisualizer.cpp
--- a/Engine/Plugins/ChaosVD/Source/ChaosVD/Private/Visualizers/ChaosVDParticleDataComponentVisualizer.cpp
+++ b/Engine/Plugins/ChaosVD/Source/ChaosVD/Private/Visualizers/ChaosVDParticleDataComponentVisualizer.cpp
@@ -79,20 +79,21 @@ void FChaosVDParticleDataComponentVisualizer::RegisterVisualizerMenus()
 			}
 		});
 
+		using namespace Chaos::VisualDebugger::Utils;
+
 		FNewToolMenuDelegate ParticleDataVisualizationFlagsMenuBuilder = FNewToolMenuDelegate::CreateLambda([](UToolMenu* Menu)
 		{
 			if (Menu)
 			{
 				TSharedRef<SWidget> VisualizationFlagsWidget = SNew(SChaosVDEnumFlagsMenu<EChaosVDParticleDataVisualizationFlags>)
 					.CurrentValue_Static(&UChaosVDParticleVisualizationDebugDrawSettings::GetDataDebugDrawVisualizationFlags)
-					.OnEnumSelectionChanged_Static(&UChaosVDParticleVisualizationDebugDrawSettings::SetDataDebugDrawVisualizationFlags);
+					.OnEnumSelectionChanged_Static(&UChaosVDParticleVisualizationDebugDrawSettings::SetDataDebugDrawVisualizationFlags)
+					.IsFlagEnabled_Static(&ShouldSettingsObjectVisFlagBeEnabledInUI<UChaosVDParticleVisualizationDebugDrawSettings,EChaosVDParticleDataVisualizationFlags>);
 
 				Menu->AddMenuEntry(NAME_None, FToolMenuEntry::InitWidget("ParticleDebugDrawDataVisualizationFlags", VisualizationFlagsWidget,FText::GetEmpty()));
 			}
 		});
 		
-		using namespace Chaos::VisualDebugger::Utils;
-		
 		FNewToolMenuDelegate GeometryVisualizationSettingsMenuBuilder = FNewToolMenuDelegate::CreateStatic(&CreateMenuEntryForSettingsObject<UChaosVDParticleVisualizationSettings>, EChaosVDSaveSettingsOptions::ShowResetButton);
 		FNewToolMenuDelegate ParticleDataVisualizationSettingsMenuBuilder = FNewToolMenuDelegate::CreateStatic(&CreateMenuEntryForSettingsObject<UChaosVDParticleVisualizationDebugDrawSettings>, EChaosVDSaveSettingsOptions::ShowResetButton);
 		FNewToolMenuDelegate ParticleColorizationMenuBuilder = FNewToolMenuDelegate::CreateStatic(&CreateMenuEntryForSettingsObject<UChaosVDParticleVisualizationColorSettings>, EChaosVDSaveSettingsOptions::ShowResetButton);
diff --git a/Engine/Plugins/ChaosVD/Source/ChaosVD/Public/Settings/ChaosVDCoreSettings.h b/Engine/Plugins/ChaosVD/Source/ChaosVD/Public/Settings/ChaosVDCoreSettings.h
--- a/Engine/Plugins/ChaosVD/Source/ChaosVD/Public/Settings/ChaosVDCoreSettings.h
+++ b/Engine/Plugins/ChaosVD/Source/ChaosVD/Public/Settings/ChaosVDCoreSettings.h
@@ -71,6 +71,9 @@ class CHAOSVD_API UChaosVDVisualizationSettingsObjectBase : public UChaosVDSetti
 	GENERATED_BODY()
 protected:
 	virtual void BroadcastSettingsChanged() override;
+
+	/** Returns true if a visualization flag option of this setting option should be enabled in the UI */
+	virtual bool CanVisualizationFlagBeChangedByUI(uint32 Flag);
 };
 
 /** Core settings class for CVD */
diff --git a/Engine/Plugins/ChaosVD/Source/ChaosVD/Public/Utils/ChaosVDUserInterfaceUtils.h b/Engine/Plugins/ChaosVD/Source/ChaosVD/Public/Utils/ChaosVDUserInterfaceUtils.h
--- a/Engine/Plugins/ChaosVD/Source/ChaosVD/Public/Utils/ChaosVDUserInterfaceUtils.h
+++ b/Engine/Plugins/ChaosVD/Source/ChaosVD/Public/Utils/ChaosVDUserInterfaceUtils.h
@@ -69,6 +69,17 @@ namespace Chaos::VisualDebugger::Utils
 		InDetailsView->SetStructureData(StructDataView);
 	}
 
+	template <typename ObjectSettingsType, typename VisualizationFlagsType>
+	bool ShouldSettingsObjectVisFlagBeEnabledInUI(VisualizationFlagsType Flag)
+	{
+		if (ObjectSettingsType* Settings = FChaosVDSettingsManager::Get().GetSettingsObject<ObjectSettingsType>())
+		{
+			return Settings->CanVisualizationFlagBeChangedByUI(static_cast<uint32>(Flag));
+		}
+
+		return true;
+	}
+
 	template <typename ObjectSettingsType, typename VisualizationFlagsType>
 	void CreateVisualizationOptionsMenuSections(UToolMenu* Menu, FName SectionName, const FText& InSectionLabel, const FText& InFlagsMenuLabel,  const FText& InFlagsMenuTooltip, FSlateIcon FlagsMenuIcon,  const FText& InSettingsMenuLabel, const FText& InSettingsMenuTooltip)
 	{
@@ -78,7 +89,8 @@ namespace Chaos::VisualDebugger::Utils
 						   {
 							   TSharedRef<SWidget> VisualizationFlagsWidget = SNew(SChaosVDEnumFlagsMenu<VisualizationFlagsType>)
 								   .CurrentValue_Static(&ObjectSettingsType::GetDataVisualizationFlags)
-								   .OnEnumSelectionChanged_Static(&ObjectSettingsType::SetDataVisualizationFlags);
+								   .OnEnumSelectionChanged_Static(&ObjectSettingsType::SetDataVisualizationFlags)
+									.IsFlagEnabled_Static(&ShouldSettingsObjectVisFlagBeEnabledInUI<ObjectSettingsType,VisualizationFlagsType>);
 			
 							   FToolMenuEntry FlagsMenuEntry = FToolMenuEntry::InitWidget("VisualizationFlags", VisualizationFlagsWidget,FText::GetEmpty());
 							   Menu->AddMenuEntry(NAME_None, FlagsMenuEntry);
@@ -89,4 +101,25 @@ namespace Chaos::VisualDebugger::Utils
 		Section.AddSubMenu(FName(InSettingsMenuLabel.ToString()), InSettingsMenuLabel, InSettingsMenuTooltip, FNewToolMenuDelegate::CreateStatic(&CreateMenuEntryForSettingsObject<ObjectSettingsType>, EChaosVDSaveSettingsOptions::ShowResetButton),
 						   false, FSlateIcon(FAppStyle::Get().GetStyleSetName(), TEXT("Icons.Toolbar.Settings")));
 	}
+
+	/**
+	 * Evaluates a flag and determines if it should be considered enabled in the UI,
+	 * based on the provided current active flags and the general enable draw flags 
+	 * @tparam VisFlagsType Type of the enum visualization flags
+	 * @return true if the flag should be enabled in the UI
+	 */
+	template<typename VisFlagsType>
+	bool ShouldVisFlagBeEnabledInUI(uint32 FlagToEvaluate, uint32 CurrentFlags, VisFlagsType EnableDrawFlags)
+	{
+		const VisFlagsType FlagToEvaluateAsEnum = static_cast<VisFlagsType>(FlagToEvaluate);
+		const VisFlagsType CurrentFlagsAsEnum = static_cast<VisFlagsType>(CurrentFlags);
+
+		// If the flags is any of the flags that will enable drawing, it needs to be always enabled
+		if (EnumHasAnyFlags(EnableDrawFlags, FlagToEvaluateAsEnum))
+		{
+			return true;
+		}
+
+		return EnumHasAnyFlags(CurrentFlagsAsEnum, EnableDrawFlags);
+	}
 }
diff --git a/Engine/Plugins/ChaosVD/Source/ChaosVD/Public/Widgets/SChaosVDEnumFlagsMenu.h b/Engine/Plugins/ChaosVD/Source/ChaosVD/Public/Widgets/SChaosVDEnumFlagsMenu.h
--- a/Engine/Plugins/ChaosVD/Source/ChaosVD/Public/Widgets/SChaosVDEnumFlagsMenu.h
+++ b/Engine/Plugins/ChaosVD/Source/ChaosVD/Public/Widgets/SChaosVDEnumFlagsMenu.h
@@ -17,6 +17,7 @@ class SChaosVDEnumFlagsMenu : public SCompoundWidget
 	
 public:
 	DECLARE_DELEGATE_OneParam(FOnEnumSelectionChanged, TEnum);
+	DECLARE_DELEGATE_RetVal_OneParam(bool, FIsEnumValueEnabled, TEnum);
 	struct FEnumInfo
 	{
 		explicit FEnumInfo(const int32 InIndex, const TEnum InValue, const FText& InDisplayName, const FText& InTooltipText)
@@ -35,6 +36,7 @@ public:
 
 	SLATE_ATTRIBUTE(TEnum, CurrentValue)
 	SLATE_EVENT(FOnEnumSelectionChanged, OnEnumSelectionChanged)
+	SLATE_EVENT(FIsEnumValueEnabled, IsFlagEnabled)
 
 		
 	SLATE_END_ARGS()
@@ -44,9 +46,12 @@ public:
 	void UpdateEnumFlagValue(FEnumInfo FlagInfo);
 	bool IsEnumFlagSet(FEnumInfo FlagInfo) const;
 
+	bool IsEnumFlagEnabled(TEnum FlagValue) const;
+
 	TAttribute<TEnum> CurrentValue;
 	
 	FOnEnumSelectionChanged OnEnumSelectionChangedDelegate;
+	FIsEnumValueEnabled EnumValueEnabledDelegate;
 	
 	const UEnum* Enum = nullptr;
 };
@@ -68,6 +73,7 @@ void SChaosVDEnumFlagsMenu<TEnum>::Construct(const FArguments& InArgs)
 
 	CurrentValue = InArgs._CurrentValue;
 	OnEnumSelectionChangedDelegate = InArgs._OnEnumSelectionChanged;
+	EnumValueEnabledDelegate = InArgs._IsFlagEnabled;
 
 	using namespace Chaos::VisualDebugger;
 	
@@ -123,7 +129,7 @@ void SChaosVDEnumFlagsMenu<TEnum>::Construct(const FArguments& InArgs)
 				FSlateIcon(),
 				FUIAction(
 				FExecuteAction::CreateSP(this, &SChaosVDEnumFlagsMenu::UpdateEnumFlagValue, EnumEntryInfo),
-				FCanExecuteAction::CreateLambda([](){ return true; }),
+				FCanExecuteAction::CreateSP(this, &SChaosVDEnumFlagsMenu::IsEnumFlagEnabled, EnumEntryInfo.Value),
 				FIsActionChecked::CreateSP(this, &SChaosVDEnumFlagsMenu::IsEnumFlagSet, EnumEntryInfo)),
 				NAME_None, EUserInterfaceActionType::ToggleButton);
 		}
@@ -149,3 +155,9 @@ bool SChaosVDEnumFlagsMenu<TEnum>::IsEnumFlagSet(FEnumInfo FlagInfo) const
 {
 	return EnumHasAnyFlags(CurrentValue.Get(), FlagInfo.Value);
 }
+
+template <typename TEnum>
+bool SChaosVDEnumFlagsMenu<TEnum>::IsEnumFlagEnabled(TEnum FlagValue) const
+{
+	return EnumValueEnabledDelegate.IsBound() ? EnumValueEnabledDelegate.Execute(FlagValue) : true;
+}
diff --git a/Engine/Plugins/ChaosVD/Source/ChaosVDBuiltInExtensions/Private/AccelerationStructures/Settings/ChaosVDAccelerationStructureVisualizationSettings.cpp b/Engine/Plugins/ChaosVD/Source/ChaosVDBuiltInExtensions/Private/AccelerationStructures/Settings/ChaosVDAccelerationStructureVisualizationSettings.cpp
--- a/Engine/Plugins/ChaosVD/Source/ChaosVDBuiltInExtensions/Private/AccelerationStructures/Settings/ChaosVDAccelerationStructureVisualizationSettings.cpp
+++ b/Engine/Plugins/ChaosVD/Source/ChaosVDBuiltInExtensions/Private/AccelerationStructures/Settings/ChaosVDAccelerationStructureVisualizationSettings.cpp
@@ -3,6 +3,8 @@
 #include "ChaosVDAccelerationStructureVisualizationSettings.h"
 
 #include "ChaosVDSettingsManager.h"
+#include "Utils/ChaosVDUserInterfaceUtils.h"
+
 
 void UChaosVDAccelerationStructureVisualizationSettings::SetDataVisualizationFlags(EChaosVDAccelerationStructureDataVisualizationFlags NewFlags)
 {
@@ -22,3 +24,8 @@ EChaosVDAccelerationStructureDataVisualizationFlags UChaosVDAccelerationStructur
 
 	return EChaosVDAccelerationStructureDataVisualizationFlags::None;
 }
+
+bool UChaosVDAccelerationStructureVisualizationSettings::CanVisualizationFlagBeChangedByUI(uint32 Flag)
+{
+	return Chaos::VisualDebugger::Utils::ShouldVisFlagBeEnabledInUI(Flag, AccelerationStructureDataVisualizationFlags, EChaosVDAccelerationStructureDataVisualizationFlags::EnableDraw);
+}
diff --git a/Engine/Plugins/ChaosVD/Source/ChaosVDBuiltInExtensions/Private/AccelerationStructures/Settings/ChaosVDAccelerationStructureVisualizationSettings.h b/Engine/Plugins/ChaosVD/Source/ChaosVDBuiltInExtensions/Private/AccelerationStructures/Settings/ChaosVDAccelerationStructureVisualizationSettings.h
--- a/Engine/Plugins/ChaosVD/Source/ChaosVDBuiltInExtensions/Private/AccelerationStructures/Settings/ChaosVDAccelerationStructureVisualizationSettings.h
+++ b/Engine/Plugins/ChaosVD/Source/ChaosVDBuiltInExtensions/Private/AccelerationStructures/Settings/ChaosVDAccelerationStructureVisualizationSettings.h
@@ -52,6 +52,8 @@ public:
 	
 	static void SetDataVisualizationFlags(EChaosVDAccelerationStructureDataVisualizationFlags NewFlags);
 	static EChaosVDAccelerationStructureDataVisualizationFlags GetDataVisualizationFlags();
+	
+	virtual bool CanVisualizationFlagBeChangedByUI(uint32 Flag) override;
 private:
 	UPROPERTY(config, meta = (Bitmask, BitmaskEnum = "/Script/ChaosVD.EChaosVDAccelerationStructureDataVisualizationFlags"))
 	uint32 AccelerationStructureDataVisualizationFlags = static_cast<uint32>(EChaosVDAccelerationStructureDataVisualizationFlags::DrawNodesBounds | EChaosVDAccelerationStructureDataVisualizationFlags::DrawBranches);
diff --git a/Engine/Plugins/ChaosVD/Source/ChaosVDBuiltInExtensions/Private/GenericDebugDraw/Settings/ChaosVDGenericDebugDrawSettings.cpp b/Engine/Plugins/ChaosVD/Source/ChaosVDBuiltInExtensions/Private/GenericDebugDraw/Settings/ChaosVDGenericDebugDrawSettings.cpp
--- a/Engine/Plugins/ChaosVD/Source/ChaosVDBuiltInExtensions/Private/GenericDebugDraw/Settings/ChaosVDGenericDebugDrawSettings.cpp
+++ b/Engine/Plugins/ChaosVD/Source/ChaosVDBuiltInExtensions/Private/GenericDebugDraw/Settings/ChaosVDGenericDebugDrawSettings.cpp
@@ -3,6 +3,8 @@
 #include "ChaosVDGenericDebugDrawSettings.h"
 
 #include "ChaosVDSettingsManager.h"
+#include "Utils/ChaosVDUserInterfaceUtils.h"
+
 
 void UChaosVDGenericDebugDrawSettings::SetDataVisualizationFlags(EChaosVDGenericDebugDrawVisualizationFlags NewFlags)
 {
@@ -22,3 +24,8 @@ EChaosVDGenericDebugDrawVisualizationFlags UChaosVDGenericDebugDrawSettings::Get
 
 	return EChaosVDGenericDebugDrawVisualizationFlags::None;
 }
+
+bool UChaosVDGenericDebugDrawSettings::CanVisualizationFlagBeChangedByUI(uint32 Flag)
+{
+	return Chaos::VisualDebugger::Utils::ShouldVisFlagBeEnabledInUI(Flag, DebugDrawFlags, EChaosVDGenericDebugDrawVisualizationFlags::EnableDraw);
+}
diff --git a/Engine/Plugins/ChaosVD/Source/ChaosVDBuiltInExtensions/Private/GenericDebugDraw/Settings/ChaosVDGenericDebugDrawSettings.h b/Engine/Plugins/ChaosVD/Source/ChaosVDBuiltInExtensions/Private/GenericDebugDraw/Settings/ChaosVDGenericDebugDrawSettings.h
--- a/Engine/Plugins/ChaosVD/Source/ChaosVDBuiltInExtensions/Private/GenericDebugDraw/Settings/ChaosVDGenericDebugDrawSettings.h
+++ b/Engine/Plugins/ChaosVD/Source/ChaosVDBuiltInExtensions/Private/GenericDebugDraw/Settings/ChaosVDGenericDebugDrawSettings.h
@@ -40,6 +40,8 @@ public:
 	static void SetDataVisualizationFlags(EChaosVDGenericDebugDrawVisualizationFlags NewFlags);
 	static EChaosVDGenericDebugDrawVisualizationFlags GetDataVisualizationFlags();
 
+	virtual bool CanVisualizationFlagBeChangedByUI(uint32 Flag) override;
+
 private:
 
 	/** Set of flags to enable/disable visualization of debug draw data shapes */
diff --git a/Engine/Plugins/Experimental/Mover/Source/MoverCVDEditor/Private/MoverCVDSimDataSettings.cpp b/Engine/Plugins/Experimental/Mover/Source/MoverCVDEditor/Private/MoverCVDSimDataSettings.cpp
--- a/Engine/Plugins/Experimental/Mover/Source/MoverCVDEditor/Private/MoverCVDSimDataSettings.cpp
+++ b/Engine/Plugins/Experimental/Mover/Source/MoverCVDEditor/Private/MoverCVDSimDataSettings.cpp
@@ -3,6 +3,7 @@
 #include "MoverCVDSimDataSettings.h"
 
 #include "ChaosVDSettingsManager.h"
+#include "Utils/ChaosVDUserInterfaceUtils.h"
 
 void UMoverCVDSimDataSettings::SetDataVisualizationFlags(EMoverCVDSimDataVisualizationFlags NewFlags)
 {
@@ -22,3 +23,8 @@ EMoverCVDSimDataVisualizationFlags UMoverCVDSimDataSettings::GetDataVisualizatio
 
 	return EMoverCVDSimDataVisualizationFlags::None;
 }
+
+bool UMoverCVDSimDataSettings::CanVisualizationFlagBeChangedByUI(uint32 Flag)
+{
+	return Chaos::VisualDebugger::Utils::ShouldVisFlagBeEnabledInUI(Flag, DebugDrawFlags, EMoverCVDSimDataVisualizationFlags::EnableDraw);
+}
diff --git a/Engine/Plugins/Experimental/Mover/Source/MoverCVDEditor/Private/MoverCVDSimDataSettings.h b/Engine/Plugins/Experimental/Mover/Source/MoverCVDEditor/Private/MoverCVDSimDataSettings.h
--- a/Engine/Plugins/Experimental/Mover/Source/MoverCVDEditor/Private/MoverCVDSimDataSettings.h
+++ b/Engine/Plugins/Experimental/Mover/Source/MoverCVDEditor/Private/MoverCVDSimDataSettings.h
@@ -36,6 +36,8 @@ public:
 	static void SetDataVisualizationFlags(EMoverCVDSimDataVisualizationFlags NewFlags);
 	static EMoverCVDSimDataVisualizationFlags GetDataVisualizationFlags();
 
+	virtual bool CanVisualizationFlagBeChangedByUI(uint32 Flag) override;
+
 private:
 
 	/** Set of flags to enable/disable visualization of debug draw data shapes */
