[ChaosVD] Fix for an issue where the solver visibility control in the solver tracks widget can get stuck.

The issue is that we cache the visibility state for perf reasons, and rely on visibility changes events to update it.

The problem is that we were not covering the case where the solver track control widget might be created after the track is no longer visible, therefore as the default state in the widget is visible, you can't toggle it anymore.

This CL add a step to get and cache the correct visibility state on during the initialization of the playback control widget, and ensure the  playback logic triggers at least one visibility change state with the correct value, when a new track is initialized.

diff --git a/Engine/Plugins/ChaosVD/Source/ChaosVD/Private/Actors/ChaosVDSolverInfoActor.cpp b/Engine/Plugins/ChaosVD/Source/ChaosVD/Private/Actors/ChaosVDSolverInfoActor.cpp
--- a/Engine/Plugins/ChaosVD/Source/ChaosVD/Private/Actors/ChaosVDSolverInfoActor.cpp
+++ b/Engine/Plugins/ChaosVD/Source/ChaosVD/Private/Actors/ChaosVDSolverInfoActor.cpp
@@ -48,6 +48,8 @@ void AChaosVDSolverInfoActor::SetScene(TWeakPtr<FChaosVDScene> InScene)
 	if (TSharedPtr<FChaosVDScene> ScenePtr = InScene.Pin())
 	{
 		RegisterSelectionSetObject(ScenePtr->GetElementSelectionSet());
+
+		ScenePtr->OnSolverVisibilityUpdated().Broadcast(SolverDataID, IsVisible());
 	}
 
 	TInlineComponentArray<UChaosVDSolverDataComponent*> SolverDataComponents(this);
diff --git a/Engine/Plugins/ChaosVD/Source/ChaosVD/Private/ChaosVDPlaybackController.cpp b/Engine/Plugins/ChaosVD/Source/ChaosVD/Private/ChaosVDPlaybackController.cpp
--- a/Engine/Plugins/ChaosVD/Source/ChaosVD/Private/ChaosVDPlaybackController.cpp
+++ b/Engine/Plugins/ChaosVD/Source/ChaosVD/Private/ChaosVDPlaybackController.cpp
@@ -949,6 +949,27 @@ void FChaosVDPlaybackController::UpdateTrackVisibility(EChaosVDTrackType Type, i
 	}
 }
 
+bool FChaosVDPlaybackController::IsTrackVisible(EChaosVDTrackType Type, int32 TrackID)
+{
+	switch (Type)
+    {
+    	case EChaosVDTrackType::Solver:
+    		{
+    			if (const TSharedPtr<FChaosVDScene> ScenePtr = SceneToControl.Pin())
+    			{
+    				if (AChaosVDSolverInfoActor* SolverActorInfo = ScenePtr->GetSolverInfoActor(TrackID))
+    				{
+    					return SolverActorInfo->IsVisible();
+    				}
+    			}
+    		return false;
+    		}
+    	case EChaosVDTrackType::Game:
+    	default:
+    			return ensure(false);	
+    }
+}
+
 void FChaosVDPlaybackController::HandleFramePlaybackControlInput(EChaosVDPlaybackButtonsID ButtonID, const TSharedRef<const FChaosVDTrackInfo>& InTrackInfoRef, FGuid Instigator)
 {
 	if (!IsRecordingLoaded())
diff --git a/Engine/Plugins/ChaosVD/Source/ChaosVD/Private/Widgets/SChaosVDSolverPlaybackControls.cpp b/Engine/Plugins/ChaosVD/Source/ChaosVD/Private/Widgets/SChaosVDSolverPlaybackControls.cpp
--- a/Engine/Plugins/ChaosVD/Source/ChaosVD/Private/Widgets/SChaosVDSolverPlaybackControls.cpp
+++ b/Engine/Plugins/ChaosVD/Source/ChaosVD/Private/Widgets/SChaosVDSolverPlaybackControls.cpp
@@ -376,6 +376,7 @@ void SChaosVDSolverPlaybackControls::RegisterNewController(TWeakPtr<FChaosVDPlay
 	{
 		if (TSharedPtr<FChaosVDScene> Scene = NewPlaybackControllerPtr->GetControllerScene().Pin())
 		{
+			bIsVisible = NewPlaybackControllerPtr->IsTrackVisible(SolverTrackInfoRef->TrackType, SolverTrackInfoRef->TrackID);
 			Scene->OnSolverVisibilityUpdated().AddRaw(this, &SChaosVDSolverPlaybackControls::HandleSolverVisibilityChanged);
 		}
 	}
diff --git a/Engine/Plugins/ChaosVD/Source/ChaosVD/Public/ChaosVDPlaybackController.h b/Engine/Plugins/ChaosVD/Source/ChaosVD/Public/ChaosVDPlaybackController.h
--- a/Engine/Plugins/ChaosVD/Source/ChaosVD/Public/ChaosVDPlaybackController.h
+++ b/Engine/Plugins/ChaosVD/Source/ChaosVD/Public/ChaosVDPlaybackController.h
@@ -335,6 +335,14 @@ public:
 	 */
 	void UpdateTrackVisibility(EChaosVDTrackType Type, int32 TrackID, bool bNewVisibility);
 
+	/**
+	 * Return the current visibility state of the specified track
+	 * @param Type Type of the track (Game or Solver)
+	 * @param TrackID ID of the track
+	 * @return true if the track is visible
+	 */
+	bool IsTrackVisible(EChaosVDTrackType Type, int32 TrackID);
+
 	/**
 	 * Handles an external playback control input (usually from the UI) for a frame
 	 * @param ButtonID ID determinig the playback action
