[ChaosVD] Fix for an issue where if you undock a tab, the title bar in all CVD windows starts flickering.

The problem was that several tabs were created as major tabs when they should have been Panel tabs, therefore we were creating the main menu in them incorrectly

diff --git a/Engine/Plugins/ChaosVD/Source/ChaosVD/Private/ChaosVDCollisionDataDetailsTab.cpp b/Engine/Plugins/ChaosVD/Source/ChaosVD/Private/ChaosVDCollisionDataDetailsTab.cpp
--- a/Engine/Plugins/ChaosVD/Source/ChaosVD/Private/ChaosVDCollisionDataDetailsTab.cpp
+++ b/Engine/Plugins/ChaosVD/Source/ChaosVD/Private/ChaosVDCollisionDataDetailsTab.cpp
@@ -17,7 +17,7 @@ TSharedRef<SDockTab> FChaosVDCollisionDataDetailsTab::HandleTabSpawnRequest(cons
 {
 	TSharedRef<SDockTab> DetailsPanelTab =
 		SNew(SDockTab)
-		.TabRole(ETabRole::MajorTab)
+		.TabRole(ETabRole::PanelTab)
 		.Label(LOCTEXT("CollisionInspectorTab", "Collision Data Inspector"))
 		.ToolTipText(LOCTEXT("CollisionInspectorTabTip", "See the details of the any collision data for the selected object"));
 
diff --git a/Engine/Plugins/ChaosVD/Source/ChaosVD/Private/ChaosVDConstraintDataInspectorTab.cpp b/Engine/Plugins/ChaosVD/Source/ChaosVD/Private/ChaosVDConstraintDataInspectorTab.cpp
--- a/Engine/Plugins/ChaosVD/Source/ChaosVD/Private/ChaosVDConstraintDataInspectorTab.cpp
+++ b/Engine/Plugins/ChaosVD/Source/ChaosVD/Private/ChaosVDConstraintDataInspectorTab.cpp
@@ -15,7 +15,7 @@ TSharedRef<SDockTab> FChaosVDConstraintDataInspectorTab::HandleTabSpawnRequest(c
 {
 	TSharedRef<SDockTab> DetailsPanelTab =
 	SNew(SDockTab)
-	.TabRole(ETabRole::MajorTab)
+	.TabRole(ETabRole::PanelTab)
 	.Label(LOCTEXT("ConstraintDataInspectorTab", "Constraint Inspector"))
 	.ToolTipText(LOCTEXT("ConstraintDataInspectorTabToolTip", "See the details of the any non-collision constraint selected in the viewport"));
 
diff --git a/Engine/Plugins/ChaosVD/Source/ChaosVD/Private/ChaosVDObjectDetailsTab.cpp b/Engine/Plugins/ChaosVD/Source/ChaosVD/Private/ChaosVDObjectDetailsTab.cpp
--- a/Engine/Plugins/ChaosVD/Source/ChaosVD/Private/ChaosVDObjectDetailsTab.cpp
+++ b/Engine/Plugins/ChaosVD/Source/ChaosVD/Private/ChaosVDObjectDetailsTab.cpp
@@ -37,7 +37,7 @@ TSharedRef<SDockTab> FChaosVDStandAloneObjectDetailsTab::HandleTabSpawnRequest(c
 {
 	TSharedRef<SDockTab> DetailsPanelTab =
 	SNew(SDockTab)
-	.TabRole(ETabRole::MajorTab)
+	.TabRole(ETabRole::PanelTab)
 	.Label(LOCTEXT("DetailsPanel", "Details"))
 	.ToolTipText(LOCTEXT("DetailsPanelToolTip", "See the details of the selected object"));
 
diff --git a/Engine/Plugins/ChaosVD/Source/ChaosVD/Private/ChaosVDSceneQueryDataInspectorTab.cpp b/Engine/Plugins/ChaosVD/Source/ChaosVD/Private/ChaosVDSceneQueryDataInspectorTab.cpp
--- a/Engine/Plugins/ChaosVD/Source/ChaosVD/Private/ChaosVDSceneQueryDataInspectorTab.cpp
+++ b/Engine/Plugins/ChaosVD/Source/ChaosVD/Private/ChaosVDSceneQueryDataInspectorTab.cpp
@@ -17,7 +17,7 @@ TSharedRef<SDockTab> FChaosVDSceneQueryDataInspectorTab::HandleTabSpawnRequest(c
 {
 	TSharedRef<SDockTab> DetailsPanelTab =
 	SNew(SDockTab)
-	.TabRole(ETabRole::MajorTab)
+	.TabRole(ETabRole::PanelTab)
 	.Label(LOCTEXT("SceneQueryInspectorTab", "Scene Query Data Inspector"))
 	.ToolTipText(LOCTEXT("SceneQueryInspectorTabTip", "See the details of the any scene query selected in the viewport"));
 
diff --git a/Engine/Plugins/ChaosVD/Source/ChaosVD/Private/ChaosVDWorldOutlinerTab.cpp b/Engine/Plugins/ChaosVD/Source/ChaosVD/Private/ChaosVDWorldOutlinerTab.cpp
--- a/Engine/Plugins/ChaosVD/Source/ChaosVD/Private/ChaosVDWorldOutlinerTab.cpp
+++ b/Engine/Plugins/ChaosVD/Source/ChaosVD/Private/ChaosVDWorldOutlinerTab.cpp
@@ -143,9 +143,9 @@ TSharedRef<SDockTab> FChaosVDWorldOutlinerTab::HandleTabSpawnRequest(const FSpaw
 
 	TSharedRef<SDockTab> OutlinerTab =
 		SNew(SDockTab)
-		.TabRole(ETabRole::MajorTab)
+		.TabRole(ETabRole::PanelTab)
 		.Label(LOCTEXT("Physics World Outliner", "Physics World Outliner"))
-		.ToolTipText(LOCTEXT("PhysicsWorldOutlinerTabToolTip", "Hierachy view of the physics objects by category"));
+		.ToolTipText(LOCTEXT("PhysicsWorldOutlinerTabToolTip", "Hierarchy view of the physics objects by category"));
 	
 	OutlinerTab->SetContent
 	(
diff --git a/Engine/Plugins/ChaosVD/Source/ChaosVD/Private/TabSpawers/ChaosVDSceneQueryBrowserTab.cpp b/Engine/Plugins/ChaosVD/Source/ChaosVD/Private/TabSpawers/ChaosVDSceneQueryBrowserTab.cpp
--- a/Engine/Plugins/ChaosVD/Source/ChaosVD/Private/TabSpawers/ChaosVDSceneQueryBrowserTab.cpp
+++ b/Engine/Plugins/ChaosVD/Source/ChaosVD/Private/TabSpawers/ChaosVDSceneQueryBrowserTab.cpp
@@ -17,7 +17,7 @@ TSharedRef<SDockTab> FChaosVDSceneQueryBrowserTab::HandleTabSpawnRequest(const F
 {
 	TSharedRef<SDockTab> SceneQueryBrowserTab =
 	SNew(SDockTab)
-	.TabRole(ETabRole::MajorTab)
+	.TabRole(ETabRole::PanelTab)
 	.Label(LOCTEXT("SceneQueryBrowserTab", "Scene Query Browser"))
 	.ToolTipText(LOCTEXT("SceneQueryBrowserTabTabTip", "Shows all recorded scene queries for the current frame, and allows you to select them or play then in order"));
 
