[ChaosVD] Adding missing keyshortcuts from the normal level editor, that can be translated to similar actions in CVD, plus other ones CVD specific.

[CTRL + O] Opens the browse file modal
[CTRL + SHIFT + S] Opens the browse live sessions modal
[CTRL + Q] Opens the scene query browser
[CTRL + SHIFT + C] Combines all open sessions in a single file if possible.
[H] Hides the selected particle
[CTRL + H] Unhides any particle that was manually hidden
[ESCAPE] Clears any active selection (Particle or Solver Data).

Playback control shortcuts and other editor shortcuts were implemented in previous CLs.

diff --git a/Engine/Plugins/ChaosVD/Source/ChaosVD/Private/ChaosVDCommands.cpp b/Engine/Plugins/ChaosVD/Source/ChaosVD/Private/ChaosVDCommands.cpp
--- a/Engine/Plugins/ChaosVD/Source/ChaosVD/Private/ChaosVDCommands.cpp
+++ b/Engine/Plugins/ChaosVD/Source/ChaosVD/Private/ChaosVDCommands.cpp
@@ -12,19 +12,30 @@ FChaosVDCommands::FChaosVDCommands() : TCommands<FChaosVDCommands>(TEXT("ChaosVD
 
 void FChaosVDCommands::RegisterCommands()
 {
+	// Viewport Commands
 	UI_COMMAND(ToggleFollowSelectedObject, "Follow Selected Object", "Start or Stop following the selected object", EUserInterfaceActionType::ToggleButton, FInputChord(EKeys::F8));
 	UI_COMMAND(OverridePlaybackFrameRate, "Override Recorded Framerate", "When enabled, allows to playback the recording at a fixed framerate", EUserInterfaceActionType::ToggleButton, FInputChord(EModifierKey::Control, EKeys::R));
 	UI_COMMAND(AllowTranslucentSelection, "Allow Translucent Selection", "Allows translucent objects to be selected", EUserInterfaceActionType::ToggleButton, FInputChord(EKeys::T));
+	UI_COMMAND(HideSelected, "Hide Selected", "Hides the selected particle", EUserInterfaceActionType::Button, FInputChord(EKeys::H));
+	UI_COMMAND(ShowAll, "Show All Particles", "Un-hides any manually hidden particle", EUserInterfaceActionType::Button, FInputChord(EModifierKey::Control , EKeys::H));
+	UI_COMMAND(DeselectAll, "Deselect all objects", "Clears any active selection (Particle or Solver Data)", EUserInterfaceActionType::Button, FInputChord(EKeys::Escape));
+
+	// Main Toolbar Commands
+	UI_COMMAND(OpenFile, "Open File", "Opens the file browser modal", EUserInterfaceActionType::Button, FInputChord(EModifierKey::Control ,EKeys::O));
+	UI_COMMAND(BrowseLiveSessions, "Browse Live Sessions", "Open the connect to session modal", EUserInterfaceActionType::Button, FInputChord(EModifierKey::Control | EModifierKey::Shift, EKeys::S));
+	UI_COMMAND(CombineOpenFiles, "Combine Files", "Combines all open sessions into a single file (if possible)", EUserInterfaceActionType::Button, FInputChord(EModifierKey::Control | EModifierKey::Shift , EKeys::C));
+	UI_COMMAND(OpenSceneQueryBrowser, "Scene query Browser", "Opens the Scene Query browser window", EUserInterfaceActionType::Button, FInputChord(EModifierKey::Control , EKeys::Q));
+
 	
 	// Playback Controls commands
-    UI_COMMAND(PlayPauseTrack, "Play/Pause", "Plays or pauses the playback for the current active track", EUserInterfaceActionType::Button, FInputChord(FInputChord(EModifierKey::Control, EKeys::SpaceBar)));
+    UI_COMMAND(PlayPauseTrack, "Play/Pause", "Plays or pauses the playback for the current active track", EUserInterfaceActionType::Button, FInputChord(EModifierKey::Control, EKeys::SpaceBar));
     UI_COMMAND(StopTrack, "Stop", "Stops the playback for the current active track", EUserInterfaceActionType::Button, FInputChord(EModifierKey::Shift, EKeys::SpaceBar));
 
     UI_COMMAND(NextFrame, "Next Frame", "Plays the next frame of the current track", EUserInterfaceActionType::Button, FInputChord(EKeys::Period));
     UI_COMMAND(PrevFrame, "Prev Frame", "Plays the previous frame of the current track", EUserInterfaceActionType::Button, FInputChord(EKeys::Comma));
 
-    UI_COMMAND(NextStage, "Next Stage", "Plays the next solver stage of the current track", EUserInterfaceActionType::Button, FInputChord(FInputChord(EModifierKey::Shift, EKeys::Period)));
-    UI_COMMAND(PrevStage, "Prev Stage", "Plays the previous solver stage of the current track", EUserInterfaceActionType::Button, FInputChord(FInputChord(EModifierKey::Shift, EKeys::Comma)));
+    UI_COMMAND(NextStage, "Next Stage", "Plays the next solver stage of the current track", EUserInterfaceActionType::Button, FInputChord(EModifierKey::Shift, EKeys::Period));
+    UI_COMMAND(PrevStage, "Prev Stage", "Plays the previous solver stage of the current track", EUserInterfaceActionType::Button, FInputChord(EModifierKey::Shift, EKeys::Comma));
 }
 
 #undef LOCTEXT_NAMESPACE
diff --git a/Engine/Plugins/ChaosVD/Source/ChaosVD/Private/ChaosVDCommands.h b/Engine/Plugins/ChaosVD/Source/ChaosVD/Private/ChaosVDCommands.h
--- a/Engine/Plugins/ChaosVD/Source/ChaosVD/Private/ChaosVDCommands.h
+++ b/Engine/Plugins/ChaosVD/Source/ChaosVD/Private/ChaosVDCommands.h
@@ -14,10 +14,19 @@ public:
 	// TCommands<> interface
 	virtual void RegisterCommands() override;
 
-public:
+	// Viewport Commands
 	TSharedPtr<FUICommandInfo> ToggleFollowSelectedObject;
 	TSharedPtr<FUICommandInfo> OverridePlaybackFrameRate;
 	TSharedPtr<FUICommandInfo> AllowTranslucentSelection;
+	TSharedPtr<FUICommandInfo> DeselectAll;
+    TSharedPtr<FUICommandInfo> HideSelected;
+    TSharedPtr<FUICommandInfo> ShowAll;
+
+	// Main Toolbar Commands
+	TSharedPtr<FUICommandInfo> OpenFile;
+	TSharedPtr<FUICommandInfo> BrowseLiveSessions;
+	TSharedPtr<FUICommandInfo> CombineOpenFiles;
+	TSharedPtr<FUICommandInfo> OpenSceneQueryBrowser;
 
 	// Playback Controls commands
 	TSharedPtr<FUICommandInfo> PlayPauseTrack;
diff --git a/Engine/Plugins/ChaosVD/Source/ChaosVD/Private/ChaosVDSceneParticle.cpp b/Engine/Plugins/ChaosVD/Source/ChaosVD/Private/ChaosVDSceneParticle.cpp
--- a/Engine/Plugins/ChaosVD/Source/ChaosVD/Private/ChaosVDSceneParticle.cpp
+++ b/Engine/Plugins/ChaosVD/Source/ChaosVD/Private/ChaosVDSceneParticle.cpp
@@ -798,3 +798,17 @@ void FChaosVDSceneParticle::RemoveHiddenFlag(EChaosVDHideParticleFlags Flag)
 		EnumAddFlags(DirtyFlags, EChaosVDSceneParticleDirtyFlags::Visibility);
 	}
 }
+
+void FChaosVDSceneParticle::HideImmediate(EChaosVDHideParticleFlags Flag)
+{
+	AddHiddenFlag(Flag);
+	UpdateGeometryComponentsVisibility(EChaosVDParticleVisibilityUpdateFlags::DirtyScene);
+	Chaos::VD::TedsUtils::AddColumnToObject<FTypedElementSyncFromWorldTag>(this);
+}
+
+void FChaosVDSceneParticle::ShowImmediate()
+{
+	HideParticleFlags = EChaosVDHideParticleFlags::None;
+	UpdateGeometryComponentsVisibility(EChaosVDParticleVisibilityUpdateFlags::DirtyScene);
+	Chaos::VD::TedsUtils::AddColumnToObject<FTypedElementSyncFromWorldTag>(this);
+}
diff --git a/Engine/Plugins/ChaosVD/Source/ChaosVD/Private/Widgets/SChaosVDMainTab.cpp b/Engine/Plugins/ChaosVD/Source/ChaosVD/Private/Widgets/SChaosVDMainTab.cpp
--- a/Engine/Plugins/ChaosVD/Source/ChaosVD/Private/Widgets/SChaosVDMainTab.cpp
+++ b/Engine/Plugins/ChaosVD/Source/ChaosVD/Private/Widgets/SChaosVDMainTab.cpp
@@ -12,6 +12,7 @@
 #include "ChaosVDScene.h"
 #include "ChaosVDSolversTracksTab.h"
 #include "ChaosVDCollisionDataDetailsTab.h"
+#include "ChaosVDCommands.h"
 #include "ChaosVDConstraintDataInspectorTab.h"
 #include "ChaosVDIndependentDetailsPanelManager.h"
 #include "ChaosVDSceneParticleCustomization.h"
@@ -73,11 +74,15 @@
 
 void SChaosVDMainTab::Construct(const FArguments& InArgs, TSharedPtr<FChaosVDEngine> InChaosVDEngine)
 {
+	ChaosVDEngine = InChaosVDEngine;
+
 	EditorModeTools = MakeShared<FChaosVDEditorModeTools>(InChaosVDEngine->GetCurrentScene());
 
 	GlobalCommandList = MakeShared<FUICommandList>();
+
+	BindUICommands(GlobalCommandList.ToSharedRef());
+	
 	EditorModeTools->SetToolkitHost(StaticCastSharedRef<SChaosVDMainTab>(AsShared()));
-	ChaosVDEngine = InChaosVDEngine;
 	OwnerTab = InArgs._OwnerTab;
 
 	RegisterComponentVisualizer(UChaosVDSolverCollisionDataComponent::StaticClass()->GetFName(), MakeShared<FChaosVDSolverCollisionDataComponentVisualizer>());
@@ -215,6 +220,35 @@ SChaosVDMainTab::~SChaosVDMainTab()
 	CleanUpDisableCPUThrottlingDelegate();
 }
 
+void SChaosVDMainTab::BindUICommands(const TSharedRef<FUICommandList>& InGlobalUICommandsRef)
+{
+	const FChaosVDCommands& Commands = FChaosVDCommands::Get();
+
+	FUIAction OpenFileAction;
+	OpenFileAction.ExecuteAction.BindSPLambda(SharedThis(this), [this]
+	{
+		BrowseAndOpenChaosVDRecording();
+	});
+	InGlobalUICommandsRef->MapAction(Commands.OpenFile, OpenFileAction);
+
+	FUIAction CombineOpenFilesAction;
+	CombineOpenFilesAction.ExecuteAction.BindSP(SharedThis(this), &SChaosVDMainTab::CombineOpenSessions);
+	CombineOpenFilesAction.CanExecuteAction.BindSP(GetChaosVDEngineInstance(), &FChaosVDEngine::CanCombineOpenSessions);
+	InGlobalUICommandsRef->MapAction(Commands.CombineOpenFiles, CombineOpenFilesAction);
+	
+	FUIAction BrowseLiveSessionsAction;
+	BrowseLiveSessionsAction.ExecuteAction.BindSP(SharedThis(this), &SChaosVDMainTab::BrowseLiveSessionsFromTraceStore);
+	InGlobalUICommandsRef->MapAction(Commands.BrowseLiveSessions, BrowseLiveSessionsAction);
+	
+	FUIAction OpenSceneQueryBrowserAction;
+	OpenSceneQueryBrowserAction.ExecuteAction.BindSPLambda(SharedThis(this), [this]()
+	{
+		TabManager->TryInvokeTab(FChaosVDTabID::SceneQueryBrowser);
+	});
+
+	InGlobalUICommandsRef->MapAction(Commands.OpenSceneQueryBrowser, OpenSceneQueryBrowserAction);
+}
+
 void SChaosVDMainTab::BringToFront()
 {
 	if (TabManager.IsValid())
diff --git a/Engine/Plugins/ChaosVD/Source/ChaosVD/Private/Widgets/SChaosVDPlaybackViewport.cpp b/Engine/Plugins/ChaosVD/Source/ChaosVD/Private/Widgets/SChaosVDPlaybackViewport.cpp
--- a/Engine/Plugins/ChaosVD/Source/ChaosVD/Private/Widgets/SChaosVDPlaybackViewport.cpp
+++ b/Engine/Plugins/ChaosVD/Source/ChaosVD/Private/Widgets/SChaosVDPlaybackViewport.cpp
@@ -8,11 +8,17 @@
 #include "ChaosVDPlaybackController.h"
 #include "ChaosVDPlaybackViewportClient.h"
 #include "ChaosVDScene.h"
+#include "ChaosVDSceneParticle.h"
 #include "EditorModeManager.h"
 #include "EditorViewportCommands.h"
 #include "SChaosVDGameFramesPlaybackControls.h"
+#include "Actors/ChaosVDSolverInfoActor.h"
+#include "ChaosVisualDebugger/ChaosVDDataWrapperUtils.h"
+#include "Elements/Actor/ActorElementData.h"
 #include "Widgets/SChaosVDMainTab.h"
 #include "Framework/Application/SlateApplication.h"
+#include "TEDS/ChaosVDSelectionInterface.h"
+#include "TEDS/ChaosVDStructTypedElementData.h"
 #include "Widgets/SChaosVDTimelineWidget.h"
 #include "Widgets/SChaosVDViewportToolbar.h"
 #include "Widgets/Input/SCheckBox.h"
@@ -237,6 +243,18 @@ void SChaosVDPlaybackViewport::BindGlobalUICommands()
 	FUIAction PrevStagePlaybackAction;
 	PrevStagePlaybackAction.ExecuteAction.BindSP(SharedThis(this), &SChaosVDPlaybackViewport::HandleFrameStagePlaybackControlInput, EChaosVDPlaybackButtonsID::Next);
 	GlobalUICommandsList->MapAction(Commands.NextStage, PrevStagePlaybackAction);
+
+	FUIAction DeselectAllAction;
+	DeselectAllAction.ExecuteAction.BindSP(SharedThis(this), &SChaosVDPlaybackViewport::DeselectAll);
+	GlobalUICommandsList->MapAction(Commands.DeselectAll, DeselectAllAction);
+
+	FUIAction HideSelectedAction;
+	HideSelectedAction.ExecuteAction.BindSP(SharedThis(this), &SChaosVDPlaybackViewport::HideSelected);
+	GlobalUICommandsList->MapAction(Commands.HideSelected, HideSelectedAction);
+
+	FUIAction ShowAllAction;
+	ShowAllAction.ExecuteAction.BindSP(SharedThis(this), &SChaosVDPlaybackViewport::ShowAll);
+	GlobalUICommandsList->MapAction(Commands.ShowAll, ShowAllAction);
 }
 
 void SChaosVDPlaybackViewport::UnBindEditorViewportUnsupportedCommands()
@@ -502,6 +520,75 @@ void SChaosVDPlaybackViewport::HandleFrameStagePlaybackControlInput(EChaosVDPlay
 	}
 }
 
+void SChaosVDPlaybackViewport::DeselectAll() const
+{
+	const TSharedPtr<FChaosVDPlaybackController> PlaybackControllerPtr = PlaybackController.Pin();
+	const TSharedPtr<FChaosVDScene> ScenePtr = PlaybackControllerPtr ? PlaybackControllerPtr->GetControllerScene().Pin() : nullptr;
+	if (!ScenePtr)
+	{
+		return;
+	}
+
+	if (UTypedElementSelectionSet* SelectionSet = ScenePtr->GetElementSelectionSet())
+	{
+		SelectionSet->ClearSelection(FTypedElementSelectionOptions());
+	}
+	
+	if (const TSharedPtr<FChaosVDSolverDataSelection> SolverDataSelection = ScenePtr->GetSolverDataSelectionObject().Pin())
+	{
+		SolverDataSelection->SelectData(nullptr);
+	}
+}
+
+void SChaosVDPlaybackViewport::HideSelected() const
+{
+	if (UTypedElementSelectionSet* SelectionSet = EditorModeTools->GetEditorSelectionSet())
+	{
+		//TODO: Update this if we add multi selection support
+		constexpr int32 MaxElements = 1;
+		TArray<FTypedElementHandle, TInlineAllocator<MaxElements>> TypedElementHandles;
+		SelectionSet->GetSelectedElementHandles(TypedElementHandles, UChaosVDSelectionInterface::StaticClass());
+
+		if (TypedElementHandles.Num() > 0)
+		{
+			FTypedElementHandle& SelectionHandle = TypedElementHandles[0];
+			using namespace Chaos::VD::TypedElementDataUtil;
+			if (FChaosVDSceneParticle* Particle = GetStructDataFromTypedElementHandle<FChaosVDSceneParticle>(SelectionHandle))
+			{
+				Particle->HideImmediate(EChaosVDHideParticleFlags::HiddenBySceneOutliner);
+			}
+			else if (AChaosVDSolverInfoActor* SolverInfoActor = Cast<AChaosVDSolverInfoActor>(ActorElementDataUtil::GetActorFromHandle(SelectionHandle)))
+			{
+				SolverInfoActor->SetIsTemporarilyHiddenInEditor(true);
+			}
+		}
+	}
+}
+
+void SChaosVDPlaybackViewport::ShowAll() const
+{
+	const TSharedPtr<FChaosVDPlaybackController> PlaybackControllerPtr = PlaybackController.Pin();
+	const TSharedPtr<FChaosVDScene> ScenePtr = PlaybackControllerPtr ? PlaybackControllerPtr->GetControllerScene().Pin() : nullptr;
+	if (!ScenePtr)
+	{
+		return;
+	}
+
+	const FChaosVDSolverInfoByIDMap& SolverInfoActorsByID = ScenePtr->GetSolverInfoActorsMap();
+	for (const TPair<int32, AChaosVDSolverInfoActor*>& SolverInfoWithID : SolverInfoActorsByID)
+	{
+		if (SolverInfoWithID.Value)
+		{
+			// Note : Depending on the scene, this will be slow. If that is the case, we should expose the Batch Visibility update for particles
+			// in a way other systems can use it.
+			// Currently, we process the data changes and apply them to the mesh components at the end of the frame, and de-duplicate any operations
+			// that nullify themselves and the final visibility state doesn't actually change, therefore the perf hit of doing this is mostly mitigated
+			SolverInfoWithID.Value->SetIsTemporarilyHiddenInEditor(true);
+			SolverInfoWithID.Value->SetIsTemporarilyHiddenInEditor(false);
+		}
+	}
+}
+
 void SChaosVDPlaybackViewport::OnPlaybackSceneUpdated()
 {
 	PlaybackViewportClient->HandleCVDSceneUpdated();
diff --git a/Engine/Plugins/ChaosVD/Source/ChaosVD/Private/Widgets/SChaosVDPlaybackViewport.h b/Engine/Plugins/ChaosVD/Source/ChaosVD/Private/Widgets/SChaosVDPlaybackViewport.h
--- a/Engine/Plugins/ChaosVD/Source/ChaosVD/Private/Widgets/SChaosVDPlaybackViewport.h
+++ b/Engine/Plugins/ChaosVD/Source/ChaosVD/Private/Widgets/SChaosVDPlaybackViewport.h
@@ -83,6 +83,10 @@ protected:
 	void HandleFramePlaybackControlInput(EChaosVDPlaybackButtonsID ButtonID);
 	void HandleFrameStagePlaybackControlInput(EChaosVDPlaybackButtonsID ButtonID);
 
+	void DeselectAll() const;
+	void HideSelected() const;
+	void ShowAll() const;
+
 	void OnPlaybackSceneUpdated();
 	void OnSolverVisibilityUpdated(int32 SolverID, bool bNewVisibility);
 	void BindToSceneUpdateEvents();
diff --git a/Engine/Plugins/ChaosVD/Source/ChaosVD/Public/ChaosVDSceneParticle.h b/Engine/Plugins/ChaosVD/Source/ChaosVD/Public/ChaosVDSceneParticle.h
--- a/Engine/Plugins/ChaosVD/Source/ChaosVD/Public/ChaosVDSceneParticle.h
+++ b/Engine/Plugins/ChaosVD/Source/ChaosVD/Public/ChaosVDSceneParticle.h
@@ -58,6 +58,19 @@ public:
 
 	void RemoveHiddenFlag(EChaosVDHideParticleFlags Flag);
 
+	/** Performs all the required steps to hide a particle and update the viewport / scene outliner.
+	 * @param Flag Context flags about what is hiding this particle
+	 * @note Don't use when performing operation on multiple particles or other performance sensitive operations.
+	 * See UChaosVDParticleDataComponent for examples of batch visibility update operations
+	 */
+	void HideImmediate(EChaosVDHideParticleFlags Flag);
+
+	/** Performs all the required steps to show a particle and update the viewport / scene outliner.
+	 * @note Don't use when performing operation on multiple particles or other performance sensitive operations.
+	 * See UChaosVDParticleDataComponent for examples of batch visibility update operations
+	 */
+	void ShowImmediate();
+
 	EChaosVDHideParticleFlags GetHideFlags() const
 	{
 		return HideParticleFlags;
diff --git a/Engine/Plugins/ChaosVD/Source/ChaosVD/Public/Widgets/SChaosVDMainTab.h b/Engine/Plugins/ChaosVD/Source/ChaosVD/Public/Widgets/SChaosVDMainTab.h
--- a/Engine/Plugins/ChaosVD/Source/ChaosVD/Public/Widgets/SChaosVDMainTab.h
+++ b/Engine/Plugins/ChaosVD/Source/ChaosVD/Public/Widgets/SChaosVDMainTab.h
@@ -58,6 +58,8 @@ public:
 	void Construct(const FArguments& InArgs, TSharedPtr<FChaosVDEngine> InChaosVDEngine);
 
 	virtual ~SChaosVDMainTab() override;
+	
+	void BindUICommands(const TSharedRef<FUICommandList>& InGlobalUICommandsRef);
 
 	TSharedRef<FChaosVDEngine> GetChaosVDEngineInstance() const
 	{
