[ChaosVD] Fix CIS issue

The new method CanVisualizationFlagBeChangedByUI was implemented as a protected memeber in the base class but it should be public.

diff --git a/Engine/Plugins/ChaosVD/Source/ChaosVD/Public/Settings/ChaosVDCoreSettings.h b/Engine/Plugins/ChaosVD/Source/ChaosVD/Public/Settings/ChaosVDCoreSettings.h
--- a/Engine/Plugins/ChaosVD/Source/ChaosVD/Public/Settings/ChaosVDCoreSettings.h
+++ b/Engine/Plugins/ChaosVD/Source/ChaosVD/Public/Settings/ChaosVDCoreSettings.h
@@ -72,6 +72,7 @@ class CHAOSVD_API UChaosVDVisualizationSettingsObjectBase : public UChaosVDSetti
 protected:
 	virtual void BroadcastSettingsChanged() override;
 
+public:
 	/** Returns true if a visualization flag option of this setting option should be enabled in the UI */
 	virtual bool CanVisualizationFlagBeChangedByUI(uint32 Flag);
 };
