[ChaosVD] Improvements to the Scene Query inspector

- Now hits and the structures showing the visited shapes show the particle name instead of the particle index
- Now the SQ Visit timeline control, auto selects the particle and collision shape automatically
- Made the Select Visited Shape smaller and part of the timeline widget.
- Update the hit flags property in hit data to show the human readable name instead of just an integer.
- Fixed an issue where the show collision button was being shown in the Geometry shape details panel.
- Added visibility to the Reject Reason in the recorded SQ Visit data. We were serializing that already, but it was not marked as visible and therefore it was not available in the details panel

diff --git a/Engine/Plugins/ChaosVD/Source/ChaosVD/Private/ChaosVDSceneParticleCustomization.cpp b/Engine/Plugins/ChaosVD/Source/ChaosVD/Private/ChaosVDSceneParticleCustomization.cpp
--- a/Engine/Plugins/ChaosVD/Source/ChaosVD/Private/ChaosVDSceneParticleCustomization.cpp
+++ b/Engine/Plugins/ChaosVD/Source/ChaosVD/Private/ChaosVDSceneParticleCustomization.cpp
@@ -44,6 +44,28 @@ TSharedRef<IDetailCustomization> FChaosVDSceneParticleCustomization::MakeInstanc
 	return MakeShared<FChaosVDSceneParticleCustomization>(InMainTab);
 }
 
+
+void FChaosVDSceneParticleCustomization::AddParticleDataButtons(IDetailLayoutBuilder& DetailBuilder)
+{
+	IDetailCategoryBuilder& ParticleDataCategoryBuilder = DetailBuilder.EditCategory(ParticleDataCategoryName);
+
+	ParticleDataCategoryBuilder.AddCustomRow(FText::GetEmpty()).
+	                            WholeRowContent()
+	[
+		GenerateOpenInNewDetailsPanelButton().ToSharedRef()
+	];
+
+	const FText CollisionDataRowLabel = LOCTEXT("ParticleCollisionData", "Particle Collision Data");
+	IDetailCategoryBuilder& CategoryBuilder = DetailBuilder.EditCategory("CollisionData", CollisionDataRowLabel);
+
+	CategoryBuilder.AddCustomRow(CollisionDataRowLabel).
+	                WholeRowContent()
+	[
+		GenerateShowCollisionDataButton().ToSharedRef()
+	];
+}
+
+
 void FChaosVDSceneParticleCustomization::CustomizeDetails(IDetailLayoutBuilder& DetailBuilder)
 {
 	FChaosVDDetailsCustomizationUtils::HideAllCategories(DetailBuilder, AllowedCategories);
@@ -100,10 +122,19 @@ void FChaosVDSceneParticleCustomization::CustomizeDetails(IDetailLayoutBuilder&
 			if (TSharedPtr<FChaosVDInstancedMeshData> SelectedGeometryInstance = SelectedParticleInstance->GetSelectedMeshInstance().Pin())
 			{
 				InspectedDataPropertyHandlePtr = AddExternalStructure(CachedGeometryDataInstanceCopy, DetailBuilder, GeometryCategoryName, LOCTEXT("GeometryShapeDataStructName", "Geometry Shape Data"));
+
+				IDetailCategoryBuilder& ParticleDataCategoryBuilder = DetailBuilder.EditCategory(GeometryCategoryName);
+
+				ParticleDataCategoryBuilder.AddCustomRow(FText::GetEmpty()).
+											WholeRowContent()
+				[
+					GenerateOpenInNewDetailsPanelButton().ToSharedRef()
+				];
 			}
 			else
 			{
 				InspectedDataPropertyHandlePtr = AddExternalStructure(CachedParticleData, DetailBuilder, ParticleDataCategoryName, LOCTEXT("ParticleDataStructName", "Particle Data"));
+				AddParticleDataButtons(DetailBuilder);
 			}
 
 			if (InspectedDataPropertyHandlePtr)
@@ -111,23 +142,6 @@ void FChaosVDSceneParticleCustomization::CustomizeDetails(IDetailLayoutBuilder&
 				TSharedRef<IPropertyHandle> InspectedDataPropertyHandleRef = InspectedDataPropertyHandlePtr.ToSharedRef();
 				FChaosVDDetailsCustomizationUtils::HideInvalidCVDDataWrapperProperties({&InspectedDataPropertyHandleRef, 1}, DetailBuilder);
 			}
-
-			IDetailCategoryBuilder& ParticleDataCategoryBuilder = DetailBuilder.EditCategory(ParticleDataCategoryName);
-
-			ParticleDataCategoryBuilder.AddCustomRow(FText::GetEmpty()).
-			WholeRowContent()
-			[
-				GenerateOpenInNewDetailsPanelButton().ToSharedRef()
-			];
-
-			const FText CollisionDataRowLabel = LOCTEXT("ParticleCollisionData", "Particle Collision Data");
-			IDetailCategoryBuilder& CategoryBuilder = DetailBuilder.EditCategory("CollisionData", CollisionDataRowLabel);
-			
-			CategoryBuilder.AddCustomRow(CollisionDataRowLabel).
-			WholeRowContent()
-			[
-				GenerateShowCollisionDataButton().ToSharedRef()
-			];
 		}
 	}
 	else
diff --git a/Engine/Plugins/ChaosVD/Source/ChaosVD/Private/ChaosVDSceneParticleCustomization.h b/Engine/Plugins/ChaosVD/Source/ChaosVD/Private/ChaosVDSceneParticleCustomization.h
--- a/Engine/Plugins/ChaosVD/Source/ChaosVD/Private/ChaosVDSceneParticleCustomization.h
+++ b/Engine/Plugins/ChaosVD/Source/ChaosVD/Private/ChaosVDSceneParticleCustomization.h
@@ -30,6 +30,8 @@ public:
 
 private:
 
+	void AddParticleDataButtons(IDetailLayoutBuilder& DetailBuilder);
+
 	template<typename TStruct>
 	TSharedPtr<IPropertyHandle> AddExternalStructure(TStruct& CachedStruct, IDetailLayoutBuilder& DetailBuilder, FName CategoryName, const FText& InPropertyName);
 
diff --git a/Engine/Plugins/ChaosVD/Source/ChaosVD/Private/DetailsCustomizations/ChaosVDQueryDataWrappersCustomizationDetails.cpp b/Engine/Plugins/ChaosVD/Source/ChaosVD/Private/DetailsCustomizations/ChaosVDQueryDataWrappersCustomizationDetails.cpp
--- a/Engine/Plugins/ChaosVD/Source/ChaosVD/Private/DetailsCustomizations/ChaosVDQueryDataWrappersCustomizationDetails.cpp
+++ b/Engine/Plugins/ChaosVD/Source/ChaosVD/Private/DetailsCustomizations/ChaosVDQueryDataWrappersCustomizationDetails.cpp
@@ -6,6 +6,7 @@
 #include "ChaosVDModule.h"
 #include "ChaosVDRecording.h"
 #include "ChaosVDScene.h"
+#include "ChaosVDSceneParticle.h"
 #include "CollisionQueryParams.h"
 #include "DataWrappers/ChaosVDQueryDataWrappers.h"
 #include "DetailsCustomizations/ChaosVDDetailsCustomizationUtils.h"
@@ -14,6 +15,7 @@
 #include "IDetailGroup.h"
 #include "PropertyCustomizationHelpers.h"
 #include "PropertyHandle.h"
+#include "Widgets/SChaosVDEnumFlagsMenu.h"
 #include "Widgets/SChaosVDMainTab.h"
 #include "Widgets/SChaosVDWarningMessageBox.h"
 #include "Widgets/Input/SCheckBox.h"
@@ -22,6 +24,45 @@
 
 #define LOCTEXT_NAMESPACE "ChaosVisualDebugger"
 
+namespace Chaos::VD::Utils::Private
+{
+	TSharedPtr<FChaosVDSceneParticle> GetParticleSceneFromCVDToolkitHost(const TWeakPtr<SChaosVDMainTab>& InCVDToolkitHost, int32 SolverID, int32 ParticleID)
+	{
+		TSharedPtr<SChaosVDMainTab> CVDToolkitHost = InCVDToolkitHost.Pin();
+		TSharedPtr<FChaosVDScene> CVDScene = CVDToolkitHost ? CVDToolkitHost->GetScene() : nullptr;
+		return CVDScene ? CVDScene->GetParticleInstance(SolverID, ParticleID) : nullptr;
+	}
+
+	FDetailWidgetRow& PopulateWidgetRowWithParticleName(FDetailWidgetRow& InOutWidgetRow, const FText& InParticleName)
+	{
+		InOutWidgetRow.NameContent()
+		[
+			SNew(SHorizontalBox)
+			+ SHorizontalBox::Slot()
+			.HAlign(HAlign_Right)
+			[ 
+				SNew(STextBlock)
+				.Text(LOCTEXT("OwningParticleName", "Owning Particle"))
+				.Font(IDetailLayoutBuilder::GetDetailFont())
+			]
+		]
+		.ValueContent()
+		[
+			SNew(SHorizontalBox)
+			+ SHorizontalBox::Slot()
+			.HAlign(HAlign_Right)
+			[ 
+				SNew(STextBlock)
+				.Text(InParticleName)
+				.ToolTipText(InParticleName)
+				.Font(IDetailLayoutBuilder::GetDetailFont())
+			]
+		];
+
+		return InOutWidgetRow;
+	}
+}
+
 TSharedRef<IPropertyTypeCustomization> FChaosVDQueryDataWrappersCustomizationDetails::MakeInstance()
 {
 	return MakeShareable(new FChaosVDQueryDataWrappersCustomizationDetails());
@@ -49,13 +90,62 @@ void FChaosVDQueryDataWrappersCustomizationDetails::CustomizeChildren(TSharedRef
 	FChaosVDDetailsCustomizationUtils::HideInvalidCVDDataWrapperProperties(Handles);
 }
 
-TSharedRef<IDetailCustomization> FChaosVDQueryVisitDataCustomization::MakeInstance()
+FChaosVDQueryVisitDataCustomization::FChaosVDQueryVisitDataCustomization(const TWeakPtr<SChaosVDMainTab>& InMainTab) : MainTabWeakPtr(InMainTab)
 {
-	return MakeShareable( new FChaosVDQueryVisitDataCustomization );
+}
+
+TSharedRef<IDetailCustomization> FChaosVDQueryVisitDataCustomization::MakeInstance(TWeakPtr<SChaosVDMainTab> MainTab)
+{
+	return MakeShareable( new FChaosVDQueryVisitDataCustomization(MainTab) );
+}
+
+
+void FChaosVDQueryVisitDataCustomization::ReplaceParticleIndexWithParticleName(IDetailLayoutBuilder& DetailBuilder)
+{
+	TArray<TSharedPtr<FStructOnScope>> StructsCustomized;	
+	DetailBuilder.GetStructsBeingCustomized(StructsCustomized);
+
+	if (StructsCustomized.IsEmpty())
+	{
+		return;
+	}
+
+	if (!ensure(StructsCustomized.Num() == 1))
+	{
+		UE_LOG(LogChaosVDEditor, Warning, TEXT("[%hs] Multiple struct selection is not supported yet. The first one will be customized"), __func__);
+		return;
+	}
+		
+	const TSharedPtr<FStructOnScope>& StructScope = StructsCustomized[0];
+
+	if (!StructScope || StructScope->GetStruct() != FChaosVDQueryVisitStep::StaticStruct())
+	{
+		UE_LOG(LogChaosVDEditor, Warning, TEXT("[%hs] attempted to customize an invalid struct "), __func__);
+	}
+		
+	const FChaosVDQueryVisitStep* SQVisitData = reinterpret_cast<const FChaosVDQueryVisitStep*>(StructScope->GetStructMemory());
+
+	using namespace Chaos::VD::Utils::Private;
+	TSharedPtr<FChaosVDSceneParticle> SceneParticle = GetParticleSceneFromCVDToolkitHost(MainTabWeakPtr, SQVisitData->SolverID_Editor, SQVisitData->ParticleIndex);
+	if (!SceneParticle)
+	{
+		return;
+	}
+
+	TSharedRef<IPropertyHandle> ParticleIndexPropertyHandle = DetailBuilder.GetProperty(GET_MEMBER_NAME_CHECKED(FChaosVDQueryVisitStep, ParticleIndex));
+		
+	if (IDetailPropertyRow* ParticleIndexPropertyRow = DetailBuilder.EditDefaultProperty(ParticleIndexPropertyHandle))
+	{
+		const FText ParticleNameAsText = FText::FromString(SceneParticle->GetDisplayName());
+		PopulateWidgetRowWithParticleName(ParticleIndexPropertyRow->CustomWidget(), ParticleNameAsText);
+	}
 }
 
 void FChaosVDQueryVisitDataCustomization::CustomizeDetails(IDetailLayoutBuilder& DetailBuilder)
 {
+
+	ReplaceParticleIndexWithParticleName(DetailBuilder);
+
 	TArray<TSharedRef<IPropertyHandle>> PotentialPropertiesToHide;
 	
 	constexpr int32 PropertiesToEvaluateNum = 2;
@@ -67,9 +157,71 @@ void FChaosVDQueryVisitDataCustomization::CustomizeDetails(IDetailLayoutBuilder&
 	FChaosVDDetailsCustomizationUtils::HideInvalidCVDDataWrapperProperties(PotentialPropertiesToHide, DetailBuilder);
 }
 
+TSharedRef<IPropertyTypeCustomization> FChaosVDQueryVisitDataPropertyCustomization::MakeInstance(TWeakPtr<SChaosVDMainTab> MainTab)
+{
+	return MakeShareable( new FChaosVDQueryVisitDataPropertyCustomization(MainTab) );
+}
+
+void FChaosVDQueryVisitDataPropertyCustomization::CustomizeHeader(TSharedRef<IPropertyHandle> StructPropertyHandle, FDetailWidgetRow& HeaderRow,
+	IPropertyTypeCustomizationUtils& StructCustomizationUtils)
+{
+	HeaderRow.NameContent()
+	[
+		StructPropertyHandle->CreatePropertyNameWidget()
+	];
+
+	HeaderRow.ValueContent()
+	[
+		StructPropertyHandle->CreatePropertyValueWidget()
+	];
+}
+
+void FChaosVDQueryVisitDataPropertyCustomization::CustomizeChildren(TSharedRef<IPropertyHandle> StructPropertyHandle,
+                                                                    IDetailChildrenBuilder& StructBuilder, IPropertyTypeCustomizationUtils& StructCustomizationUtils)
+{
+	FChaosVDDetailsPropertyDataHandle<FChaosVDQueryVisitStep> StructPropertyDataHandle(StructPropertyHandle);
+
+	FChaosVDQueryVisitStep* SQVisitDataBeingCustomized = StructPropertyDataHandle.GetDataInstance();
+	if (!SQVisitDataBeingCustomized)
+	{
+		return;
+	}
+
+	using namespace Chaos::VD::Utils::Private;
+	TSharedPtr<FChaosVDSceneParticle> SceneParticle = GetParticleSceneFromCVDToolkitHost(MainTabWeakPtr, SQVisitDataBeingCustomized->SolverID_Editor, SQVisitDataBeingCustomized->ParticleIndex);
+	if (!SceneParticle)
+	{
+		return;
+	}
+
+	uint32 NumChild = 0;
+	StructPropertyHandle->GetNumChildren(NumChild);
+
+	for (uint32 PropIndex = 0; PropIndex < NumChild; PropIndex++)
+	{
+		TSharedPtr<IPropertyHandle> InnerPropHandle = StructPropertyHandle->GetChildHandle(PropIndex);
+
+		if (!InnerPropHandle)
+		{
+			continue;
+		}
+
+		if (InnerPropHandle->GetProperty()->GetFName() == GET_MEMBER_NAME_CHECKED(FChaosVDQueryVisitStep, ParticleIndex))
+		{
+			const FText ParticleNameAsText = FText::FromString(SceneParticle->GetDisplayName());
+			const FText SearchText = FText::AsCultureInvariant(TEXT("Owning Particle"));
+			PopulateWidgetRowWithParticleName(StructBuilder.AddCustomRow(SearchText), ParticleNameAsText);
+		}
+		else
+		{
+			StructBuilder.AddProperty(InnerPropHandle.ToSharedRef());
+		}
+	}
+}
+
 TSharedRef<IDetailCustomization> FChaosVDQueryDataWrapperCustomization::MakeInstance()
 {
-	return MakeShareable( new FChaosVDQueryVisitDataCustomization );
+	return MakeShareable( new FChaosVDQueryDataWrapperCustomization );
 }
 
 void FChaosVDQueryDataWrapperCustomization::CustomizeDetails(IDetailLayoutBuilder& DetailBuilder)
diff --git a/Engine/Plugins/ChaosVD/Source/ChaosVD/Private/DetailsCustomizations/ChaosVDQueryDataWrappersCustomizationDetails.h b/Engine/Plugins/ChaosVD/Source/ChaosVD/Private/DetailsCustomizations/ChaosVDQueryDataWrappersCustomizationDetails.h
--- a/Engine/Plugins/ChaosVD/Source/ChaosVD/Private/DetailsCustomizations/ChaosVDQueryDataWrappersCustomizationDetails.h
+++ b/Engine/Plugins/ChaosVD/Source/ChaosVD/Private/DetailsCustomizations/ChaosVDQueryDataWrappersCustomizationDetails.h
@@ -24,12 +24,36 @@ class FChaosVDQueryDataWrappersCustomizationDetails : public IPropertyTypeCustom
 class FChaosVDQueryVisitDataCustomization : public IDetailCustomization
 {
 public:
-	FChaosVDQueryVisitDataCustomization(){};
-	virtual ~FChaosVDQueryVisitDataCustomization() override {};
+	FChaosVDQueryVisitDataCustomization(const TWeakPtr<SChaosVDMainTab>& InMainTab);
 
-	static TSharedRef<IDetailCustomization> MakeInstance();
+	virtual ~FChaosVDQueryVisitDataCustomization() override {}
+
+	static TSharedRef<IDetailCustomization> MakeInstance(TWeakPtr<SChaosVDMainTab> MainTab);
+	void ReplaceParticleIndexWithParticleName(IDetailLayoutBuilder& DetailBuilder);
 
 	virtual void CustomizeDetails(IDetailLayoutBuilder& DetailBuilder) override;
+
+protected:
+	TWeakPtr<SChaosVDMainTab> MainTabWeakPtr = nullptr;
+};
+
+/** Custom property details for the ChaosVD SQ Visit Data struct */
+class FChaosVDQueryVisitDataPropertyCustomization : public IPropertyTypeCustomization
+{
+public:
+	FChaosVDQueryVisitDataPropertyCustomization(const TWeakPtr<SChaosVDMainTab>& InMainTab) : MainTabWeakPtr(InMainTab)
+	{
+	}
+
+	virtual ~FChaosVDQueryVisitDataPropertyCustomization() override {}
+
+	static TSharedRef<IPropertyTypeCustomization> MakeInstance(TWeakPtr<SChaosVDMainTab> MainTab);
+
+	virtual void CustomizeHeader(TSharedRef<IPropertyHandle> StructPropertyHandle, FDetailWidgetRow& HeaderRow, IPropertyTypeCustomizationUtils& StructCustomizationUtils) override;
+	virtual void CustomizeChildren(TSharedRef<IPropertyHandle> StructPropertyHandle, IDetailChildrenBuilder& StructBuilder, IPropertyTypeCustomizationUtils& StructCustomizationUtils) override;
+
+protected:
+	TWeakPtr<SChaosVDMainTab> MainTabWeakPtr = nullptr;
 };
 
 /** Custom details panel for the ChaosVD SQ Data Wrapper struct */
diff --git a/Engine/Plugins/ChaosVD/Source/ChaosVD/Private/Trace/DataProcessors/ChaosVDSceneQueryVisitDataProcessor.cpp b/Engine/Plugins/ChaosVD/Source/ChaosVD/Private/Trace/DataProcessors/ChaosVDSceneQueryVisitDataProcessor.cpp
--- a/Engine/Plugins/ChaosVD/Source/ChaosVD/Private/Trace/DataProcessors/ChaosVDSceneQueryVisitDataProcessor.cpp
+++ b/Engine/Plugins/ChaosVD/Source/ChaosVD/Private/Trace/DataProcessors/ChaosVDSceneQueryVisitDataProcessor.cpp
@@ -36,13 +36,15 @@ bool FChaosVDSceneQueryVisitDataProcessor::ProcessRawData(const TArray<uint8>& I
 					TSharedPtr<FChaosVDQueryDataWrapper> QueryDataPtr = *QueryDataPtrPtr;
 					if (QueryDataPtrPtr->IsValid())
 					{
-						//TODO: There is an existing issue where CVD is recording garbage data for the Hit Face Normalif the query is a line trace.
+						//TODO: There is an existing issue where CVD is recording garbage data for the Hit Face Normal if the query is a line trace.
 						// That value is not used in line traces, so for now just clear it out until we can implement a better solution
 						// where we include any post-processing done on the hit data during the HitConversion step.
 						if (QueryDataPtr->Type == EChaosVDSceneQueryType::RayCast)
 						{
 							VisitStepData.HitData.FaceNormal = FVector::ZeroVector;
 						}
+						
+						VisitStepData.SolverID_Editor = QueryDataPtr->WorldSolverID;
 
 						if (VisitStepData.HitData.HasValidData())
 						{
diff --git a/Engine/Plugins/ChaosVD/Source/ChaosVD/Private/Widgets/SChaosVDMainTab.cpp b/Engine/Plugins/ChaosVD/Source/ChaosVD/Private/Widgets/SChaosVDMainTab.cpp
--- a/Engine/Plugins/ChaosVD/Source/ChaosVD/Private/Widgets/SChaosVDMainTab.cpp
+++ b/Engine/Plugins/ChaosVD/Source/ChaosVD/Private/Widgets/SChaosVDMainTab.cpp
@@ -744,18 +744,18 @@ void SChaosVDMainTab::SetCustomPropertyLayouts(IDetailsView* DetailsView)
 	DetailsView->RegisterInstancedCustomPropertyLayout(FChaosVDSceneParticle::StaticStruct(), FOnGetDetailCustomizationInstance::CreateStatic(&FChaosVDSceneParticleCustomization::MakeInstance, StaticCastWeakPtr<SChaosVDMainTab>(AsWeak())));
 	DetailsView->RegisterInstancedCustomPropertyLayout(UChaosVDInstancedStaticMeshComponent::StaticClass(), FOnGetDetailCustomizationInstance::CreateStatic(&FChaosVDGeometryComponentCustomization::MakeInstance));
 	DetailsView->RegisterInstancedCustomPropertyLayout(UChaosVDStaticMeshComponent::StaticClass(), FOnGetDetailCustomizationInstance::CreateStatic(&FChaosVDGeometryComponentCustomization::MakeInstance));
-	DetailsView->RegisterInstancedCustomPropertyLayout(FChaosVDQueryVisitStep::StaticStruct(), FOnGetDetailCustomizationInstance::CreateStatic(&FChaosVDQueryVisitDataCustomization::MakeInstance));
+	DetailsView->RegisterInstancedCustomPropertyLayout(FChaosVDQueryVisitStep::StaticStruct(), FOnGetDetailCustomizationInstance::CreateStatic(&FChaosVDQueryVisitDataCustomization::MakeInstance, StaticCastWeakPtr<SChaosVDMainTab>(AsWeak())));
 	DetailsView->RegisterInstancedCustomPropertyLayout(FChaosVDQueryDataWrapper::StaticStruct(), FOnGetDetailCustomizationInstance::CreateStatic(&FChaosVDQueryDataWrapperCustomization::MakeInstance));
 	DetailsView->RegisterInstancedCustomPropertyLayout(FChaosVDSelectionMultipleView::StaticStruct(), FOnGetDetailCustomizationInstance::CreateStatic(&FChaosVDSelectionMultipleViewCustomization::MakeInstance));
 
 	//TODO: Rename FChaosVDParticleDataWrapperCustomization to something generic as currently works with any type that wants to hide properties of type FChaosVDWrapperDataBase with invalid data.
 	// Or another option is create a new custom layout intended to be generic from the get go
 	DetailsView->RegisterInstancedCustomPropertyTypeLayout(TEXT("ChaosVDQueryDataWrapper"), FOnGetPropertyTypeCustomizationInstance::CreateStatic(&FChaosVDParticleDataWrapperCustomization::MakeInstance));
-	DetailsView->RegisterInstancedCustomPropertyTypeLayout(TEXT("ChaosVDQueryVisitStep"), FOnGetPropertyTypeCustomizationInstance::CreateStatic(&FChaosVDParticleDataWrapperCustomization::MakeInstance));
 
 	DetailsView->RegisterInstancedCustomPropertyTypeLayout(TEXT("ChaosVDCollisionResponseParams"), FOnGetPropertyTypeCustomizationInstance::CreateStatic(&FChaosVDCollisionResponseParamsCustomization::MakeInstance, StaticCastWeakPtr<SChaosVDMainTab>(AsWeak())));
 	DetailsView->RegisterInstancedCustomPropertyTypeLayout(TEXT("ChaosVDCollisionObjectQueryParams"), FOnGetPropertyTypeCustomizationInstance::CreateStatic(&FChaosVDCollisionObjectParamsCustomization::MakeInstance,StaticCastWeakPtr<SChaosVDMainTab>(AsWeak())));
 	DetailsView->RegisterInstancedCustomPropertyTypeLayout(TEXT("ChaosVDShapeCollisionData"), FOnGetPropertyTypeCustomizationInstance::CreateStatic(&FChaosVDShapeDataCustomization::MakeInstance, StaticCastWeakPtr<SChaosVDMainTab>(AsWeak())));
+	DetailsView->RegisterInstancedCustomPropertyTypeLayout(TEXT("ChaosVDQueryVisitStep"), FOnGetPropertyTypeCustomizationInstance::CreateStatic(&FChaosVDQueryVisitDataPropertyCustomization::MakeInstance, StaticCastWeakPtr<SChaosVDMainTab>(AsWeak())));
 
 	// We don't need to validate properties, and trying to do so seems to be costing between 2-4ms per tick!
 	DetailsView->SetCustomValidatePropertyNodesFunction(FOnValidateDetailsViewPropertyNodes::CreateLambda([](const FRootPropertyNodeList& Root){ return true;}));
diff --git a/Engine/Plugins/ChaosVD/Source/ChaosVD/Private/Widgets/SChaosVDSceneQueryDataInspector.cpp b/Engine/Plugins/ChaosVD/Source/ChaosVD/Private/Widgets/SChaosVDSceneQueryDataInspector.cpp
--- a/Engine/Plugins/ChaosVD/Source/ChaosVD/Private/Widgets/SChaosVDSceneQueryDataInspector.cpp
+++ b/Engine/Plugins/ChaosVD/Source/ChaosVD/Private/Widgets/SChaosVDSceneQueryDataInspector.cpp
@@ -17,6 +17,7 @@
 #include "Widgets/SChaosVDNameListPicker.h"
 #include "Widgets/SChaosVDTimelineWidget.h"
 #include "Widgets/SChaosVDWarningMessageBox.h"
+#include "Widgets/Images/SImage.h"
 #include "Widgets/Input/SButton.h"
 #include "Widgets/Layout/SBox.h"
 #include "Widgets/Layout/SScrollBox.h"
@@ -190,29 +191,9 @@ TSharedRef<SWidget> SChaosVDSceneQueryDataInspector::GenerateQueryDetailsPanelSe
 				SNew(SVerticalBox)
 				.Visibility_Raw(this, &SChaosVDSceneQueryDataInspector::GetSQVisitDetailsSectionVisibility)
 				+SVerticalBox::Slot()
-				.FillHeight(0.9f)
 				[
 					SceneQueryHitDataDetailsView->GetWidget().ToSharedRef()
 				]
-				+SVerticalBox::Slot()
-				.AutoHeight()
-				.Padding(0.0f,InnerDetailsPanelsVerticalPadding,0.0f,InnerDetailsPanelsVerticalPadding)
-				[
-					SNew(SUniformGridPanel)
-					.SlotPadding(FAppStyle::GetMargin("StandardDialog.SlotPadding"))
-					.MinDesiredSlotWidth(FAppStyle::GetFloat("StandardDialog.MinDesiredSlotWidth"))
-					.MinDesiredSlotHeight(FAppStyle::GetFloat("StandardDialog.MinDesiredSlotHeight"))
-					+SUniformGridPanel::Slot(0, 0)
-					[
-						SNew(SButton)
-						.HAlign(HAlign_Center)
-						.VAlign(VAlign_Center)
-						.ContentPadding(FAppStyle::GetMargin("StandardDialog.ContentPadding"))
-						.IsEnabled_Raw(this, &SChaosVDSceneQueryDataInspector::GetSelectParticleHitStateEnable)
-						.Text_Raw(this, &SChaosVDSceneQueryDataInspector::GetSelectParticleText)
-						.OnClicked(this, &SChaosVDSceneQueryDataInspector::SelectParticleForCurrentQueryData)
-					]
-				]
 			];
 }
 
@@ -226,22 +207,53 @@ TSharedRef<SWidget> SChaosVDSceneQueryDataInspector::GenerateVisitStepControls()
 				.Visibility_Raw(this, &SChaosVDSceneQueryDataInspector::GetQueryStepPlaybackControlsVisibility)
 				+SVerticalBox::Slot()
 				.AutoHeight()
-				.Padding(0.0f, 2.0f, 0.0f, 2.0f)
 				[
-					SNew(STextBlock)
-					.Justification(ETextJustify::Center)
-					.Text_Raw(this, &SChaosVDSceneQueryDataInspector::GetSQVisitsStepsText)
+					SNew(SBorder)
+					.BorderImage(FAppStyle::Get().GetBrush("DetailsView.CategoryTop"))
+					.BorderBackgroundColor(FLinearColor::White)
+					[
+						SNew(SBox)
+						.VAlign(VAlign_Center)
+						.MinDesiredHeight(26.0f)
+						[
+							SNew(STextBlock)
+							.Justification(ETextJustify::Center)
+							.Text_Raw(this, &SChaosVDSceneQueryDataInspector::GetSQVisitsStepsText)
+							.Font(FAppStyle::GetFontStyle("DetailsView.CategoryFontStyle"))
+							.TextStyle(FAppStyle::Get(), "DetailsView.CategoryTextStyle")
+						]
+					]
 				]
 				+SVerticalBox::Slot()
+				.Padding(0.0f, 6.0f, 0.0f, 2.0f)
+				.AutoHeight()
 				[
-					SAssignNew(QueryStepsTimelineWidget, SChaosVDTimelineWidget)
-					.ButtonVisibilityFlags(EChaosVDTimelineElementIDFlags::AllManualStepping)
-					.IsEnabled_Raw(this, &SChaosVDSceneQueryDataInspector::GetSQVisitStepsEnabled)
-					.OnFrameChanged_Raw(this, &SChaosVDSceneQueryDataInspector::HandleQueryStepSelectionUpdated)
-					.OnButtonClicked(this, &SChaosVDSceneQueryDataInspector::HandleSQVisitTimelineInput)
-					.MinFrames_Raw(this, &SChaosVDSceneQueryDataInspector::GetCurrentMinSQVisitIndex)
-					.MaxFrames_Raw(this, &SChaosVDSceneQueryDataInspector::GetCurrentMaxSQVisitIndex)
-					.CurrentFrame_Raw(this, &SChaosVDSceneQueryDataInspector::GetCurrentSQVisitIndex)
+					SNew(SHorizontalBox)
+					+SHorizontalBox::Slot()
+					.FillWidth(0.9)
+					[
+						SAssignNew(QueryStepsTimelineWidget, SChaosVDTimelineWidget)
+						.ButtonVisibilityFlags(EChaosVDTimelineElementIDFlags::AllManualStepping)
+						.IsEnabled_Raw(this, &SChaosVDSceneQueryDataInspector::GetSQVisitStepsEnabled)
+						.OnFrameChanged_Raw(this, &SChaosVDSceneQueryDataInspector::HandleQueryStepSelectionUpdated)
+						.OnButtonClicked(this, &SChaosVDSceneQueryDataInspector::HandleSQVisitTimelineInput)
+						.MinFrames_Raw(this, &SChaosVDSceneQueryDataInspector::GetCurrentMinSQVisitIndex)
+						.MaxFrames_Raw(this, &SChaosVDSceneQueryDataInspector::GetCurrentMaxSQVisitIndex)
+						.CurrentFrame_Raw(this, &SChaosVDSceneQueryDataInspector::GetCurrentSQVisitIndex)
+					]
+					+SHorizontalBox::Slot()
+					.AutoWidth()
+					[
+						SNew(SButton)
+						.OnClicked_Raw(this, &SChaosVDSceneQueryDataInspector::SelectParticleForCurrentQueryData)
+						.ToolTipText(LOCTEXT("SelectSQVisitToolTip","Selects the current visited particle and collision shape"))
+						[
+							SNew(SImage)
+							.Image(FAppStyle::GetBrush("GenericCommands.SelectAll"))
+							.DesiredSizeOverride(FVector2D(16.0f,16.0f))
+							.ColorAndOpacity(FSlateColor::UseForeground())
+						]
+					]
 				]
 			];
 }
@@ -250,41 +262,18 @@ TSharedRef<SWidget> SChaosVDSceneQueryDataInspector::GenerateQueryTagInfoRow()
 {
 	return SNew(SBorder)
 			.Padding(0.5f)
-			.BorderImage(FAppStyle::Get().GetBrush("DetailsView.CategoryMiddle"))
-			.BorderBackgroundColor(FAppStyle::Get().GetSlateColor("Colors.Panel"))
-			.Padding(0)
+			.BorderImage(FAppStyle::Get().GetBrush("DetailsView.CategoryTop"))
+			.BorderBackgroundColor(FLinearColor::White)
 			[
 				SNew(SBox)
+				.MinDesiredHeight(26.0f)
+				.VAlign(VAlign_Center)
 				.Padding(2.0f)
 				[
-					SNew(SSplitter)
-					.Style(FAppStyle::Get(), "DetailsView.Splitter")
-					.PhysicalSplitterHandleSize(1.5f)
-					.HitDetectionSplitterHandleSize(5.0f)
-					+SSplitter::Slot()
-					.Value(0.2)
-					[
-						SNew(SBox)
-						.VAlign(VAlign_Center)
-						.Padding(1.0)
-						.VAlign(VAlign_Center)
-						[
-							SNew(STextBlock)
-							.Text(LOCTEXT("SceneQueriesNameLabel", "Query Tag" ))
-							.Font(FAppStyle::GetFontStyle("PropertyWindow.NormalFont"))
-						]	
-					]
-					+SSplitter::Slot()
-					[
-						SNew(SBox)
-						.VAlign(VAlign_Center)
-						.Padding(1.0)
-						[
-							SNew(STextBlock)
-							.Text_Raw(this, &SChaosVDSceneQueryDataInspector::GetQueryBeingInspectedTag)
-							.Font(FAppStyle::GetFontStyle("PropertyWindow.NormalFont"))
-						]
-					]
+					SNew(STextBlock)
+					.Text_Raw(this, &SChaosVDSceneQueryDataInspector::GetQueryBeingInspectedTag)
+					.Font(FAppStyle::GetFontStyle("DetailsView.CategoryFontStyle"))
+					.TextStyle(FAppStyle::Get(), "DetailsView.CategoryTextStyle")
 				]
 			];
 }
@@ -312,6 +301,12 @@ void SChaosVDSceneQueryDataInspector::SetQueryDataToInspect(const TSharedPtr<FCh
 				const TSharedPtr<FStructOnScope> SQVisitDataDataView = MakeShared<FStructOnScope>(FChaosVDQueryVisitStep::StaticStruct(), reinterpret_cast<uint8*>(&QueryDataToInspect->SQVisitData[SelectionContext->SQVisitIndex]));
 				SceneQueryHitDataDetailsView->SetStructureData(SQVisitDataDataView);
 				QueryDataToInspect->CurrentVisitIndex = SelectionContext->SQVisitIndex;
+
+				FReply SelectionResponse = SelectParticleForCurrentQueryData();
+				if (!SelectionResponse.IsEventHandled())
+				{
+					UE_LOG(LogChaosVDEditor, Warning, TEXT("[%hs] Failed to auto select visited particle"), __func__);
+				}
 			}
 		}
 
@@ -396,6 +391,12 @@ void SChaosVDSceneQueryDataInspector::HandleQueryStepSelectionUpdated(int32 NewS
 	FScopedSQInspectorSilencedSelectionEvents IgnoreSelectionEventsScope(*this);
 	SelectionObject->SelectData(NewSelection);
 
+	FReply SelectionResponse = SelectParticleForCurrentQueryData();
+	if (!SelectionResponse.IsEventHandled())
+	{
+		UE_LOG(LogChaosVDEditor, Warning, TEXT("[%hs] Failed to auto select visited particle"), __func__);
+	}
+
 	if (const TSharedPtr<FEditorModeTools> EditorModeToolsPtr = EditorModeToolsWeakPtr.Pin())
 	{
 		// We need to request a re-draw to make sure the debug draw view and selection outline is updated
@@ -409,19 +410,13 @@ void SChaosVDSceneQueryDataInspector::HandleQueryStepSelectionUpdated(int32 NewS
 FText SChaosVDSceneQueryDataInspector::GetQueryBeingInspectedTag() const
 {
 	const TSharedPtr<FChaosVDQueryDataWrapper> QueryDataBeingInspected = CurrentSceneQueryBeingInspectedHandle->GetDataAsShared<FChaosVDQueryDataWrapper>();
-	return FText::AsCultureInvariant(QueryDataBeingInspected ? QueryDataBeingInspected->CollisionQueryParams.TraceTag.ToString() : TEXT("None"));
-}
-
-FText SChaosVDSceneQueryDataInspector::GetSelectParticleText() const
-{
-	return LOCTEXT("SelectVisitedParticleHit", "Select Visited Particle Shape");
+	const FText QueryTag = FText::AsCultureInvariant(QueryDataBeingInspected ? QueryDataBeingInspected->CollisionQueryParams.TraceTag.ToString() : TEXT("None"));
+	return FText::Format(FTextFormat(LOCTEXT("SceneQueriesNameLabel", "Query Tag | {0}")), QueryTag);
 }
 
 FText SChaosVDSceneQueryDataInspector::GetSQVisitsStepsText() const
 {
-	const TSharedPtr<FChaosVDQueryDataWrapper> QueryDataBeingInspected = CurrentSceneQueryBeingInspectedHandle->GetDataAsShared<FChaosVDQueryDataWrapper>();
-	const int32 AvailableSQVisitDataNum = QueryDataBeingInspected ? QueryDataBeingInspected->SQVisitData.Num() : 0;
-	return FText::Format(FTextFormat(LOCTEXT("SQVisitStepsPlaybackControlsTitle", "Recorded SQ Visits Available {0}")), AvailableSQVisitDataNum);
+	return LOCTEXT("SQVisitStepsPlaybackControlsTitle", "Visited Particle Shapes");
 }
 
 FReply SChaosVDSceneQueryDataInspector::SelectParticleForCurrentQueryData() const
diff --git a/Engine/Plugins/ChaosVD/Source/ChaosVD/Private/Widgets/SChaosVDSceneQueryDataInspector.h b/Engine/Plugins/ChaosVD/Source/ChaosVD/Private/Widgets/SChaosVDSceneQueryDataInspector.h
--- a/Engine/Plugins/ChaosVD/Source/ChaosVD/Private/Widgets/SChaosVDSceneQueryDataInspector.h
+++ b/Engine/Plugins/ChaosVD/Source/ChaosVD/Private/Widgets/SChaosVDSceneQueryDataInspector.h
@@ -56,7 +56,6 @@ protected:
 	void HandleQueryStepSelectionUpdated(int32 NewStepIndex);
 	
 	FText GetQueryBeingInspectedTag() const;
-	FText GetSelectParticleText() const;
 	FText GetSQVisitsStepsText() const;
 	
 	FReply SelectParticleForCurrentQueryData() const;
diff --git a/Engine/Source/Runtime/Experimental/ChaosVisualDebugger/Public/DataWrappers/ChaosVDQueryDataWrappers.h b/Engine/Source/Runtime/Experimental/ChaosVisualDebugger/Public/DataWrappers/ChaosVDQueryDataWrappers.h
--- a/Engine/Source/Runtime/Experimental/ChaosVisualDebugger/Public/DataWrappers/ChaosVDQueryDataWrappers.h
+++ b/Engine/Source/Runtime/Experimental/ChaosVisualDebugger/Public/DataWrappers/ChaosVDQueryDataWrappers.h
@@ -277,6 +277,21 @@ struct TStructOpsTypeTraits<FChaosVDQueryFastData> : public TStructOpsTypeTraits
 	};
 };
 
+/** Possible results from a scene query */
+UENUM(meta = (Bitflags, UseEnumValuesAsMaskValuesInEditor = "true"))
+enum class EChaosVDSQHitFlags : uint16
+{
+	None =			0,
+	Position =		1 << 0,
+	Normal =		1 << 1,
+	Distance =		1 << 2,
+	UV =			1 << 3,
+	MTD =			1 << 9,
+	FaceIndex =		1 << 10
+};
+
+ENUM_CLASS_FLAGS(EChaosVDSQHitFlags)
+
 USTRUCT()
 struct FChaosVDQueryHitData : public FChaosVDWrapperDataBase
 {
@@ -290,7 +305,7 @@ struct FChaosVDQueryHitData : public FChaosVDWrapperDataBase
 	UPROPERTY(VisibleAnywhere, Category=QueryData)
 	int32 FaceIdx = INDEX_NONE;
 
-	UPROPERTY(VisibleAnywhere, Category=QueryData)
+	UPROPERTY(VisibleAnywhere, Category=QueryData, meta = (Bitmask, BitmaskEnum = "/Script/ChaosVDRuntime.EChaosVDSQHitFlags"))
 	uint16 Flags = 0;
 
 	UPROPERTY(VisibleAnywhere, Category=QueryData)
@@ -378,8 +393,13 @@ struct FChaosVDQueryVisitStep : public FChaosVDWrapperDataBase
 	UPROPERTY(VisibleAnywhere, Category="SQ Visit Data")
 	FChaosVDQueryHitData HitData;
 
-	UPROPERTY()
+	UPROPERTY(VisibleAnywhere, Category="SQ Visit Data")
 	EChaosVDSQVisitRejectReason RejectReason = EChaosVDSQVisitRejectReason::None;
+
+	/** Cached solver id property populated and used during playback of an already loaded file .
+	 * This should not be serialized
+	 */
+	int32 SolverID_Editor = INDEX_NONE;
 };
 
 inline FArchive& operator<<(FArchive& Ar, FChaosVDQueryVisitStep& Data)
