[ChaosVD] Code clean up for 5.6

Removing unnecessary exports + removing the console command to spawn multiple CVD instances withing the editor.

Now in 5.6 CVD can open multiple files in a single CVD instance, and for other use cases 5.6 also has CVD standalone.

There are a few features like session discovery that will not work properly with multiple CVD instances in the same editor.

That feature was always hidden in the UI because it was meant to be a temp workaround until we had multi-file support (which we do now)

diff --git a/Engine/Plugins/ChaosVD/Source/ChaosVD/Private/Actors/ChaosVDGeometryContainer.h b/Engine/Plugins/ChaosVD/Source/ChaosVD/Private/Actors/ChaosVDGeometryContainer.h
--- a/Engine/Plugins/ChaosVD/Source/ChaosVD/Private/Actors/ChaosVDGeometryContainer.h
+++ b/Engine/Plugins/ChaosVD/Source/ChaosVD/Private/Actors/ChaosVDGeometryContainer.h
@@ -11,7 +11,7 @@
 
 /** Actor that contains Static Mesh Components used to visualize the geometry we generated from the recorded data */
 UCLASS()
-class CHAOSVD_API AChaosVDGeometryContainer : public AActor
+class AChaosVDGeometryContainer : public AActor
 {
 	GENERATED_BODY()
 
diff --git a/Engine/Plugins/ChaosVD/Source/ChaosVD/Private/ChaosVDModule.cpp b/Engine/Plugins/ChaosVD/Source/ChaosVD/Private/ChaosVDModule.cpp
--- a/Engine/Plugins/ChaosVD/Source/ChaosVD/Private/ChaosVDModule.cpp
+++ b/Engine/Plugins/ChaosVD/Source/ChaosVD/Private/ChaosVDModule.cpp
@@ -22,15 +22,6 @@
 
 DEFINE_LOG_CATEGORY(LogChaosVDEditor);
 
-FAutoConsoleCommand ChaosVDSpawnNewCVDInstance(
-	TEXT("p.Chaos.VD.SpawnNewCVDInstance"),
-	TEXT("Opens a new CVD windows wothout closing an existing one"),
-	FConsoleCommandWithArgsDelegate::CreateLambda([](const TArray<FString>& Args)
-	{
-		FChaosVDModule::Get().SpawnCVDTab();
-	})
-);
-
 FString FChaosVDModule::ChaosVisualDebuggerProgramName = TEXT("ChaosVisualDebugger");
 
 FChaosVDModule& FChaosVDModule::Get()
@@ -70,11 +61,6 @@ void FChaosVDModule::ShutdownModule()
 
 	FGlobalTabmanager::Get()->UnregisterNomadTabSpawner(FChaosVDTabID::ChaosVisualDebuggerTab);
 
-	for (const FName TabID : CreatedExtraTabSpawnersIDs)
-	{
-		FGlobalTabmanager::Get()->UnregisterNomadTabSpawner(TabID);
-	}
-
 	if (IsStandaloneChaosVisualDebugger())
 	{
 		FCoreDelegates::OnFEngineLoopInitComplete.RemoveAll(this);
@@ -97,18 +83,7 @@ void FChaosVDModule::SpawnCVDTab()
 		FModuleManager::Get().LoadModule("StatusBar");
 	}
 
-	// Registering new tab spawners with random names is not the best idea to spawn new tabs, but it is good enough to test we can run multiple instances of CVD withing the editor.
-	// This is also why spawning new tabs it is only exposed via console commands for now.
-	// When this feature is deemed stable and exposed in the UI I will investigate a more correct way of implementing this
-	const FName NewTabID = FName(FChaosVDTabID::ChaosVisualDebuggerTab.ToString() + FGuid::NewGuid().ToString());
-
-	FGlobalTabmanager::Get()->RegisterNomadTabSpawner(NewTabID, FOnSpawnTab::CreateRaw(this, &FChaosVDModule::SpawnMainTab))
-						.SetDisplayName(LOCTEXT("VisualDebuggerTabTitle", "Chaos Visual Debugger"))
-						.SetTooltipText(LOCTEXT("VisualDebuggerTabDesc", "Opens the Chaos Visual Debugger window"));
-
-	FGlobalTabmanager::Get()->TryInvokeTab(NewTabID);
-
-	CreatedExtraTabSpawnersIDs.Add(NewTabID);
+	FGlobalTabmanager::Get()->TryInvokeTab(FChaosVDTabID::ChaosVisualDebuggerTab);
 }
 
 TSharedRef<SDockTab> FChaosVDModule::SpawnMainTab(const FSpawnTabArgs& Args)
diff --git a/Engine/Plugins/ChaosVD/Source/ChaosVD/Private/ChaosVDSkySphereInterface.h b/Engine/Plugins/ChaosVD/Source/ChaosVD/Private/ChaosVDSkySphereInterface.h
--- a/Engine/Plugins/ChaosVD/Source/ChaosVD/Private/ChaosVDSkySphereInterface.h
+++ b/Engine/Plugins/ChaosVD/Source/ChaosVD/Private/ChaosVDSkySphereInterface.h
@@ -16,7 +16,7 @@ class UChaosVDSkySphereInterface : public UInterface
 /**
  * Simple interface to allow update the light source in the BP SkySphere from C++
  */
-class CHAOSVD_API IChaosVDSkySphereInterface
+class IChaosVDSkySphereInterface
 {
 	GENERATED_BODY()
 
diff --git a/Engine/Plugins/ChaosVD/Source/ChaosVD/Private/Components/ChaosVDInstancedStaticMeshComponent.h b/Engine/Plugins/ChaosVD/Source/ChaosVD/Private/Components/ChaosVDInstancedStaticMeshComponent.h
--- a/Engine/Plugins/ChaosVD/Source/ChaosVD/Private/Components/ChaosVDInstancedStaticMeshComponent.h
+++ b/Engine/Plugins/ChaosVD/Source/ChaosVD/Private/Components/ChaosVDInstancedStaticMeshComponent.h
@@ -22,7 +22,7 @@ enum class EChaosVDMeshInstanceOperationsFlags
 ENUM_CLASS_FLAGS(EChaosVDMeshInstanceOperationsFlags);
 
 /** CVD version of an Instance Static Mesh Component that holds additional CVD data */
-UCLASS(HideCategories=("Transform"), MinimalAPI)
+UCLASS(HideCategories=("Transform"))
 class UChaosVDInstancedStaticMeshComponent : public UInstancedStaticMeshComponent, public IChaosVDGeometryComponent, public IChaosVDPooledObject
 {
 	GENERATED_BODY()
diff --git a/Engine/Plugins/ChaosVD/Source/ChaosVD/Private/Components/ChaosVDStaticMeshComponent.h b/Engine/Plugins/ChaosVD/Source/ChaosVD/Private/Components/ChaosVDStaticMeshComponent.h
--- a/Engine/Plugins/ChaosVD/Source/ChaosVD/Private/Components/ChaosVDStaticMeshComponent.h
+++ b/Engine/Plugins/ChaosVD/Source/ChaosVD/Private/Components/ChaosVDStaticMeshComponent.h
@@ -8,7 +8,7 @@
 #include "ChaosVDStaticMeshComponent.generated.h"
 
 /** CVD version of a Static Mesh Component that holds additional CVD data */
-UCLASS(HideCategories=("Transform"), MinimalAPI)
+UCLASS(HideCategories=("Transform"))
 class UChaosVDStaticMeshComponent : public UStaticMeshComponent, public IChaosVDGeometryComponent, public IChaosVDPooledObject
 {
 	GENERATED_BODY()
diff --git a/Engine/Plugins/ChaosVD/Source/ChaosVD/Private/Settings/ChaosVDGeneralSettings.h b/Engine/Plugins/ChaosVD/Source/ChaosVD/Private/Settings/ChaosVDGeneralSettings.h
--- a/Engine/Plugins/ChaosVD/Source/ChaosVD/Private/Settings/ChaosVDGeneralSettings.h
+++ b/Engine/Plugins/ChaosVD/Source/ChaosVD/Private/Settings/ChaosVDGeneralSettings.h
@@ -9,7 +9,7 @@
  * General settings that controls how CVD behaves
  */
 UCLASS(config=ChaosVD, PerObjectConfig)
-class CHAOSVD_API UChaosVDGeneralSettings : public UChaosVDSettingsObjectBase
+class UChaosVDGeneralSettings : public UChaosVDSettingsObjectBase
 {
 	GENERATED_BODY()
 
diff --git a/Engine/Plugins/ChaosVD/Source/ChaosVD/Private/Settings/ChaosVDMiscSettings.h b/Engine/Plugins/ChaosVD/Source/ChaosVD/Private/Settings/ChaosVDMiscSettings.h
--- a/Engine/Plugins/ChaosVD/Source/ChaosVD/Private/Settings/ChaosVDMiscSettings.h
+++ b/Engine/Plugins/ChaosVD/Source/ChaosVD/Private/Settings/ChaosVDMiscSettings.h
@@ -55,7 +55,7 @@ struct FChaosVDRecentFile
  * General non-core CVD settings
  */
 UCLASS(config=ChaosVD, PerObjectConfig)
-class CHAOSVD_API UChaosVDMiscSettings : public UChaosVDSettingsObjectBase
+class UChaosVDMiscSettings : public UChaosVDSettingsObjectBase
 {
 	GENERATED_BODY()
 
diff --git a/Engine/Plugins/ChaosVD/Source/ChaosVD/Private/Widgets/SChaosVDRecordingControls.h b/Engine/Plugins/ChaosVD/Source/ChaosVD/Private/Widgets/SChaosVDRecordingControls.h
--- a/Engine/Plugins/ChaosVD/Source/ChaosVD/Private/Widgets/SChaosVDRecordingControls.h
+++ b/Engine/Plugins/ChaosVD/Source/ChaosVD/Private/Widgets/SChaosVDRecordingControls.h
@@ -25,7 +25,7 @@ struct FSlateBrush;
 struct FSlateIcon;
 
 UCLASS()
-class CHAOSVD_API UChaosVDRecordingToolbarMenuContext : public UObject
+class UChaosVDRecordingToolbarMenuContext : public UObject
 {
 	GENERATED_BODY()
 
diff --git a/Engine/Plugins/ChaosVD/Source/ChaosVD/Private/Widgets/SChaosVDSceneQueryBrowser.h b/Engine/Plugins/ChaosVD/Source/ChaosVD/Private/Widgets/SChaosVDSceneQueryBrowser.h
--- a/Engine/Plugins/ChaosVD/Source/ChaosVD/Private/Widgets/SChaosVDSceneQueryBrowser.h
+++ b/Engine/Plugins/ChaosVD/Source/ChaosVD/Private/Widgets/SChaosVDSceneQueryBrowser.h
@@ -17,7 +17,7 @@ struct FChaosVDSceneQueryTreeItem;
 struct FChaosVDSolverDataSelectionHandle;
 
 UCLASS()
-class CHAOSVD_API UChaosVDSceneQueryBrowserToolbarMenuContext : public UObject
+class UChaosVDSceneQueryBrowserToolbarMenuContext : public UObject
 {
 	GENERATED_BODY()
 
@@ -30,7 +30,7 @@ public:
  * Widget class of the Scene Query Browser window, where all available scene queries are shown for the currently visualized frame
  * in a Scene Outliner kind of way
  */
-class CHAOSVD_API SChaosVDSceneQueryBrowser : public SCompoundWidget
+class SChaosVDSceneQueryBrowser : public SCompoundWidget
 {
 public:
 	SLATE_BEGIN_ARGS(SChaosVDSceneQueryBrowser)
diff --git a/Engine/Plugins/ChaosVD/Source/ChaosVD/Private/Widgets/SChaosVDSolverTracks.h b/Engine/Plugins/ChaosVD/Source/ChaosVD/Private/Widgets/SChaosVDSolverTracks.h
--- a/Engine/Plugins/ChaosVD/Source/ChaosVD/Private/Widgets/SChaosVDSolverTracks.h
+++ b/Engine/Plugins/ChaosVD/Source/ChaosVD/Private/Widgets/SChaosVDSolverTracks.h
@@ -16,7 +16,7 @@ class SChaosVDMainTab;
 struct FChaosVDTrackInfo;
 
 UCLASS()
-class CHAOSVD_API UChaosVDSolverTracksToolbarMenuContext : public UObject
+class UChaosVDSolverTracksToolbarMenuContext : public UObject
 {
 	GENERATED_BODY()
 
diff --git a/Engine/Plugins/ChaosVD/Source/ChaosVD/Private/Widgets/SChaosVDViewportToolbar.h b/Engine/Plugins/ChaosVD/Source/ChaosVD/Private/Widgets/SChaosVDViewportToolbar.h
--- a/Engine/Plugins/ChaosVD/Source/ChaosVD/Private/Widgets/SChaosVDViewportToolbar.h
+++ b/Engine/Plugins/ChaosVD/Source/ChaosVD/Private/Widgets/SChaosVDViewportToolbar.h
@@ -13,7 +13,7 @@ class SEditorViewportViewMenu;
 /**
  * Viewport toolbar widget used by the Chaos Visual Debugger
  */
-class CHAOSVD_API SChaosVDViewportToolbar : public SCommonEditorViewportToolbarBase
+class SChaosVDViewportToolbar : public SCommonEditorViewportToolbarBase
 {
 public:
 	SLATE_BEGIN_ARGS(SChaosVDViewportToolbar)
diff --git a/Engine/Plugins/ChaosVD/Source/ChaosVD/Public/Actors/ChaosVDDataContainerBaseActor.h b/Engine/Plugins/ChaosVD/Source/ChaosVD/Public/Actors/ChaosVDDataContainerBaseActor.h
--- a/Engine/Plugins/ChaosVD/Source/ChaosVD/Public/Actors/ChaosVDDataContainerBaseActor.h
+++ b/Engine/Plugins/ChaosVD/Source/ChaosVD/Public/Actors/ChaosVDDataContainerBaseActor.h
@@ -13,31 +13,31 @@ struct FChaosVDSolverFrameData;
 struct FChaosVDGameFrameData;
 
 /** Base class for any CVD actor that will contain frame related data (either solver frame or game frame) */
-UCLASS(Abstract, NotBlueprintable, NotPlaceable)
-class CHAOSVD_API AChaosVDDataContainerBaseActor : public AActor
+UCLASS(MinimalAPI, Abstract, NotBlueprintable, NotPlaceable)
+class AChaosVDDataContainerBaseActor : public AActor
 {
 	GENERATED_BODY()
 
 public:
-	AChaosVDDataContainerBaseActor();
-	virtual ~AChaosVDDataContainerBaseActor() override;
+	CHAOSVD_API AChaosVDDataContainerBaseActor();
+	CHAOSVD_API virtual ~AChaosVDDataContainerBaseActor() override;
 
-	virtual void UpdateFromNewGameFrameData(const FChaosVDGameFrameData& InGameFrameData);
+	CHAOSVD_API virtual void UpdateFromNewGameFrameData(const FChaosVDGameFrameData& InGameFrameData);
 
-	virtual void UpdateFromNewSolverFrameData(const FChaosVDSolverFrameData& InSolverFrameData);
+	CHAOSVD_API virtual void UpdateFromNewSolverFrameData(const FChaosVDSolverFrameData& InSolverFrameData);
 
-	virtual void UpdateFromNewSolverStageData(const FChaosVDSolverFrameData& InSolverFrameData, const FChaosVDFrameStageData& InSolverFrameStageData);
+	CHAOSVD_API virtual void UpdateFromNewSolverStageData(const FChaosVDSolverFrameData& InSolverFrameData, const FChaosVDFrameStageData& InSolverFrameStageData);
 
-	virtual void Destroyed() override;
+	CHAOSVD_API virtual void Destroyed() override;
 
 	virtual void SetScene(TWeakPtr<FChaosVDScene> InScene)
 	{
 		SceneWeakPtr = InScene;
 	}
 
-	virtual void HandleWorldStreamingLocationUpdated(const FVector& InLocation);
+	CHAOSVD_API virtual void HandleWorldStreamingLocationUpdated(const FVector& InLocation);
 	
-	void SetSolverID(int32 InSolverID);
+	CHAOSVD_API void SetSolverID(int32 InSolverID);
 
 	int32 GetSolverID() const
 	{
@@ -49,14 +49,14 @@ public:
 		return FTransform::Identity;
 	}
 
-	virtual void UpdateVisibility(bool bIsVisible);
+	CHAOSVD_API virtual void UpdateVisibility(bool bIsVisible);
 
 	virtual bool IsVisible() const
 	{
 		return true;
 	}
 
-	virtual void PostActorCreated() override;
+	CHAOSVD_API virtual void PostActorCreated() override;
 	
 	TWeakPtr<FChaosVDScene> GetScene() const
 	{
@@ -64,12 +64,12 @@ public:
 	}
 	
 #if WITH_EDITOR
-	void SetIsTemporarilyHiddenInEditor(bool bIsHidden) override;
+	CHAOSVD_API void SetIsTemporarilyHiddenInEditor(bool bIsHidden) override;
 #endif
 	
 protected:
 
-	void HandlePostInitializationExtensionRegistered(const TSharedRef<FChaosVDExtension>& NewExtension);
+	CHAOSVD_API void HandlePostInitializationExtensionRegistered(const TSharedRef<FChaosVDExtension>& NewExtension);
 	
 	TWeakPtr<FChaosVDScene> SceneWeakPtr;
 
diff --git a/Engine/Plugins/ChaosVD/Source/ChaosVD/Public/ChaosVDBaseSceneObject.h b/Engine/Plugins/ChaosVD/Source/ChaosVD/Public/ChaosVDBaseSceneObject.h
--- a/Engine/Plugins/ChaosVD/Source/ChaosVD/Public/ChaosVDBaseSceneObject.h
+++ b/Engine/Plugins/ChaosVD/Source/ChaosVD/Public/ChaosVDBaseSceneObject.h
@@ -49,7 +49,7 @@ struct FChaosVDBaseSceneObject
 	 * Sets a weak reference to an actor that is the parent of this object.
 	 * @param NewParent Parent Actor
 	 */
-	void SetParentParentActor(AActor* NewParent);
+	CHAOSVD_API void SetParentParentActor(AActor* NewParent);
 
 
 	/**
diff --git a/Engine/Plugins/ChaosVD/Source/ChaosVD/Public/ChaosVDModule.h b/Engine/Plugins/ChaosVD/Source/ChaosVD/Public/ChaosVDModule.h
--- a/Engine/Plugins/ChaosVD/Source/ChaosVD/Public/ChaosVDModule.h
+++ b/Engine/Plugins/ChaosVD/Source/ChaosVD/Public/ChaosVDModule.h
@@ -59,8 +59,6 @@ private:
 
 	TMap<FGuid, TWeakPtr<SDockTab>> ActiveCVDTabs;
 
-	TArray<FName> CreatedExtraTabSpawnersIDs;
-
 	TSharedPtr<FChaosVDTraceManager> ChaosVDTraceManager;
 
 	static FString ChaosVisualDebuggerProgramName;
diff --git a/Engine/Plugins/ChaosVD/Source/ChaosVD/Public/ChaosVDObjectDetailsTab.h b/Engine/Plugins/ChaosVD/Source/ChaosVD/Public/ChaosVDObjectDetailsTab.h
--- a/Engine/Plugins/ChaosVD/Source/ChaosVD/Public/ChaosVDObjectDetailsTab.h
+++ b/Engine/Plugins/ChaosVD/Source/ChaosVD/Public/ChaosVDObjectDetailsTab.h
@@ -57,24 +57,24 @@ void FChaosVDStandAloneObjectDetailsTab::SetStructToInspect(TStruct* NewStruct)
 }
 
 /** Spawns and handles an instance for the visual debugger details panel */
-class CHAOSVD_API FChaosVDObjectDetailsTab : public FChaosVDStandAloneObjectDetailsTab, public FChaosVDSceneSelectionObserver
+class FChaosVDObjectDetailsTab : public FChaosVDStandAloneObjectDetailsTab, public FChaosVDSceneSelectionObserver
 {
 public:
 
-	FChaosVDObjectDetailsTab(const FName& InTabID, TSharedPtr<FTabManager> InTabManager, TWeakPtr<SChaosVDMainTab> InOwningTabWidget) : FChaosVDStandAloneObjectDetailsTab(InTabID, InTabManager, InOwningTabWidget)
+	CHAOSVD_API FChaosVDObjectDetailsTab(const FName& InTabID, TSharedPtr<FTabManager> InTabManager, TWeakPtr<SChaosVDMainTab> InOwningTabWidget) : FChaosVDStandAloneObjectDetailsTab(InTabID, InTabManager, InOwningTabWidget)
 	{
 	}
 
 protected:
 
-	virtual TSharedRef<SDockTab> HandleTabSpawnRequest(const FSpawnTabArgs& Args) override;
-	virtual void HandleTabClosed(TSharedRef<SDockTab> InTabClosed) override;
+	CHAOSVD_API virtual TSharedRef<SDockTab> HandleTabSpawnRequest(const FSpawnTabArgs& Args) override;
+	CHAOSVD_API virtual void HandleTabClosed(TSharedRef<SDockTab> InTabClosed) override;
 
-	virtual void HandlePostSelectionChange(const UTypedElementSelectionSet* ChangedSelectionSet) override;
+	CHAOSVD_API virtual void HandlePostSelectionChange(const UTypedElementSelectionSet* ChangedSelectionSet) override;
 
-	void HandleActorsSelection(TArrayView<AActor*> SelectedActors);
+	CHAOSVD_API void HandleActorsSelection(TArrayView<AActor*> SelectedActors);
 
-	virtual void HandleSolverDataSelectionChange(const TSharedPtr<FChaosVDSolverDataSelectionHandle>& SelectionHandle);
+	CHAOSVD_API virtual void HandleSolverDataSelectionChange(const TSharedPtr<FChaosVDSolverDataSelectionHandle>& SelectionHandle);
 
 	FDelegateHandle SelectionDelegateHandle;
 	TWeakObjectPtr<UObject> CurrentSelectedObject = nullptr;
diff --git a/Engine/Plugins/ChaosVD/Source/ChaosVD/Public/ChaosVDSceneSelectionObserver.h b/Engine/Plugins/ChaosVD/Source/ChaosVD/Public/ChaosVDSceneSelectionObserver.h
--- a/Engine/Plugins/ChaosVD/Source/ChaosVD/Public/ChaosVDSceneSelectionObserver.h
+++ b/Engine/Plugins/ChaosVD/Source/ChaosVD/Public/ChaosVDSceneSelectionObserver.h
@@ -10,17 +10,17 @@ class UTypedElementSelectionSet;
 /** Base class that provided access to selection events from the Chaos visual debugger
  * local selection system.
  */
-class CHAOSVD_API FChaosVDSceneSelectionObserver
+class FChaosVDSceneSelectionObserver
 {
 public:
 
-	FChaosVDSceneSelectionObserver();
-	virtual ~FChaosVDSceneSelectionObserver();
+	CHAOSVD_API FChaosVDSceneSelectionObserver();
+	CHAOSVD_API virtual ~FChaosVDSceneSelectionObserver();
 
 protected:
 
 	/** Stores a weak ptr to the selection set object containing the current selection state */
-	void RegisterSelectionSetObject(UTypedElementSelectionSet* SelectionSetObject);
+	CHAOSVD_API void RegisterSelectionSetObject(UTypedElementSelectionSet* SelectionSetObject);
 	
 	const UTypedElementSelectionSet* GetSelectionSetObject() const { return ObservedSelection.IsValid() ? ObservedSelection.Get() : nullptr; };
 
diff --git a/Engine/Plugins/ChaosVD/Source/ChaosVD/Public/Components/ChaosVDSolverDataComponent.h b/Engine/Plugins/ChaosVD/Source/ChaosVD/Public/Components/ChaosVDSolverDataComponent.h
--- a/Engine/Plugins/ChaosVD/Source/ChaosVD/Public/Components/ChaosVDSolverDataComponent.h
+++ b/Engine/Plugins/ChaosVD/Source/ChaosVD/Public/Components/ChaosVDSolverDataComponent.h
@@ -12,13 +12,13 @@ struct FChaosVDGameFrameData;
 /**
  * Base class for all components that stores recorded solver data
  */
-UCLASS(Abstract)
-class CHAOSVD_API UChaosVDSolverDataComponent : public UActorComponent
+UCLASS(MinimalAPI, Abstract)
+class UChaosVDSolverDataComponent : public UActorComponent
 {
 	GENERATED_BODY()
 public:
-	virtual void ClearData() PURE_VIRTUAL(UChaosVDSolverDataComponent::ClearData);
-	virtual void SetScene(const TWeakPtr<FChaosVDScene>& InSceneWeakPtr);
+	CHAOSVD_API virtual void ClearData() PURE_VIRTUAL(UChaosVDSolverDataComponent::ClearData);
+	CHAOSVD_API virtual void SetScene(const TWeakPtr<FChaosVDScene>& InSceneWeakPtr);
 
 	virtual void UpdateFromNewGameFrameData(const FChaosVDGameFrameData& InGameFrameData) {}
 	virtual void UpdateFromNewSolverStageData(const FChaosVDSolverFrameData& InSolverFrameData, const FChaosVDFrameStageData& InSolverFrameStageData) {}
@@ -29,9 +29,9 @@ public:
 		return bIsVisible;
 	}
 
-	virtual void SetVisibility(bool bNewIsVisible);
+	CHAOSVD_API virtual void SetVisibility(bool bNewIsVisible);
 
-	virtual void HandleWorldStreamingLocationUpdated(const FVector& InLocation) {};
+	virtual void HandleWorldStreamingLocationUpdated(const FVector& InLocation) {}
 	
 	void SetSolverID(int32 InSolverID)
 	{
diff --git a/Engine/Plugins/ChaosVD/Source/ChaosVD/Public/ExtensionsSystem/ChaosVDExtension.h b/Engine/Plugins/ChaosVD/Source/ChaosVD/Public/ExtensionsSystem/ChaosVDExtension.h
--- a/Engine/Plugins/ChaosVD/Source/ChaosVD/Public/ExtensionsSystem/ChaosVDExtension.h
+++ b/Engine/Plugins/ChaosVD/Source/ChaosVD/Public/ExtensionsSystem/ChaosVDExtension.h
@@ -19,18 +19,13 @@ class UActorComponent;
 struct FChaosVDTrackInfo;
 
 /** Base class for any CVD Extension. This object will auto register itself with the extensions system, and receive any relevant CVD callbacks */
-class CHAOSVD_API FChaosVDExtension : public IModularFeature
+class FChaosVDExtension : public IModularFeature
 {
 public:
 
-	FChaosVDExtension()
-	{
-	}
+	CHAOSVD_API FChaosVDExtension() = default;
 
-	virtual ~FChaosVDExtension()
-	{
-		
-	}
+	CHAOSVD_API virtual ~FChaosVDExtension() = default;
 
 	/** FName used as Type within CVD Extension system */
 	FName GetExtensionType() const
diff --git a/Engine/Plugins/ChaosVD/Source/ChaosVD/Public/Interfaces/ChaosVDPooledObject.h b/Engine/Plugins/ChaosVD/Source/ChaosVD/Public/Interfaces/ChaosVDPooledObject.h
--- a/Engine/Plugins/ChaosVD/Source/ChaosVD/Public/Interfaces/ChaosVDPooledObject.h
+++ b/Engine/Plugins/ChaosVD/Source/ChaosVD/Public/Interfaces/ChaosVDPooledObject.h
@@ -5,7 +5,7 @@
 #include "UObject/Interface.h"
 #include "ChaosVDPooledObject.generated.h"
 
-UINTERFACE()
+UINTERFACE(MinimalAPI)
 class UChaosVDPooledObject : public UInterface
 {
 	GENERATED_BODY()
@@ -14,11 +14,11 @@ class UChaosVDPooledObject : public UInterface
 /**
  * Interface used by any object that can be pooled
  */
-class CHAOSVD_API IChaosVDPooledObject
+class IChaosVDPooledObject
 {
 	GENERATED_BODY()
 
 public:
-	virtual void OnAcquired() PURE_VIRTUAL(IChaosVDPooledObject::OnAcquired)
-	virtual void OnDisposed() PURE_VIRTUAL(IChaosVDPooledObject::OnDisposed)
+	CHAOSVD_API virtual void OnAcquired() PURE_VIRTUAL(IChaosVDPooledObject::OnAcquired)
+	CHAOSVD_API virtual void OnDisposed() PURE_VIRTUAL(IChaosVDPooledObject::OnDisposed)
 };
diff --git a/Engine/Plugins/ChaosVD/Source/ChaosVD/Public/Interfaces/ChaosVDSelectableObject.h b/Engine/Plugins/ChaosVD/Source/ChaosVD/Public/Interfaces/ChaosVDSelectableObject.h
--- a/Engine/Plugins/ChaosVD/Source/ChaosVD/Public/Interfaces/ChaosVDSelectableObject.h
+++ b/Engine/Plugins/ChaosVD/Source/ChaosVD/Public/Interfaces/ChaosVDSelectableObject.h
@@ -5,7 +5,7 @@
 #include "UObject/Interface.h"
 #include "ChaosVDSelectableObject.generated.h"
 
-UINTERFACE()
+UINTERFACE(MinimalAPI)
 class UChaosVDSelectableObject : public UInterface
 {
 	GENERATED_BODY()
@@ -14,11 +14,11 @@ class UChaosVDSelectableObject : public UInterface
 /**
  * Interface used for any object that can be selected in CVD and need to process selection events performed on it.
  */
-class CHAOSVD_API IChaosVDSelectableObject
+class IChaosVDSelectableObject
 {
 	GENERATED_BODY()
 
 public:
-	virtual void HandleSelected() PURE_VIRTUAL(IChaosVDSelectionAware::HandleSelected);
-	virtual void HandleDeSelected() PURE_VIRTUAL(IChaosVDSelectionAware::HandleUnSelected);
+	CHAOSVD_API virtual void HandleSelected() PURE_VIRTUAL(IChaosVDSelectionAware::HandleSelected);
+	CHAOSVD_API virtual void HandleDeSelected() PURE_VIRTUAL(IChaosVDSelectionAware::HandleUnSelected);
 };
diff --git a/Engine/Plugins/ChaosVD/Source/ChaosVD/Public/Settings/ChaosVDCoreSettings.h b/Engine/Plugins/ChaosVD/Source/ChaosVD/Public/Settings/ChaosVDCoreSettings.h
--- a/Engine/Plugins/ChaosVD/Source/ChaosVD/Public/Settings/ChaosVDCoreSettings.h
+++ b/Engine/Plugins/ChaosVD/Source/ChaosVD/Public/Settings/ChaosVDCoreSettings.h
@@ -15,7 +15,7 @@ class UTextureCube;
 
 DECLARE_MULTICAST_DELEGATE_OneParam(FChaosVDSettingChanged, UObject* SettingsObject)
 
-UCLASS()
+UCLASS(MinimalAPI)
 class UChaosVDSettingsObjectsOuter : public UObject
 {
 	GENERATED_BODY()
@@ -24,33 +24,32 @@ class UChaosVDSettingsObjectsOuter : public UObject
 /** Base class to be used by any CVD settings.
  * Contains the base logic to make these settings work with CVD's options save system
  */
-UCLASS(config = ChaosVD)
-class CHAOSVD_API UChaosVDSettingsObjectBase : public UObject
+UCLASS(config = ChaosVD, MinimalAPI)
+class UChaosVDSettingsObjectBase : public UObject
 {
 public:
-	UChaosVDSettingsObjectBase();
+	CHAOSVD_API UChaosVDSettingsObjectBase();
 
 private:
 	GENERATED_BODY()
 
 public:
 
+	CHAOSVD_API virtual void PostEditChangeProperty(FPropertyChangedEvent& PropertyChangedEvent) override;
 
-	virtual void PostEditChangeProperty(FPropertyChangedEvent& PropertyChangedEvent) override;
+	CHAOSVD_API FChaosVDSettingChanged& OnSettingsChanged() { return SettingsChangedDelegate; }
 
-	FChaosVDSettingChanged& OnSettingsChanged() { return SettingsChangedDelegate; }
+	CHAOSVD_API virtual void PostEditUndo() override;
 
-	virtual void PostEditUndo() override;
+	CHAOSVD_API virtual void OverridePerObjectConfigSection(FString& SectionName) override;
 
-	virtual void OverridePerObjectConfigSection(FString& SectionName) override;
-
-	FStringView GetConfigSectionName()
+	CHAOSVD_API FStringView GetConfigSectionName()
 	{
 		return OverrideConfigSectionName;
 	}
 
 protected:
-	virtual void BroadcastSettingsChanged();
+	CHAOSVD_API virtual void BroadcastSettingsChanged();
 	
 private:
 
@@ -65,16 +64,16 @@ private:
  * And makes sure that the viewport gets re-draw when a setting changes
  * 
  */
-UCLASS(config = ChaosVD)
-class CHAOSVD_API UChaosVDVisualizationSettingsObjectBase : public UChaosVDSettingsObjectBase
+UCLASS(config = ChaosVD, MinimalAPI)
+class UChaosVDVisualizationSettingsObjectBase : public UChaosVDSettingsObjectBase
 {
 	GENERATED_BODY()
 protected:
-	virtual void BroadcastSettingsChanged() override;
+	CHAOSVD_API virtual void BroadcastSettingsChanged() override;
 
 public:
 	/** Returns true if a visualization flag option of this setting option should be enabled in the UI */
-	virtual bool CanVisualizationFlagBeChangedByUI(uint32 Flag);
+	CHAOSVD_API virtual bool CanVisualizationFlagBeChangedByUI(uint32 Flag);
 };
 
 /** Core settings class for CVD */
diff --git a/Engine/Plugins/ChaosVD/Source/ChaosVD/Public/TEDS/ChaosVDSelectionInterface.h b/Engine/Plugins/ChaosVD/Source/ChaosVD/Public/TEDS/ChaosVDSelectionInterface.h
--- a/Engine/Plugins/ChaosVD/Source/ChaosVD/Public/TEDS/ChaosVDSelectionInterface.h
+++ b/Engine/Plugins/ChaosVD/Source/ChaosVD/Public/TEDS/ChaosVDSelectionInterface.h
@@ -9,8 +9,8 @@
 /**
  * CVD Custom Selection Interface. We don't have any specific implementation overrides yet
  */
-UCLASS()
-class CHAOSVD_API UChaosVDSelectionInterface : public UObject, public ITypedElementSelectionInterface
+UCLASS(MinimalAPI)
+class UChaosVDSelectionInterface : public UObject, public ITypedElementSelectionInterface
 {
 	GENERATED_BODY()
 };
diff --git a/Engine/Plugins/ChaosVD/Source/ChaosVD/Public/TEDS/ChaosVDStructTypedElementData.h b/Engine/Plugins/ChaosVD/Source/ChaosVD/Public/TEDS/ChaosVDStructTypedElementData.h
--- a/Engine/Plugins/ChaosVD/Source/ChaosVD/Public/TEDS/ChaosVDStructTypedElementData.h
+++ b/Engine/Plugins/ChaosVD/Source/ChaosVD/Public/TEDS/ChaosVDStructTypedElementData.h
@@ -81,7 +81,7 @@ namespace Chaos::VD::TypedElementDataUtil
 		return FTypedElementHandle();
 	}
 
-	void DestroyTypedElementHandleForStruct(void* InElementData);
+	CHAOSVD_API void DestroyTypedElementHandleForStruct(void* InElementData);
 
 	template <typename StructDataType>
 	StructDataType* GetStructDataFromTypedElementHandle(const FTypedElementHandle& InHandle, const bool bSilent = true)
diff --git a/Engine/Plugins/ChaosVD/Source/ChaosVD/Public/Trace/DataProcessors/ChaosVDDataProcessorBase.h b/Engine/Plugins/ChaosVD/Source/ChaosVD/Public/Trace/DataProcessors/ChaosVDDataProcessorBase.h
--- a/Engine/Plugins/ChaosVD/Source/ChaosVD/Public/Trace/DataProcessors/ChaosVDDataProcessorBase.h
+++ b/Engine/Plugins/ChaosVD/Source/ChaosVD/Public/Trace/DataProcessors/ChaosVDDataProcessorBase.h
@@ -91,9 +91,11 @@ protected:
 
 class FChaosVDGenericDataProcessor : FChaosVDDataProcessorBase
 {
-	explicit FChaosVDGenericDataProcessor(FStringView InCompatibleType, const TFunction<bool(const TArray<uint8>&)>& InProcessDataCallback);
+public:
+
+	CHAOSVD_API explicit FChaosVDGenericDataProcessor(FStringView InCompatibleType, const TFunction<bool(const TArray<uint8>&)>& InProcessDataCallback);
 
-	virtual bool ProcessRawData(const TArray<uint8>& InData) override;
+	CHAOSVD_API virtual bool ProcessRawData(const TArray<uint8>& InData) override;
 
 private:
 	TFunction<bool(const TArray<uint8>&)> ProcessDataCallback;
diff --git a/Engine/Plugins/ChaosVD/Source/ChaosVD/Public/Visualizers/ChaosVDDebugDrawUtils.h b/Engine/Plugins/ChaosVD/Source/ChaosVD/Public/Visualizers/ChaosVDDebugDrawUtils.h
--- a/Engine/Plugins/ChaosVD/Source/ChaosVD/Public/Visualizers/ChaosVDDebugDrawUtils.h
+++ b/Engine/Plugins/ChaosVD/Source/ChaosVD/Public/Visualizers/ChaosVDDebugDrawUtils.h
@@ -21,26 +21,26 @@ enum class EChaosVDDebugDrawTextLocationMode
 };
 
 /** Utility methods that allows Debug draw into the Chaos VD Editor */
-class CHAOSVD_API FChaosVDDebugDrawUtils
+class FChaosVDDebugDrawUtils
 {
 public:
 	//TODO: Create a common args struct that can be passed as parameter to all debug draw methods
 
-	static void DrawArrowVector(FPrimitiveDrawInterface* PDI, const FVector& StartLocation, const FVector& EndLocation, const FText& InDebugText, const FColor& Color, ESceneDepthPriorityGroup DepthPriority = SDPG_World, float Thickness = 0.0f);
-	static void DrawPoint(FPrimitiveDrawInterface* PDI, const FVector& Location, const FText& InDebugText, const FColor& Color, float Size = 10.0f, ESceneDepthPriorityGroup DepthPriority = SDPG_World);
-	static void DrawString(FStringView StringToDraw, const FVector& Location, const FColor& Color, EChaosVDDebugDrawTextLocationMode LocationMode = EChaosVDDebugDrawTextLocationMode::World);
-	static void DrawText(const FText& InText, const FVector& Location, const FColor& Color, EChaosVDDebugDrawTextLocationMode LocationMode = EChaosVDDebugDrawTextLocationMode::World);
-	static void DrawOnScreenWarning(const FText& InText, const FColor& Color);
-	static void DrawCircle(FPrimitiveDrawInterface* PDI, const FVector& Origin, float Radius, int32 Segments, const FColor& Color, float Thickness, const FVector& XAxis, const FVector& YAxis, const FText& InDebugText, ESceneDepthPriorityGroup DepthPriority = SDPG_World);
-	static void DrawBox(FPrimitiveDrawInterface* PDI, const FVector& InExtents, const FColor& InColor, const FTransform& InTransform, const FText& InDebugText, ESceneDepthPriorityGroup DepthPriority = SDPG_World, float Thickness = 2.0f);
-	static void DrawLine(FPrimitiveDrawInterface* PDI, const FVector& InStartPosition, const FVector& InEndPosition, const FColor& InColor, const FText& InDebugText, ESceneDepthPriorityGroup DepthPriority = SDPG_World, float Thickness = 2.0f);
-	static void DrawImplicitObject(FPrimitiveDrawInterface* PDI, const TSharedPtr<FChaosVDGeometryBuilder>& GeometryGenerator, const Chaos::FConstImplicitObjectPtr& ImplicitObject, const FTransform& InWorldTransform, const FColor& InColor, const FText& InDebugText, ESceneDepthPriorityGroup DepthPriority = SDPG_World, float Thickness = 2.0f);
+	CHAOSVD_API static void DrawArrowVector(FPrimitiveDrawInterface* PDI, const FVector& StartLocation, const FVector& EndLocation, const FText& InDebugText, const FColor& Color, ESceneDepthPriorityGroup DepthPriority = SDPG_World, float Thickness = 0.0f);
+	CHAOSVD_API static void DrawPoint(FPrimitiveDrawInterface* PDI, const FVector& Location, const FText& InDebugText, const FColor& Color, float Size = 10.0f, ESceneDepthPriorityGroup DepthPriority = SDPG_World);
+	CHAOSVD_API static void DrawString(FStringView StringToDraw, const FVector& Location, const FColor& Color, EChaosVDDebugDrawTextLocationMode LocationMode = EChaosVDDebugDrawTextLocationMode::World);
+	CHAOSVD_API static void DrawText(const FText& InText, const FVector& Location, const FColor& Color, EChaosVDDebugDrawTextLocationMode LocationMode = EChaosVDDebugDrawTextLocationMode::World);
+	CHAOSVD_API static void DrawOnScreenWarning(const FText& InText, const FColor& Color);
+	CHAOSVD_API static void DrawCircle(FPrimitiveDrawInterface* PDI, const FVector& Origin, float Radius, int32 Segments, const FColor& Color, float Thickness, const FVector& XAxis, const FVector& YAxis, const FText& InDebugText, ESceneDepthPriorityGroup DepthPriority = SDPG_World);
+	CHAOSVD_API static void DrawBox(FPrimitiveDrawInterface* PDI, const FVector& InExtents, const FColor& InColor, const FTransform& InTransform, const FText& InDebugText, ESceneDepthPriorityGroup DepthPriority = SDPG_World, float Thickness = 2.0f);
+	CHAOSVD_API static void DrawLine(FPrimitiveDrawInterface* PDI, const FVector& InStartPosition, const FVector& InEndPosition, const FColor& InColor, const FText& InDebugText, ESceneDepthPriorityGroup DepthPriority = SDPG_World, float Thickness = 2.0f);
+	CHAOSVD_API static void DrawImplicitObject(FPrimitiveDrawInterface* PDI, const TSharedPtr<FChaosVDGeometryBuilder>& GeometryGenerator, const Chaos::FConstImplicitObjectPtr& ImplicitObject, const FTransform& InWorldTransform, const FColor& InColor, const FText& InDebugText, ESceneDepthPriorityGroup DepthPriority = SDPG_World, float Thickness = 2.0f);
 	
-	static void DrawSphere(FPrimitiveDrawInterface* PDI, const FVector& Center, float Radius, int32 Segments, const FColor& InColor, const FText& InDebugText, ESceneDepthPriorityGroup DepthPriority = SDPG_World, float Thickness = 2.0f);
+	CHAOSVD_API static void DrawSphere(FPrimitiveDrawInterface* PDI, const FVector& Center, float Radius, int32 Segments, const FColor& InColor, const FText& InDebugText, ESceneDepthPriorityGroup DepthPriority = SDPG_World, float Thickness = 2.0f);
 
-	static void DrawCanvas(FViewport& InViewport, FSceneView& View, FCanvas& Canvas);
+	CHAOSVD_API static void DrawCanvas(FViewport& InViewport, FSceneView& View, FCanvas& Canvas);
 	
-	static bool CanDebugDraw();
+	CHAOSVD_API static bool CanDebugDraw();
 	
 	static int32 MaxLinesToDrawPerFrame;
 
diff --git a/Engine/Plugins/ChaosVD/Source/ChaosVD/Public/Widgets/SChaosVDDetailsView.h b/Engine/Plugins/ChaosVD/Source/ChaosVD/Public/Widgets/SChaosVDDetailsView.h
--- a/Engine/Plugins/ChaosVD/Source/ChaosVD/Public/Widgets/SChaosVDDetailsView.h
+++ b/Engine/Plugins/ChaosVD/Source/ChaosVD/Public/Widgets/SChaosVDDetailsView.h
@@ -21,7 +21,7 @@ class UObject;
 /**
  * Simple details for CVD objects and structures
  */
-class CHAOSVD_API SChaosVDDetailsView : public SCompoundWidget
+class SChaosVDDetailsView : public SCompoundWidget
 {
 public:
 	SLATE_BEGIN_ARGS(SChaosVDDetailsView)
@@ -33,13 +33,13 @@ public:
 	void Construct(const FArguments& InArgs, const TSharedRef<SChaosVDMainTab>& InMainTab);
 
 	/** Updates the current object this view details is viewing */
-	void SetSelectedObject(UObject* NewObject);
+	CHAOSVD_API void SetSelectedObject(UObject* NewObject);
 
 	/** Updates the current object this view details is viewing */
 	template<typename TStruct>
 	void SetSelectedStruct(TStruct* NewStruct);
 
-	void SetSelectedStruct(const TSharedPtr<FStructOnScope>& NewStruct);
+	CHAOSVD_API void SetSelectedStruct(const TSharedPtr<FStructOnScope>& NewStruct);
 
 protected:
 
