[ChaosVD] Adding missing exports from my previous change.

diff --git a/Engine/Plugins/ChaosVD/Source/ChaosVD/Public/ChaosVDObjectDetailsTab.h b/Engine/Plugins/ChaosVD/Source/ChaosVD/Public/ChaosVDObjectDetailsTab.h
--- a/Engine/Plugins/ChaosVD/Source/ChaosVD/Public/ChaosVDObjectDetailsTab.h
+++ b/Engine/Plugins/ChaosVD/Source/ChaosVD/Public/ChaosVDObjectDetailsTab.h
@@ -23,7 +23,7 @@ class FChaosVDStandAloneObjectDetailsTab : public FChaosVDTabSpawnerBase, public
 {
 public:
 
-	FChaosVDStandAloneObjectDetailsTab(const FName& InTabID, TSharedPtr<FTabManager> InTabManager, TWeakPtr<SChaosVDMainTab> InOwningTabWidget) : FChaosVDTabSpawnerBase(InTabID, InTabManager, InOwningTabWidget)
+	CHAOSVD_API FChaosVDStandAloneObjectDetailsTab(const FName& InTabID, TSharedPtr<FTabManager> InTabManager, TWeakPtr<SChaosVDMainTab> InOwningTabWidget) : FChaosVDTabSpawnerBase(InTabID, InTabManager, InOwningTabWidget)
 	{
 	}
 	
@@ -32,7 +32,7 @@ public:
 		return DetailsPanelView;
 	}
 
-	void AddUnsupportedStruct(const UStruct* Struct);
+	CHAOSVD_API void AddUnsupportedStruct(const UStruct* Struct);
 	
 	/** Updates the current object this view details is viewing */
 	template<typename TStruct>
@@ -40,10 +40,10 @@ public:
 
 protected:
 
-	virtual TSharedRef<SDockTab> HandleTabSpawnRequest(const FSpawnTabArgs& Args) override;
-	virtual void HandleTabClosed(TSharedRef<SDockTab> InTabClosed) override;
+	CHAOSVD_API virtual TSharedRef<SDockTab> HandleTabSpawnRequest(const FSpawnTabArgs& Args) override;
+	CHAOSVD_API virtual void HandleTabClosed(TSharedRef<SDockTab> InTabClosed) override;
 
-	bool IsSupportedStruct(const TWeakObjectPtr<const UStruct>& InWeakStructPtr);
+	CHAOSVD_API bool IsSupportedStruct(const TWeakObjectPtr<const UStruct>& InWeakStructPtr);
 
 	TSharedPtr<SChaosVDDetailsView> DetailsPanelView;
 
