[ChaosVD] Fix for an issue where GT data like scene queries, does not respect the selected timeline sync mode.

The problem is that the sync modes, like "Network tick" only work between solver data tracks, not GT data tracks. Proper fix is to remove the distinction between the two and just promote solver data tracks to be "Data tracks", but that is large change we can't do for 5.6 at this time.

This workarounds the issue in a similar way to how we handle the new multi-source feature that faces the same issues.

Instead of using the new Game Frame data callbacks directly, we only rely on the solver frame callbacks, and use their timing data to find the correct game therad frame data.

diff --git a/Engine/Plugins/ChaosVD/Source/ChaosVD/Private/Actors/ChaosVDSolverInfoActor.cpp b/Engine/Plugins/ChaosVD/Source/ChaosVD/Private/Actors/ChaosVDSolverInfoActor.cpp
--- a/Engine/Plugins/ChaosVD/Source/ChaosVD/Private/Actors/ChaosVDSolverInfoActor.cpp
+++ b/Engine/Plugins/ChaosVD/Source/ChaosVD/Private/Actors/ChaosVDSolverInfoActor.cpp
@@ -35,8 +35,41 @@ void AChaosVDSolverInfoActor::SetSolverName(const FName& InSolverName)
 	SetActorLabel(TEXT("Solver Data Container | ") + InSolverName.ToString());
 }
 
+void AChaosVDSolverInfoActor::FindAndUpdateFromCorrectGameFrameData(const FChaosVDSolverFrameData& InSolverFrameData)
+{
+	// Note: This is needed because the original implementation of game thread data playback don't support different sync modes
+	// therefore until that is re-implemented in a future version (UE-277464), we work around it by re-routing update game frame data callbacks from
+	// each solver data actor, using solver timing data as a starting point
+
+	TSharedPtr<FChaosVDScene> CVDScene = SceneWeakPtr.Pin();
+	TSharedPtr<FChaosVDRecording> Recording = CVDScene ? CVDScene->GetLoadedRecording() : nullptr;
+
+	const FChaosVDGameFrameData* GameFrameData = Recording ? Recording->GetGameFrameData_AssumesLocked(Recording->GetLowestGameFrameNumberAtCycle_AssumesLocked(InSolverFrameData.FrameCycle)) : nullptr;
+	if (!GameFrameData)
+	{
+		return;
+	}
+
+	{
+		FScopedGameFrameDataReRouting ScopedGTDataReRoute(this);
+		UpdateFromNewGameFrameData(*GameFrameData);
+	}
+}
+
+void AChaosVDSolverInfoActor::UpdateFromNewGameFrameData(const FChaosVDGameFrameData& InGameFrameData)
+{
+	if (!ensure(bInternallyReRoutingGameFrameData))
+	{
+		return;
+	}
+
+	Super::UpdateFromNewGameFrameData(InGameFrameData);
+}
+
 void AChaosVDSolverInfoActor::UpdateFromNewSolverFrameData(const FChaosVDSolverFrameData& InSolverFrameData)
 {
+	FindAndUpdateFromCorrectGameFrameData(InSolverFrameData);
+
 	SetSimulationTransform(InSolverFrameData.SimulationTransform);
 	Super::UpdateFromNewSolverFrameData(InSolverFrameData);
 }
diff --git a/Engine/Plugins/ChaosVD/Source/ChaosVD/Private/Actors/ChaosVDSolverInfoActor.h b/Engine/Plugins/ChaosVD/Source/ChaosVD/Private/Actors/ChaosVDSolverInfoActor.h
--- a/Engine/Plugins/ChaosVD/Source/ChaosVD/Private/Actors/ChaosVDSolverInfoActor.h
+++ b/Engine/Plugins/ChaosVD/Source/ChaosVD/Private/Actors/ChaosVDSolverInfoActor.h
@@ -46,6 +46,9 @@ public:
 		return bIsServer;
 	}
 
+
+	virtual void UpdateFromNewGameFrameData(const FChaosVDGameFrameData& InGameFrameData) override;
+
 	virtual void UpdateFromNewSolverFrameData(const FChaosVDSolverFrameData& InSolverFrameData) override;
 
 	virtual void SetScene(TWeakPtr<FChaosVDScene> InScene) override;
@@ -90,6 +93,11 @@ public:
 #endif
 
 protected:
+	/**
+	 * Finds the correct game frame data and executes the update flow for it based on the provided solver frame timing data.
+	 * @param InSolverFrameData Solver frame data where to get the timing data from
+	 */
+	void FindAndUpdateFromCorrectGameFrameData(const FChaosVDSolverFrameData& InSolverFrameData);
 
 	UPROPERTY(VisibleAnywhere, Category="Solver Data")
 	FTransform SimulationTransform;
diff --git a/Engine/Plugins/ChaosVD/Source/ChaosVD/Private/ChaosVDScene.cpp b/Engine/Plugins/ChaosVD/Source/ChaosVD/Private/ChaosVDScene.cpp
--- a/Engine/Plugins/ChaosVD/Source/ChaosVD/Private/ChaosVDScene.cpp
+++ b/Engine/Plugins/ChaosVD/Source/ChaosVD/Private/ChaosVDScene.cpp
@@ -267,9 +267,10 @@ void FChaosVDScene::HandleEnterNewGameFrame(int32 FrameNumber, const TArray<int3
 	{
 		AvailableSolversSet.Add(SolverID);
 
-		if (AChaosVDSolverInfoActor* SolverInfoActor = GetOrCreateSolverInfoActor(SolverID))
+		AChaosVDSolverInfoActor* SolverInfoActor = GetOrCreateSolverInfoActor(SolverID);
+		if (!ensure(SolverInfoActor))
 		{
-			SolverInfoActor->UpdateFromNewGameFrameData(InNewGameFrameData);
+			UE_LOG(LogChaosVDEditor, Error, TEXT("[%hs] Failed to create solver data actor for id [%d]"), __func__, SolverID);
 		}
 	}
 
diff --git a/Engine/Plugins/ChaosVD/Source/ChaosVD/Private/Components/ChaosVDAdditionalGTDataRouterComponent.cpp b/Engine/Plugins/ChaosVD/Source/ChaosVD/Private/Components/ChaosVDAdditionalGTDataRouterComponent.cpp
--- a/Engine/Plugins/ChaosVD/Source/ChaosVD/Private/Components/ChaosVDAdditionalGTDataRouterComponent.cpp
+++ b/Engine/Plugins/ChaosVD/Source/ChaosVD/Private/Components/ChaosVDAdditionalGTDataRouterComponent.cpp
@@ -21,9 +21,10 @@ void UChaosVDAdditionalGTDataRouterComponent::UpdateFromSolverFrameData(const FC
 			{
 				for (const TObjectPtr<AChaosVDDataContainerBaseActor>& DataContainerActor : CVDScene->GetDataContainerActorsView())
 				{
+					AChaosVDDataContainerBaseActor::FScopedGameFrameDataReRouting ScopedGTDataUpdate(DataContainerActor.Get());
 					DataContainerActor->UpdateFromNewGameFrameData(*GTFrameDataWrapper->FrameData);
 				}
-			}			
+			}
 		}
 	}
 }
diff --git a/Engine/Plugins/ChaosVD/Source/ChaosVD/Public/Actors/ChaosVDDataContainerBaseActor.h b/Engine/Plugins/ChaosVD/Source/ChaosVD/Public/Actors/ChaosVDDataContainerBaseActor.h
--- a/Engine/Plugins/ChaosVD/Source/ChaosVD/Public/Actors/ChaosVDDataContainerBaseActor.h
+++ b/Engine/Plugins/ChaosVD/Source/ChaosVD/Public/Actors/ChaosVDDataContainerBaseActor.h
@@ -62,6 +62,30 @@ public:
 	{
 		return SceneWeakPtr;
 	}
+
+	/**
+	 * Flags a CVD Data container actor for game frame data re-routing.
+	 * Used by some data container actors to know when the GT data they are processing comes from
+	 * a re-routing execution path.
+	 */
+	struct FScopedGameFrameDataReRouting
+	{
+		explicit FScopedGameFrameDataReRouting(AChaosVDDataContainerBaseActor* InDataContainerBaseActor)
+		{
+			check(InDataContainerBaseActor);
+			DataContainerBaseActor = InDataContainerBaseActor;
+			
+			DataContainerBaseActor->bInternallyReRoutingGameFrameData = true;
+		}
+
+		~FScopedGameFrameDataReRouting()
+		{
+			DataContainerBaseActor->bInternallyReRoutingGameFrameData = false;
+		}
+
+	private:
+		AChaosVDDataContainerBaseActor* DataContainerBaseActor = nullptr;
+	};
 	
 #if WITH_EDITOR
 	CHAOSVD_API void SetIsTemporarilyHiddenInEditor(bool bIsHidden) override;
@@ -74,4 +98,6 @@ protected:
 	TWeakPtr<FChaosVDScene> SceneWeakPtr;
 
 	int32 SolverDataID = INDEX_NONE;
+
+	bool bInternallyReRoutingGameFrameData = false;
 };
