[ChaosVD] Fix an issue where some UI elements were reacting to events they shouldn't + a fix introduced last week in the network sync mode where syncing from Client to solver track calculated the wrong offset (was a copy/paste line above error).

For example, selecting a line in the recorded log browser to scrub to the closest game frame, would then re-update the log browser UI when it shouldn't, selecting another line in the process (as the frame selection was approximate).

The issue was that the sync tracks methods were not properly passing along the playback instigator ID, therefore UI features like the log browser couldn't ignore playback events started by itself.

# tested

diff --git a/Engine/Plugins/ChaosVD/Source/ChaosVD/Private/ChaosVDPlaybackController.cpp b/Engine/Plugins/ChaosVD/Source/ChaosVD/Private/ChaosVDPlaybackController.cpp
--- a/Engine/Plugins/ChaosVD/Source/ChaosVD/Private/ChaosVDPlaybackController.cpp
+++ b/Engine/Plugins/ChaosVD/Source/ChaosVD/Private/ChaosVDPlaybackController.cpp
@@ -367,7 +367,7 @@ void FChaosVDPlaybackController::GoToTrackFrame_AssumesLockedAndSync(FGuid Insti
 	{
 		if (TSharedPtr<const FChaosVDTrackInfo> TrackToSyncWith = GetTrackInfo(TrackType, InTrackID))
 		{
-			SyncTracks_AssumesLocked(TrackToSyncWith.ToSharedRef(), CurrentSyncMode);
+			SyncTracks_AssumesLocked(TrackToSyncWith.ToSharedRef(), InstigatorID, CurrentSyncMode);
 		}
 	}
 }
@@ -514,11 +514,16 @@ int32 FChaosVDPlaybackController::ConvertCurrentFrameToOtherTrackFrame_AssumesLo
 										}
 										else // Client --> Client Track or Client --> Server Track
 										{
+
+											// Server tracks will not have a frame offset set, so we need to clamp it to 0 first
+											const int32 ToTrackNetworkTick = InToTrack->bIsServer && ToSolverFrameData->NetworkTickOffset == INDEX_NONE ? 0 : ToSolverFrameData->NetworkTickOffset;
+											const int32 FromTrackNetworkTick = InFromTrack->bIsServer && ToSolverFrameData->NetworkTickOffset == INDEX_NONE ? 0 : FromSolverFrameData->NetworkTickOffset;
+
 											// This works for Client --> Server conversion because in that case we want to add the frame offset.
 											// As the tick offset in server tracks is 0,  the following calculation will return a negative offset, which it will result in the intended addition in the last caluclation
-											if (ToSolverFrameData->NetworkTickOffset != INDEX_NONE || FromSolverFrameData->NetworkTickOffset != INDEX_NONE)
+											if (ToTrackNetworkTick != INDEX_NONE || FromTrackNetworkTick != INDEX_NONE)
 											{
-												int32 FrameOffset = ToSolverFrameData->NetworkTickOffset - FromSolverFrameData->NetworkTickOffset != INDEX_NONE;
+												const int32 FrameOffset = ToTrackNetworkTick - FromTrackNetworkTick;
 												ToFrame = FromSolverFrameData->InternalFrameNumber - FrameOffset;
 											}
 											else
@@ -1121,7 +1126,7 @@ void FChaosVDPlaybackController::GetTracksByType(EChaosVDTrackType Type, TArray<
 	}
 }
 
-void FChaosVDPlaybackController::SyncTracks(const TSharedRef<const FChaosVDTrackInfo>& FromTrack, EChaosVDSyncTimelinesMode TrackSyncMode)
+void FChaosVDPlaybackController::SyncTracks(const TSharedRef<const FChaosVDTrackInfo>& FromTrack, FGuid InstigatorID, EChaosVDSyncTimelinesMode TrackSyncMode)
 {
 	if (!LoadedRecording.IsValid())
 	{
@@ -1129,10 +1134,10 @@ void FChaosVDPlaybackController::SyncTracks(const TSharedRef<const FChaosVDTrack
 	}
 
 	FReadScopeLock ReadLock(LoadedRecording->GetRecordingDataLock());	
-	SyncTracks_AssumesLocked(FromTrack, TrackSyncMode);
+	SyncTracks_AssumesLocked(FromTrack, InstigatorID, TrackSyncMode);
 }
 
-void FChaosVDPlaybackController::SyncTracks_AssumesLocked(const TSharedRef<const FChaosVDTrackInfo>& FromTrack, EChaosVDSyncTimelinesMode TrackSyncMode)
+void FChaosVDPlaybackController::SyncTracks_AssumesLocked(const TSharedRef<const FChaosVDTrackInfo>& FromTrack, FGuid InstigatorID, EChaosVDSyncTimelinesMode TrackSyncMode)
 {
 	if (!FromTrack->bTrackSyncEnabled)
 	{
@@ -1162,9 +1167,9 @@ void FChaosVDPlaybackController::SyncTracks_AssumesLocked(const TSharedRef<const
 					const int32 SolverFrameNumber = LoadedRecording->GetLowestSolverFrameNumberGameFrame_AssumesLocked(TrackToReSyncFrom->TrackID, FromTrack->CurrentFrame);
 					const int32 StageNumber = GetTrackLastStageAtFrame_AssumesLocked(EChaosVDTrackType::Solver, TrackToReSyncFrom->TrackID, SolverFrameNumber);
 
-					GoToTrackFrame_AssumesLocked(PlaybackSelfInstigatorID, EChaosVDTrackType::Solver, TrackToReSyncFrom->TrackID, SolverFrameNumber, StageNumber);
+					GoToTrackFrame_AssumesLocked(InstigatorID, EChaosVDTrackType::Solver, TrackToReSyncFrom->TrackID, SolverFrameNumber, StageNumber);
 					
-					return SyncTracks_AssumesLocked(TrackToReSyncFrom.ToSharedRef(), TrackSyncMode);
+					return SyncTracks_AssumesLocked(TrackToReSyncFrom.ToSharedRef(), InstigatorID, TrackSyncMode);
 				}
 			}
 	
@@ -1178,7 +1183,7 @@ void FChaosVDPlaybackController::SyncTracks_AssumesLocked(const TSharedRef<const
 				const int32 SolverFrameNumber = LoadedRecording->GetLowestSolverFrameNumberGameFrame_AssumesLocked(SolverTrack->TrackID, FromTrack->CurrentFrame);
 				const int32 StageNumber = GetTrackLastStageAtFrame_AssumesLocked(EChaosVDTrackType::Solver, SolverTrack->TrackID, SolverFrameNumber);
 
-				GoToTrackFrame_AssumesLocked(PlaybackSelfInstigatorID, EChaosVDTrackType::Solver, SolverTrack->TrackID, SolverFrameNumber, StageNumber);
+				GoToTrackFrame_AssumesLocked(InstigatorID, EChaosVDTrackType::Solver, SolverTrack->TrackID, SolverFrameNumber, StageNumber);
 			}
 			break;
 		}	
@@ -1188,7 +1193,7 @@ void FChaosVDPlaybackController::SyncTracks_AssumesLocked(const TSharedRef<const
 			{
 				const int32 TargetGameFrameNumber = ConvertCurrentFrameToOtherTrackFrame_AssumesLocked(FromTrack, GameFramesTrackInfo.ToSharedRef(), TrackSyncMode);
 				constexpr int32 StageNumber = 0;
-				GoToTrackFrame_AssumesLocked(PlaybackSelfInstigatorID, GameFramesTrackInfo->TrackType, GameFramesTrackInfo->TrackID, TargetGameFrameNumber, StageNumber);
+				GoToTrackFrame_AssumesLocked(InstigatorID, GameFramesTrackInfo->TrackType, GameFramesTrackInfo->TrackID, TargetGameFrameNumber, StageNumber);
 			}
 
 			for (const TSharedPtr<const FChaosVDTrackInfo>& SolverTrack : AvailableSolverTracks)
@@ -1206,7 +1211,7 @@ void FChaosVDPlaybackController::SyncTracks_AssumesLocked(const TSharedRef<const
 				const int32 SolverFrameNumber = ConvertCurrentFrameToOtherTrackFrame_AssumesLocked(FromTrack, SolverTrack.ToSharedRef(), TrackSyncMode);
 				const int32 StageNumber = GetTrackLastStageAtFrame_AssumesLocked(EChaosVDTrackType::Solver, SolverTrack->TrackID, SolverFrameNumber);
 
-				GoToTrackFrame_AssumesLocked(PlaybackSelfInstigatorID, EChaosVDTrackType::Solver, SolverTrack->TrackID, SolverFrameNumber, StageNumber);
+				GoToTrackFrame_AssumesLocked(InstigatorID, EChaosVDTrackType::Solver, SolverTrack->TrackID, SolverFrameNumber, StageNumber);
 			}
 			break;
 		}
@@ -1322,7 +1327,7 @@ void FChaosVDPlaybackController::SetTimelineSyncMode(EChaosVDSyncTimelinesMode S
 
 			if (CurrentSyncMode != EChaosVDSyncTimelinesMode::Manual)
 			{
-				SyncTracks(TrackToReSyncFromRef, CurrentSyncMode);
+				SyncTracks(TrackToReSyncFromRef, PlaybackSelfInstigatorID, CurrentSyncMode);
 			}
 
 			if (bWasIsPlaying)
@@ -1409,7 +1414,7 @@ void FChaosVDPlaybackController::UpdateSolverTracksData()
 
 		if (bIsNewTrack)
 		{
-			SyncTracks(ActiveTrack, CurrentSyncMode);
+			SyncTracks(ActiveTrack, PlaybackSelfInstigatorID, CurrentSyncMode);
 		}
 	}
 }
diff --git a/Engine/Plugins/ChaosVD/Source/ChaosVD/Private/Widgets/SChaosVDRecordedLogBrowser.cpp b/Engine/Plugins/ChaosVD/Source/ChaosVD/Private/Widgets/SChaosVDRecordedLogBrowser.cpp
--- a/Engine/Plugins/ChaosVD/Source/ChaosVD/Private/Widgets/SChaosVDRecordedLogBrowser.cpp
+++ b/Engine/Plugins/ChaosVD/Source/ChaosVD/Private/Widgets/SChaosVDRecordedLogBrowser.cpp
@@ -438,6 +438,11 @@ void SChaosVDRecordedLogBrowser::HandleItemFocused(const TSharedPtr<FChaosVDLogV
 {
 	TSharedPtr<FChaosVDCachedLogItemEntry> LogEntryPtr = InLogViewListItem->ItemWeakPtr.Pin();
 
+	if (!LogEntryPtr)
+	{
+		return;
+	}
+
 	TSharedPtr<FChaosVDPlaybackController> PlaybackControllerPtr = GetObservedController().Pin();
 	TSharedPtr<FChaosVDRecording> RecordingData = PlaybackControllerPtr ? PlaybackControllerPtr->GetCurrentRecording().Pin() : nullptr;
 	if (!RecordingData)
diff --git a/Engine/Plugins/ChaosVD/Source/ChaosVD/Public/ChaosVDPlaybackController.h b/Engine/Plugins/ChaosVD/Source/ChaosVD/Public/ChaosVDPlaybackController.h
--- a/Engine/Plugins/ChaosVD/Source/ChaosVD/Public/ChaosVDPlaybackController.h
+++ b/Engine/Plugins/ChaosVD/Source/ChaosVD/Public/ChaosVDPlaybackController.h
@@ -380,10 +380,11 @@ public:
 	/**
 	 * Scrubs all tracks to the correct frame to be in sync with the provided track
 	 * @param FromTrack Track that every other track needs to sync to
-	 * @param TrackSyncMode Sync mode determining how the sync is calculated (Timestamp, Network tick, etx)
+	 * @param InstigatorID ID of the system that is instigating the playback action
+	 * @param TrackSyncMode Sync mode determining how the sync is calculated (Timestamp, Network tick, etc)
 	 */
-	void SyncTracks(const TSharedRef<const FChaosVDTrackInfo>& FromTrack, EChaosVDSyncTimelinesMode TrackSyncMode = EChaosVDSyncTimelinesMode::RecordedTimestamp);
-	void SyncTracks_AssumesLocked(const TSharedRef<const FChaosVDTrackInfo>& FromTrack, EChaosVDSyncTimelinesMode TrackSyncMode = EChaosVDSyncTimelinesMode::RecordedTimestamp);
+	void SyncTracks(const TSharedRef<const FChaosVDTrackInfo>& FromTrack, FGuid InstigatorID, EChaosVDSyncTimelinesMode TrackSyncMode = EChaosVDSyncTimelinesMode::RecordedTimestamp);
+	void SyncTracks_AssumesLocked(const TSharedRef<const FChaosVDTrackInfo>& FromTrack, FGuid InstigatorID, EChaosVDSyncTimelinesMode TrackSyncMode = EChaosVDSyncTimelinesMode::RecordedTimestamp);
 
 	/**
 	 * Toggles the sync functionality on the provided track
