[ChaosVD] Fix for potential crash and data races during playback.

There were some methods in the playback system accesing the recording data without taking the corresponding read lock.

The data race can happen during playback if the file is still loading (as that is done from the Trace Analysis thread)

diff --git a/Engine/Plugins/ChaosVD/Source/ChaosVD/Private/ChaosVDPlaybackController.cpp b/Engine/Plugins/ChaosVD/Source/ChaosVD/Private/ChaosVDPlaybackController.cpp
--- a/Engine/Plugins/ChaosVD/Source/ChaosVD/Private/ChaosVDPlaybackController.cpp
+++ b/Engine/Plugins/ChaosVD/Source/ChaosVD/Private/ChaosVDPlaybackController.cpp
@@ -428,7 +428,7 @@ const FChaosVDFrameStagesContainer* FChaosVDPlaybackController::GetTrackStagesDa
 	}
 }
 
-int32 FChaosVDPlaybackController::GetTrackFramesNumber(EChaosVDTrackType TrackType, const int32 InTrackID) const
+int32 FChaosVDPlaybackController::GetTrackFramesNumber_AssumesLocked(EChaosVDTrackType TrackType, const int32 InTrackID) const
 {
 	if (!LoadedRecording.IsValid())
 	{
@@ -440,13 +440,13 @@ int32 FChaosVDPlaybackController::GetTrackFramesNumber(EChaosVDTrackType TrackTy
 		case EChaosVDTrackType::Game:
 		{
 			// There is only one game track so no ID is needed
-			int32 GameFrames = LoadedRecording->GetAvailableGameFramesNumber();
+			int32 GameFrames = LoadedRecording->GetAvailableGameFramesNumber_AssumesLocked();
 			return GameFrames > 0 ? GameFrames : INDEX_NONE;
 		}
 
 		case EChaosVDTrackType::Solver:
 		{
-			int32 SolverFrames = LoadedRecording->GetAvailableSolverFramesNumber(InTrackID);
+			int32 SolverFrames = LoadedRecording->GetAvailableSolverFramesNumber_AssumesLocked(InTrackID);
 			return  SolverFrames > 0 ? SolverFrames : INDEX_NONE;
 		}
 		default:
@@ -454,6 +454,17 @@ int32 FChaosVDPlaybackController::GetTrackFramesNumber(EChaosVDTrackType TrackTy
 	}
 }
 
+int32 FChaosVDPlaybackController::GetTrackFramesNumber(EChaosVDTrackType TrackType, int32 InTrackID) const
+{
+	if (!LoadedRecording.IsValid())
+	{
+		return INDEX_NONE;
+	}
+	
+	FReadScopeLock ReadLock(LoadedRecording->GetRecordingDataLock());	
+	return GetTrackFramesNumber_AssumesLocked(TrackType, InTrackID);
+}
+
 int32 FChaosVDPlaybackController::ConvertCurrentFrameToOtherTrackFrame_AssumesLocked(const TSharedRef<const FChaosVDTrackInfo>& InFromTrack, const TSharedRef<const FChaosVDTrackInfo>& InToTrack, EChaosVDSyncTimelinesMode TrackSyncMode)
 {	
 	if (!ensure(LoadedRecording.IsValid()))
@@ -807,7 +818,7 @@ bool FChaosVDPlaybackController::Tick(float DeltaTime)
 					bPlayedFirstFrame = true;
 				}
 			}
-			else if (LoadedRecording->GetAvailableGameFramesNumber() > 0 && LoadedRecording->GetAvailableSolversNumber_AssumesLocked() > 0)
+			else if (!LoadedRecording->IsEmpty())
 			{
 				constexpr int32 GameFrameToLoad = 0;
 				// Game frames do not have stages
@@ -1391,6 +1402,8 @@ void FChaosVDPlaybackController::UpdateSolverTracksData()
 		return;
 	}
 
+	FReadScopeLock RecordingReadLock(LoadedRecording->GetRecordingDataLock());
+
 	const TMap<int32, TArray<FChaosVDSolverFrameData>>& SolversByID = LoadedRecording->GetAvailableSolvers_AssumesLocked();
 	for (const TPair<int32, TArray<FChaosVDSolverFrameData>>& SolverIDPair : SolversByID)
 	{
@@ -1407,14 +1420,14 @@ void FChaosVDPlaybackController::UpdateSolverTracksData()
 		}
 		
 		SolverTrackInfo->TrackID = SolverIDPair.Key;
-		SolverTrackInfo->MaxFrames = GetTrackFramesNumber(EChaosVDTrackType::Solver, SolverIDPair.Key);
-		SolverTrackInfo->TrackName = LoadedRecording->GetSolverFName(SolverIDPair.Key);
+		SolverTrackInfo->MaxFrames = GetTrackFramesNumber_AssumesLocked(EChaosVDTrackType::Solver, SolverIDPair.Key);
+		SolverTrackInfo->TrackName = LoadedRecording->GetSolverFName_AssumedLocked(SolverIDPair.Key);
 		SolverTrackInfo->TrackType = EChaosVDTrackType::Solver;
 		SolverTrackInfo->bIsServer = LoadedRecording->IsServerSolver_AssumesLocked(SolverIDPair.Key);
 
 		if (bIsNewTrack)
 		{
-			SyncTracks(ActiveTrack, PlaybackSelfInstigatorID, CurrentSyncMode);
+			SyncTracks_AssumesLocked(ActiveTrack, PlaybackSelfInstigatorID, CurrentSyncMode);
 		}
 	}
 }
@@ -1441,7 +1454,7 @@ void FChaosVDPlaybackController::HandleCurrentRecordingUpdated()
 		}
 	}
 
-	GameTrackInfo->MaxFrames = LoadedRecording.IsValid() ? LoadedRecording->GetAvailableGameFrames_AssumesLocked().Num() : INDEX_NONE;
+	GameTrackInfo->MaxFrames = LoadedRecording.IsValid() ? LoadedRecording->GetAvailableGameFramesNumber() : INDEX_NONE;
 	GameTrackInfo->TrackType = EChaosVDTrackType::Game;
 
 	// Each time the recording is updated, populate or update the existing solver tracks data
diff --git a/Engine/Plugins/ChaosVD/Source/ChaosVD/Private/ChaosVDScene.cpp b/Engine/Plugins/ChaosVD/Source/ChaosVD/Private/ChaosVDScene.cpp
--- a/Engine/Plugins/ChaosVD/Source/ChaosVD/Private/ChaosVDScene.cpp
+++ b/Engine/Plugins/ChaosVD/Source/ChaosVD/Private/ChaosVDScene.cpp
@@ -187,13 +187,15 @@ void FChaosVDScene::SetLoadedRecording(const TSharedPtr<FChaosVDRecording>& NewR
 	{
 		if (LoadedRecording)
 		{
-			NewRecordingInstance->OnGeometryDataLoaded().RemoveAll(GeometryGenerator.Get());
+			FWriteScopeLock RecordingWriteLock(LoadedRecording->GetRecordingDataLock());
+			LoadedRecording->OnGeometryDataLoaded().RemoveAll(GeometryGenerator.Get());
 		}
 
 		LoadedRecording = NewRecordingInstance;
 
 		if (GeometryGenerator)
 		{
+			FWriteScopeLock RecordingWriteLock(LoadedRecording->GetRecordingDataLock());
 			LoadedRecording->OnGeometryDataLoaded().AddSP(GeometryGenerator.ToSharedRef(), &FChaosVDGeometryBuilder::HandleNewGeometryData);
 		}
 	}
diff --git a/Engine/Plugins/ChaosVD/Source/ChaosVD/Public/ChaosVDPlaybackController.h b/Engine/Plugins/ChaosVD/Source/ChaosVD/Public/ChaosVDPlaybackController.h
--- a/Engine/Plugins/ChaosVD/Source/ChaosVD/Public/ChaosVDPlaybackController.h
+++ b/Engine/Plugins/ChaosVD/Source/ChaosVD/Public/ChaosVDPlaybackController.h
@@ -184,6 +184,7 @@ public:
 	 * @return Number of available frames
 	 */
 	int32 GetTrackFramesNumber(EChaosVDTrackType TrackType, int32 InTrackID) const;
+	int32 GetTrackFramesNumber_AssumesLocked(EChaosVDTrackType TrackType, int32 InTrackID) const;
 	
 	/**
 	 * Gets the current frame number at which the specified track is at
