[ChaosVD] Added support to attach custom data to the recording itself, in addition of being able to do it to frames.

Also, I added support for CVD component visualizers to draw 2D elements to the viewport's HUD.

diff --git a/Engine/Plugins/ChaosVD/Source/ChaosVD/Private/ChaosVDGeometryDataComponent.cpp b/Engine/Plugins/ChaosVD/Source/ChaosVD/Private/ChaosVDGeometryDataComponent.cpp
--- a/Engine/Plugins/ChaosVD/Source/ChaosVD/Private/ChaosVDGeometryDataComponent.cpp
+++ b/Engine/Plugins/ChaosVD/Source/ChaosVD/Private/ChaosVDGeometryDataComponent.cpp
@@ -10,6 +10,7 @@
 #include "Components/MeshComponent.h"
 #include "Materials/Material.h"
 #include "Materials/MaterialInterface.h"
+#include "Settings/ChaosVDParticleVisualizationSettings.h"
 #include "TEDS/ChaosVDStructTypedElementData.h"
 #include "Visualizers/ChaosVDParticleDataComponentVisualizer.h"
 
diff --git a/Engine/Plugins/ChaosVD/Source/ChaosVD/Private/ChaosVDPlaybackViewportClient.cpp b/Engine/Plugins/ChaosVD/Source/ChaosVD/Private/ChaosVDPlaybackViewportClient.cpp
--- a/Engine/Plugins/ChaosVD/Source/ChaosVD/Private/ChaosVDPlaybackViewportClient.cpp
+++ b/Engine/Plugins/ChaosVD/Source/ChaosVD/Private/ChaosVDPlaybackViewportClient.cpp
@@ -2,6 +2,7 @@
 
 #include "ChaosVDPlaybackViewportClient.h"
 
+#include "CanvasTypes.h"
 #include "Actors/ChaosVDSolverInfoActor.h"
 #include "ChaosVDScene.h"
 #include "ChaosVDSkySphereInterface.h"
@@ -177,6 +178,39 @@ bool FChaosVDPlaybackViewportClient::Internal_InputKey(const FInputKeyEventArgs&
 	return FEditorViewportClient::Internal_InputKey(EventArgs);
 }
 
+void FChaosVDPlaybackViewportClient::ExecuteVisualizationCallback(const TSharedRef<SChaosVDMainTab>& InMainTabToolkitHostRef,
+	const TSharedRef<FChaosVDScene>& SceneRef, const TFunction<void(const UActorComponent*, const TSharedRef<FComponentVisualizer>&)>& VisitorCallback)
+{
+	//TODO: Currently we can safely assume that any component in these actors is meant to have a visualizer, but we might need a proper interface for these components in the future
+	TInlineComponentArray<const UActorComponent*> ComponentsToVisualize;
+
+	TConstArrayView<TObjectPtr<AChaosVDDataContainerBaseActor>> DataContainerActors = SceneRef->GetDataContainerActorsView();
+	for (const TObjectPtr<AChaosVDDataContainerBaseActor>& DataContainerActor : DataContainerActors)
+	{
+		if (DataContainerActor)
+		{
+			constexpr bool bIncludeFromChildActors = false;
+			DataContainerActor->ForEachComponent(bIncludeFromChildActors, [&ComponentsToVisualize](UActorComponent* Component)
+			{
+				ComponentsToVisualize.Emplace(Component);
+			});
+		}
+	}
+
+	for (const UActorComponent* Component : ComponentsToVisualize)
+	{
+		if (!FChaosVDDebugDrawUtils::CanDebugDraw())
+		{
+			break;
+		}
+
+		if (const TSharedPtr<FComponentVisualizer> Visualizer = InMainTabToolkitHostRef->FindComponentVisualizer(Component->GetClass()))
+		{
+			VisitorCallback(Component, Visualizer.ToSharedRef());
+		}
+	}
+}
+
 void FChaosVDPlaybackViewportClient::HandleFocusRequest(FBox BoxToFocusOn)
 {
 	FocusViewportOnBox(BoxToFocusOn);
@@ -300,35 +334,12 @@ void FChaosVDPlaybackViewportClient::Draw(const FSceneView* View, FPrimitiveDraw
 	{
 		ScenePtr->UpdateWorldStreamingLocation(GetViewLocation());
 
-		//TODO: Currently we can safely assume that any component in these actors is meant to have a visualizer, but we might need a proper interface for these components in the future
-		TInlineComponentArray<const UActorComponent*> ComponentsToVisualize;
-
-
-		TConstArrayView<TObjectPtr<AChaosVDDataContainerBaseActor>> DataContainerActors = ScenePtr->GetDataContainerActorsView();
-		for (const TObjectPtr<AChaosVDDataContainerBaseActor>& DataContainerActor : DataContainerActors)
+		auto DrawComponentVisualization = [&View, &PDI](const UActorComponent* Component, const TSharedRef<FComponentVisualizer>& InVisualizer)
 		{
-			if (DataContainerActor)
-			{
-				constexpr bool bIncludeFromChildActors = false;
-				DataContainerActor->ForEachComponent(bIncludeFromChildActors, [&ComponentsToVisualize](UActorComponent* Component)
-				{
-					ComponentsToVisualize.Emplace(Component);
-				});
-			}
-		}
-
-		for (const UActorComponent* Component : ComponentsToVisualize)
-		{
-			if (!FChaosVDDebugDrawUtils::CanDebugDraw())
-			{
-				break;
-			}
+			InVisualizer->DrawVisualization(Component, View, PDI);
+		};
 
-			if (const TSharedPtr<FComponentVisualizer> Visualizer = MainTabToolkitHost->FindComponentVisualizer(Component->GetClass()))
-			{
-				Visualizer->DrawVisualization(Component, View, PDI);
-			}
-		}
+		ExecuteVisualizationCallback(MainTabToolkitHost.ToSharedRef(), ScenePtr.ToSharedRef(), DrawComponentVisualization);
 	}
 
 	FEditorViewportClient::Draw(View, PDI);
@@ -341,4 +352,20 @@ void FChaosVDPlaybackViewportClient::DrawCanvas(FViewport& InViewport, FSceneVie
 	FEditorViewportClient::DrawCanvas(InViewport, View, Canvas);
 
 	FChaosVDDebugDrawUtils::DrawCanvas(InViewport, View, Canvas);
+
+	{
+		
+		const TSharedPtr<SChaosVDMainTab> MainTabToolkitHost = ModeTools.IsValid() ? StaticCastSharedPtr<SChaosVDMainTab>(ModeTools->GetToolkitHost()) : nullptr;
+		const TSharedPtr<FChaosVDScene> ScenePtr = CVDScene.Pin();
+
+		if (MainTabToolkitHost && ScenePtr)
+		{
+			auto DrawComponentHUDVisualization = [&View, &Canvas, &InViewport](const UActorComponent* Component, const TSharedRef<FComponentVisualizer>& InVisualizer)
+			{
+				InVisualizer->DrawVisualizationHUD(Component, &InViewport, &View, &Canvas);
+			};
+
+			ExecuteVisualizationCallback(MainTabToolkitHost.ToSharedRef(), ScenePtr.ToSharedRef(), DrawComponentHUDVisualization);
+		}
+	}
 }
diff --git a/Engine/Plugins/ChaosVD/Source/ChaosVD/Private/ChaosVDPlaybackViewportClient.h b/Engine/Plugins/ChaosVD/Source/ChaosVD/Private/ChaosVDPlaybackViewportClient.h
--- a/Engine/Plugins/ChaosVD/Source/ChaosVD/Private/ChaosVDPlaybackViewportClient.h
+++ b/Engine/Plugins/ChaosVD/Source/ChaosVD/Private/ChaosVDPlaybackViewportClient.h
@@ -3,6 +3,8 @@
 #pragma once
 #include "EditorViewportClient.h"
 
+class SChaosVDMainTab;
+class FComponentVisualizer;
 class UChaosVDCoreSettings;
 struct FChaosVDGameFrameData;
 class FChaosVDScene;
@@ -70,6 +72,9 @@ protected:
 
 private:
 
+	void ExecuteVisualizationCallback(const TSharedRef<SChaosVDMainTab>& InMainTabToolkitHostRef,
+		const TSharedRef<FChaosVDScene>& SceneRef, const TFunction<void(const UActorComponent*, const TSharedRef<FComponentVisualizer>&)>& VisitorCallback);
+
 	void HandleFocusRequest(FBox BoxToFocusOn);
 	void HandleActorMoving(AActor* MovedActor) const;
 
diff --git a/Engine/Plugins/ChaosVD/Source/ChaosVD/Private/ChaosVDSceneStreaming.h b/Engine/Plugins/ChaosVD/Source/ChaosVD/Private/ChaosVDSceneStreaming.h
deleted file mode 100644
diff --git a/Engine/Plugins/ChaosVD/Source/ChaosVD/Private/Components/ChaosVDParticleDataComponent.h b/Engine/Plugins/ChaosVD/Source/ChaosVD/Private/Components/ChaosVDParticleDataComponent.h
deleted file mode 100644
diff --git a/Engine/Plugins/ChaosVD/Source/ChaosVD/Private/Components/ChaosVDSolverCharacterGroundConstraintDataComponent.cpp b/Engine/Plugins/ChaosVD/Source/ChaosVD/Private/Components/ChaosVDSolverCharacterGroundConstraintDataComponent.cpp
--- a/Engine/Plugins/ChaosVD/Source/ChaosVD/Private/Components/ChaosVDSolverCharacterGroundConstraintDataComponent.cpp
+++ b/Engine/Plugins/ChaosVD/Source/ChaosVD/Private/Components/ChaosVDSolverCharacterGroundConstraintDataComponent.cpp
@@ -3,7 +3,7 @@
 #include "Components/ChaosVDSolverCharacterGroundConstraintDataComponent.h"
 
 #include "ChaosVDCharacterGroundConstraintDataProviderInterface.h"
-#include "ChaosVDParticleDataComponent.h"
+#include "Components/ChaosVDParticleDataComponent.h"
 #include "ChaosVDScene.h"
 #include "ChaosVDSettingsManager.h"
 #include "Selection.h"
diff --git a/Engine/Plugins/ChaosVD/Source/ChaosVD/Private/Visualizers/ChaosVDParticleDataComponentVisualizer.cpp b/Engine/Plugins/ChaosVD/Source/ChaosVD/Private/Visualizers/ChaosVDParticleDataComponentVisualizer.cpp
--- a/Engine/Plugins/ChaosVD/Source/ChaosVD/Private/Visualizers/ChaosVDParticleDataComponentVisualizer.cpp
+++ b/Engine/Plugins/ChaosVD/Source/ChaosVD/Private/Visualizers/ChaosVDParticleDataComponentVisualizer.cpp
@@ -47,6 +47,12 @@ struct FChaosVDScopedParticleHitProxy
 	FPrimitiveDrawInterface* PDIPtr = nullptr;
 };
 
+bool FChaosVDParticleDataVisualizationContext::IsVisualizationFlagEnabled(EChaosVDParticleDataVisualizationFlags Flag) const
+{
+	const EChaosVDParticleDataVisualizationFlags FlagsAsParticleFlags = static_cast<EChaosVDParticleDataVisualizationFlags>(VisualizationFlags);
+	return EnumHasAnyFlags(FlagsAsParticleFlags, Flag);
+}
+
 FChaosVDParticleDataComponentVisualizer::FChaosVDParticleDataComponentVisualizer()
 {
 	FChaosVDParticleDataComponentVisualizer::RegisterVisualizerMenus();
diff --git a/Engine/Plugins/ChaosVD/Source/ChaosVD/Private/Visualizers/ChaosVDParticleDataComponentVisualizer.h b/Engine/Plugins/ChaosVD/Source/ChaosVD/Private/Visualizers/ChaosVDParticleDataComponentVisualizer.h
deleted file mode 100644
diff --git a/Engine/Plugins/ChaosVD/Source/ChaosVD/Public/ChaosVDSceneParticle.h b/Engine/Plugins/ChaosVD/Source/ChaosVD/Public/ChaosVDSceneParticle.h
--- a/Engine/Plugins/ChaosVD/Source/ChaosVD/Public/ChaosVDSceneParticle.h
+++ b/Engine/Plugins/ChaosVD/Source/ChaosVD/Public/ChaosVDSceneParticle.h
@@ -89,7 +89,7 @@ public:
 		return bIsActive;
 	}
 
-	FBox GetBoundingBox() const;
+	CHAOSVD_API FBox GetBoundingBox() const;
 	FBox GetInflatedBoundingBox() const;
 	Chaos::TAABB<double, 3> GetChaosBoundingBox() const;
 
diff --git a/Engine/Plugins/ChaosVD/Source/ChaosVD/Public/ChaosVDSceneStreaming.h b/Engine/Plugins/ChaosVD/Source/ChaosVD/Public/ChaosVDSceneStreaming.h
new file mode 100644
--- /dev/null
+++ b/Engine/Plugins/ChaosVD/Source/ChaosVD/Public/ChaosVDSceneStreaming.h
@@ -0,0 +1,244 @@
+// Copyright Epic Games, Inc. All Rights Reserved.
+
+#pragma once
+#include "Chaos/AABBTree.h"
+#include "Chaos/ISpatialAcceleration.h"
+#include "Containers/ContainersFwd.h"
+#include "Containers/Queue.h"
+#include "Containers/Ticker.h"
+#include "Containers/Set.h"
+#include "HAL/CriticalSection.h"
+#include "Math/Box.h"
+#include "Misc/Optional.h"
+#include "Templates/SharedPointer.h"
+#include "Tasks/Task.h"
+
+#include "ChaosVDSceneStreaming.generated.h"
+
+class FChaosVDScene;
+struct FChaosVDBaseSceneObject;
+
+UENUM()
+enum class EChaosVDStreamingDirtyFlags : uint32
+{
+	None = 0,
+	StreamingExtents = 1 << 0,
+	StreamingEnabled = 1 << 1,
+	StreamingSourceLocation = 1 << 2,
+	AccelerationStructure = 1 << 3
+};
+ENUM_CLASS_FLAGS(EChaosVDStreamingDirtyFlags)
+
+class IChaosVDStreamingDataSource
+{
+public:
+	virtual ~IChaosVDStreamingDataSource() = default;
+
+	virtual TConstArrayView<TSharedRef<FChaosVDBaseSceneObject>> GetStreamableSceneObjects() const = 0;
+	virtual FRWLock& GetObjectsLock() const = 0;
+};
+
+/** Simple pseudo level streaming system that works with a collection of ChaosVDSceneObjects -
+ * This system only updates a desired streaming state (from multiple threads) and issues a sync request in the GT.
+ * How these actions are executed depends on the implementations of each ChaosVDSceneObjects derived object
+ */
+class FChaosVDSceneStreaming : public FTSTickerObjectBase
+{
+public:
+
+	template<uint32 Alignment = DEFAULT_ALIGNMENT> class TMemStackSetAllocator : public TSetAllocator<TSparseArrayAllocator<TMemStackAllocator<Alignment>, TMemStackAllocator<Alignment>>, TMemStackAllocator<Alignment>> {};
+
+	FChaosVDSceneStreaming();
+	virtual ~FChaosVDSceneStreaming() override;
+
+	/** Relevant streaming system data that needs processing */
+	struct FPendingTrackingOperation
+	{
+		enum class EType
+		{
+			None,
+			AddOrUpdate,
+			Remove
+		};
+
+		FBox Bounds = FBox(EForceInit::ForceInitToZero);
+		int32 ObjectID = INDEX_NONE;
+		EType OperationType = EType::None;
+	};
+
+	void Initialize();
+	void DeInitialize();
+
+	virtual bool Tick(float DeltaTime) override;
+
+	/** Adds an object tracking operation to the queue */
+	void EnqueuePendingTrackingOperation(const FPendingTrackingOperation&& Operation);
+	
+	/**
+	 * Creates the request object tracking operation for the provided object and adds it to the queue
+	 * @param InSceneObject Object from which create the operation from
+	 * @param Type Type of operation to perform (AddOrUpdate, Remove)
+	 */
+	void EnqueuePendingTrackingOperation(const TSharedRef<FChaosVDBaseSceneObject>& InSceneObject, FPendingTrackingOperation::EType Type);
+
+	/**
+	 * Updates the current location used to calculate what objects should be fully loaded or not
+	 * @param NewLocation New streaming source location
+	 */
+	void UpdateStreamingSourceLocation(const FVector& NewLocation);
+
+	/**
+	 * Returns true if the provided bounds are within the current calculated streaming volume
+	 * @param Bounds Bounds to evaluate
+	 */
+	bool IsInStreamingRange(const FBox& Bounds) const;
+
+	/**
+	 * Sets the correct streaming state for the provided object based on its bounds
+	 * @param InSceneObject Object to evaluate
+	 */
+	void UpdateStreamingStateForObject(const TSharedRef<FChaosVDBaseSceneObject>& InSceneObject) const;
+
+	/**
+	 * Creates an object tracking operation structure with the necessary data to update the streaming system state for this object
+	 * @param InSceneObject Object from which get the necessary data
+	 * @param Type Operation type (AddOrUpdate, Remove)
+	 */
+	FPendingTrackingOperation CreateStreamingTrackingOperation(const TSharedRef<FChaosVDBaseSceneObject>& InSceneObject, FPendingTrackingOperation::EType Type);
+
+	/**
+	 * Resets the streaming system state
+	 */
+	void Reset();
+
+	/**
+	 * Sets a ptr to an array with all the objects managed by this streaming system
+	 * @param InStreamingDataSource Object that provides access to the data this system will manage
+	 */
+	void SetStreamingDataSource(IChaosVDStreamingDataSource* InStreamingDataSource);
+
+	/** Changes the enabled state of this streaming system, and updates the world acordingly
+	 * @param bNewEnabled New enabled state
+	 */
+	void SetStreamingEnabled(bool bNewEnabled);
+
+	/**
+	 * Sets a weak ptr to the CVD scene where all the objects managed by this system live
+	 * @param InSceneWeakPtr 
+	 */
+	void SetScene(const TWeakPtr<FChaosVDScene>& InSceneWeakPtr)
+	{
+		SceneWeakPtr = InSceneWeakPtr;
+	}
+
+	/**
+	 * Handles any streaming settings changes done by the user and updates the world accordingly
+	 * @param SettingsObject Object containing the new settings data
+	 */
+	void HandleSettingsChanged(UObject* SettingsObject);
+
+
+	/** Returns true if this system is enabled */
+	bool IsEnabled() const
+	{
+		return bIsStreamingSystemEnabled;
+	}
+
+protected:
+
+	void SetDirtyFlags(EChaosVDStreamingDirtyFlags Flag);
+	void RemoveDirtyFlag(EChaosVDStreamingDirtyFlags Flag);
+
+	void ProcessPendingOperations();
+
+	void UpdateStreamingState();
+
+	void SetStreamingExtent(float NewExtent);
+
+	void UpdateStreamingQueryShape();
+
+	void ReBuildAccelStructureFromSourceDataArray();
+
+	void MakeEverythingVisible();
+
+	float StreamingExtent = 5000.0f;
+
+	float MovementThreshold = 10.0f;
+
+	FVector CurrentStreamingSourceLocation = FVector::ZeroVector;
+
+	TOptional<FVector> LastStreamingLocationUpdate;
+
+	FBox SteamingViewBox;
+
+	Chaos::TAABBTree<int32, Chaos::TAABBTreeLeafArray<int32>> StreamingAccelerationStructure;
+
+	IChaosVDStreamingDataSource* StreamingDataSource = nullptr;
+
+	TQueue<FPendingTrackingOperation, EQueueMode::Mpsc> PendingOperationsPerParticle;
+
+	EChaosVDStreamingDirtyFlags DirtyFlags;
+
+	bool bIsStreamingSystemEnabled = true;
+
+	bool bProcessPendingOperationsQueueInWorkerThread = true;
+
+	TWeakPtr<FChaosVDScene> SceneWeakPtr;
+
+	UE::Tasks::FTask CurrentProcessingTaskHandle;
+
+	mutable FRWLock StreamingAccelerationOperationsLock;
+
+	mutable FRWLock DirtyFlagsLock;
+
+	struct FQueryVisitor
+	{
+		FQueryVisitor(TSet<int32>& InResults) : CollectedResults(InResults)
+		{
+		
+		}
+
+		bool VisitOverlap(const Chaos::TSpatialVisitorData<int32>& Instance)
+		{
+			CollectedResults.Add(Instance.Payload);
+			return true;
+		}
+		bool VisitSweep(const Chaos::TSpatialVisitorData<int32>& Instance, Chaos::FQueryFastData& CurData)
+		{
+			check(false);
+			return true;
+		}
+		bool VisitRaycast(const Chaos::TSpatialVisitorData<int32>& Instance, Chaos::FQueryFastData& CurData)
+		{
+			check(false);
+			return true;
+		}
+
+		const void* GetQueryData() const
+		{
+			return nullptr;
+		}
+
+		const void* GetSimData() const
+		{
+			return nullptr;
+		}
+
+		const void* GetQueryPayload() const 
+		{ 
+			return nullptr; 
+		}
+
+		bool HasBlockingHit() const
+		{
+			return false;
+		}
+
+		bool ShouldIgnore(const Chaos::TSpatialVisitorData<int32>& Instance) const
+		{
+			return false;
+		}
+
+		TSet<int32>& CollectedResults;
+	};
+};
diff --git a/Engine/Plugins/ChaosVD/Source/ChaosVD/Public/Components/ChaosVDParticleDataComponent.h b/Engine/Plugins/ChaosVD/Source/ChaosVD/Public/Components/ChaosVDParticleDataComponent.h
new file mode 100644
--- /dev/null
+++ b/Engine/Plugins/ChaosVD/Source/ChaosVD/Public/Components/ChaosVDParticleDataComponent.h
@@ -0,0 +1,188 @@
+﻿// Copyright Epic Games, Inc. All Rights Reserved.
+
+#pragma once
+
+#include "ChaosVDScene.h"
+#include "ChaosVDSceneParticle.h"
+#include "ChaosVDSceneStreaming.h"
+#include "Chaos/AABBTree.h"
+#include "Components/ChaosVDSolverDataComponent.h"
+#include "Math/GenericOctree.h"
+#include "TEDS/ChaosVDSelectionInterface.h"
+#include "TEDS/ChaosVDStructTypedElementData.h"
+#include "ChaosVDParticleDataComponent.generated.h"
+
+class FChaosVDSceneStreaming;
+struct FChaosVDParticleDataWrapper;
+struct FChaosVDSceneParticle;
+
+/**
+ * Component that references all particle data for a specific solver for the current frame, and handles how the visualization
+ * is updated based on that data
+ */
+UCLASS()
+class UChaosVDParticleDataComponent : public UChaosVDSolverDataComponent, public IChaosVDStreamingDataSource
+{
+public:
+
+	UChaosVDParticleDataComponent();
+
+	typedef TArray<UE::Editor::DataStorage::RowHandle, TInlineAllocator<256, TMemStackAllocator<>>> FTedsHandlesForBatchArray;
+
+private:
+	GENERATED_BODY()
+
+public:
+
+	virtual void SetScene(const TWeakPtr<FChaosVDScene>& InSceneWeakPtr) override;
+	virtual void ClearData() override;
+
+	virtual void UpdateFromNewSolverStageData(const FChaosVDSolverFrameData& InSolverFrameData, const FChaosVDFrameStageData& InSolverFrameStageData) override;
+	virtual void UpdateFromSolverFrameData(const FChaosVDSolverFrameData& InSolverFrameData) override;
+	
+	template <class TCallback>
+	void VisitSelectedParticleData(TCallback VisitCallback) const;
+
+	template <class TCallback>
+	void VisitAllParticleInstances(TCallback VisitCallback) const;
+
+	template <class TCallback>
+	void VisitAllParticleData(TCallback VisitCallback) const;
+
+	virtual void SetVisibility(bool bNewIsVisible) override;
+
+	TSharedPtr<FChaosVDSceneParticle> GetParticleInstanceByID(int32 ParticleID) const;
+	TSharedPtr<FChaosVDSceneParticle> GetParticleInstanceByID_AssumesLocked(int32 ParticleID) const;
+
+	TSharedPtr<FChaosVDBaseSceneObject> GetParticleContainerByType(EChaosVDParticleType ParticleType);
+	TSharedPtr<FChaosVDBaseSceneObject> GetParticleContainerByType_AssumesLocked(EChaosVDParticleType ParticleType);
+
+	const TSortedMap<EChaosVDParticleType, TSharedPtr<FChaosVDBaseSceneObject>>& GetParticleSceneContainersByType()
+	{
+		return ParticleSceneContainersByType;
+	}
+
+	CHAOSVD_API FChaosVDSceneParticle* GetSelectedParticle() const;
+
+	virtual void HandleWorldStreamingLocationUpdated(const FVector& InLocation) override;
+
+	// BEGIN IChaosVDStreamingDataSource
+	virtual TConstArrayView<TSharedRef<FChaosVDBaseSceneObject>> GetStreamableSceneObjects() const override
+	{
+		return SolverParticlesArray;
+	}
+	
+	virtual FRWLock& GetObjectsLock() const override
+	{
+		return ParticleSceneObjectsLock;
+	}
+	// END IChaosVDStreamingDataSource
+
+protected:
+
+	int32 GetMaxElementCountForBatchTEDSUpdate_AssumesLocked() const;
+
+	void BatchTEDSSyncFromWorldTag(FTedsHandlesForBatchArray& ParticlesPendingTEDSTagSync,
+		const UE::Editor::DataStorage::ICompatibilityProvider* CompatibilityStorage,
+		UE::Editor::DataStorage::ICoreProvider* DataStorage);
+
+	void BatchTEDSAddRemoveTags(UE::Editor::DataStorage::ICoreProvider* EditorDataStorage, FTedsHandlesForBatchArray& ParticlesPendingTEDSTagSync,
+		TConstArrayView<const UScriptStruct*> ColumnsToAdd, TConstArrayView<const UScriptStruct*> ColumnsToRemove);
+
+	void ProcessRemovedParticles_AssumessLocked(const TSharedRef<FChaosVDScene>& SceneRef, const FChaosVDSolverFrameData& InSolverFrameData,
+		const FChaosVDFrameStageData& InSolverFrameStageData, const TSet<int32>& ParticlesUpdatedIDsForKeyFrameDiff, FTedsHandlesForBatchArray& OutRemovedTedsHandles);
+
+	void BatchApplySolverVisibilityToParticle(bool bNewIsVisible);
+
+	void HandleVisibilitySettingsUpdated(UObject* SettingsObject);
+	void HandleColorsSettingsUpdated(UObject* SettingsObject);
+
+	bool IsServerData();
+	const FString& GetSolverName();
+
+	/** Creates an ChaosVDParticle actor for the Provided recorded Particle Data */
+	TSharedRef<FChaosVDSceneParticle> CreateSceneParticle_AssumesLocked(const TSharedRef<FChaosVDParticleDataWrapper>& InParticleData, const FChaosVDSolverFrameData& InFrameData);
+
+	TOptional<bool> bCachedIsServerData;
+	TOptional<FString> bCachedSolverName;
+
+	TSortedMap<EChaosVDParticleType, TSharedPtr<FChaosVDBaseSceneObject>> ParticleSceneContainersByType;
+
+	TMap<int32, TSharedRef<FChaosVDSceneParticle>> SolverParticlesByID;
+	TArray<TSharedRef<FChaosVDBaseSceneObject>> SolverParticlesArray;
+
+	mutable FRWLock ParticleSceneObjectsLock;
+
+	FChaosVDSceneStreaming StreamingSystem;
+};
+
+template <typename TCallback>
+void UChaosVDParticleDataComponent::VisitSelectedParticleData(TCallback VisitCallback) const
+{
+	TSharedPtr<FChaosVDScene> ScenePtr = SceneWeakPtr.Pin();
+	if (!ScenePtr)
+	{
+		return;
+	}
+	using namespace Chaos::VD::TypedElementDataUtil;
+
+	TArray<FTypedElementHandle, TInlineAllocator<8>> SelectedParticlesHandles;
+	ScenePtr->GetElementSelectionSet()->GetSelectedElementHandles(SelectedParticlesHandles, UChaosVDSelectionInterface::StaticClass());
+
+	for (const FTypedElementHandle& SelectedParticleHandle : SelectedParticlesHandles)
+	{
+		if (FChaosVDSceneParticle* ParticleInstance = GetStructDataFromTypedElementHandle<FChaosVDSceneParticle>(SelectedParticleHandle))
+		{
+			TSharedPtr<const FChaosVDParticleDataWrapper> ParticleDataViewer = ParticleInstance ? ParticleInstance->GetParticleData() : nullptr;
+	
+			if (!ensure(ParticleDataViewer))
+			{
+				continue;
+			}
+
+			if (ParticleDataViewer->SolverID != SolverID)
+			{
+				continue;
+			}
+
+			if (!VisitCallback(ParticleDataViewer))
+			{
+				return;
+			}
+		}
+	}
+}
+
+template <typename TCallback>
+void UChaosVDParticleDataComponent::VisitAllParticleInstances(TCallback VisitCallback) const
+{
+	FReadScopeLock ReadLock(ParticleSceneObjectsLock);
+
+	for (const TPair<int32, TSharedRef<FChaosVDSceneParticle>>& ParticleWithIDPair : SolverParticlesByID)
+	{
+		if (!VisitCallback(ParticleWithIDPair.Value))
+		{
+			return;
+		}
+	}
+}
+
+template <typename TCallback>
+void UChaosVDParticleDataComponent::VisitAllParticleData(TCallback VisitCallback) const
+{
+	FReadScopeLock ReadLock(ParticleSceneObjectsLock);
+
+	for (const TPair<int32, TSharedRef<FChaosVDSceneParticle>>& ParticleWithIDPair : SolverParticlesByID)
+	{
+		TSharedPtr<const FChaosVDParticleDataWrapper> ParticleDataViewer = ParticleWithIDPair.Value->GetParticleData();
+		if (!ensure(ParticleDataViewer))
+		{
+			continue;
+		}
+
+		if (!VisitCallback(ParticleDataViewer))
+		{
+			return;
+		}
+	}
+}
diff --git a/Engine/Plugins/ChaosVD/Source/ChaosVD/Public/Trace/ChaosVDTraceProvider.h b/Engine/Plugins/ChaosVD/Source/ChaosVD/Public/Trace/ChaosVDTraceProvider.h
--- a/Engine/Plugins/ChaosVD/Source/ChaosVD/Public/Trace/ChaosVDTraceProvider.h
+++ b/Engine/Plugins/ChaosVD/Source/ChaosVD/Public/Trace/ChaosVDTraceProvider.h
@@ -126,9 +126,9 @@ private:
 
 	FString GenerateFormattedStringListFromSet(const TSet<FString>& StringsSet);
 
-	int32 RemapSolverID(int32 SolverID);
-
 public:
+	CHAOSVD_API int32 RemapSolverID(int32 SolverID);
+
 	virtual void BeginEdit() const override{};
 	virtual void EndEdit() const override{};
 	virtual void EditAccessCheck() const override{};
diff --git a/Engine/Plugins/ChaosVD/Source/ChaosVD/Public/Visualizers/ChaosVDParticleDataComponentVisualizer.h b/Engine/Plugins/ChaosVD/Source/ChaosVD/Public/Visualizers/ChaosVDParticleDataComponentVisualizer.h
new file mode 100644
--- /dev/null
+++ b/Engine/Plugins/ChaosVD/Source/ChaosVD/Public/Visualizers/ChaosVDParticleDataComponentVisualizer.h
@@ -0,0 +1,69 @@
+﻿// Copyright Epic Games, Inc. All Rights Reserved.
+
+#pragma once
+
+#include "Visualizers/ChaosVDComponentVisualizerBase.h"
+#include "ComponentVisualizer.h"
+#include "HAL/Platform.h"
+#include "Templates/SharedPointer.h"
+
+struct FChaosVDSceneParticle;
+class FChaosVDGeometryBuilder;
+class UChaosVDParticleVisualizationDebugDrawSettings;
+
+struct FChaosParticleDataDebugDrawSettings;
+struct FChaosVDParticleDataWrapper;
+
+enum class EChaosVDParticleDataVisualizationFlags : uint32;
+
+namespace Chaos
+{
+	struct FTrimeshBVH;
+	class FTriangleMeshImplicitObject;
+	class FAABBVectorized;
+} // namespace Chaos
+
+
+struct FChaosVDParticleDataVisualizationContext : public FChaosVDVisualizationContext
+{
+	TWeakPtr<FChaosVDGeometryBuilder> GeometryGenerator = nullptr;
+	bool bIsSelectedData = false;
+	bool bShowDebugText = false;
+
+	const UChaosVDParticleVisualizationDebugDrawSettings* DebugDrawSettings = nullptr;
+
+	bool IsVisualizationFlagEnabled(EChaosVDParticleDataVisualizationFlags Flag) const;
+};
+
+/**
+ * Component visualizer in charge of generating debug draw visualizations for for particles
+ */
+class FChaosVDParticleDataComponentVisualizer : public FChaosVDComponentVisualizerBase
+{
+public:
+	FChaosVDParticleDataComponentVisualizer();
+
+	virtual void RegisterVisualizerMenus() override;
+
+	virtual void DrawVisualization(const UActorComponent* Component, const FSceneView* View, FPrimitiveDrawInterface* PDI) override;
+
+	virtual bool CanHandleClick(const HChaosVDComponentVisProxy& VisProxy) override;
+
+	virtual bool SelectVisualizedData(const HChaosVDComponentVisProxy& VisProxy, const TSharedRef<FChaosVDScene>& InCVDScene, const TSharedRef<SChaosVDMainTab>& InMainTabToolkitHost) override;
+
+protected:
+	
+	void DrawParticleVector(FPrimitiveDrawInterface* PDI, const FVector& StartLocation, const FVector& InVector, EChaosVDParticleDataVisualizationFlags VectorID, const FChaosVDParticleDataVisualizationContext& InVisualizationContext, float LineThickness);
+	void DrawVisualizationForParticleData(const UActorComponent* Component, FPrimitiveDrawInterface* PDI, const FSceneView* View, const FChaosVDParticleDataVisualizationContext& InVisualizationContext, const FChaosVDSceneParticle& InParticleInstance);
+};
+
+class FChaosVDTrimeshBVHVisualizer
+{
+public:
+	void Draw(FPrimitiveDrawInterface* PDI, const FChaosVDParticleDataVisualizationContext& InVisualizationContext, const FTransform& LocalToWorldTransform, const Chaos::FTriangleMeshImplicitObject& TriMesh) const;
+private:
+	void DrawBVH(FPrimitiveDrawInterface* PDI, const FTransform& LocalToWorldTransform, const Chaos::FTrimeshBVH& BVH, const int32 TargetLevel) const;
+	void DrawBVHLevel(FPrimitiveDrawInterface* PDI, const FTransform& LocalToWorldTransform, const Chaos::FTrimeshBVH& BVH, const int32 NodeIndex, const int32 CurrentLevel, const int32 TargetLevel) const;
+	void DrawAabb(FPrimitiveDrawInterface* PDI, const FTransform& LocalToWorldTransform, const Chaos::FAABB3& Aabb, const int32 ColorSeed) const;
+	Chaos::FAABB3 ToAabb(const Chaos::FAABBVectorized& AabbVectorized) const;
+};
diff --git a/Engine/Source/Runtime/Experimental/ChaosSolverEngine/Private/Chaos/ChaosVDEngineEditorBridge.cpp b/Engine/Source/Runtime/Experimental/ChaosSolverEngine/Private/Chaos/ChaosVDEngineEditorBridge.cpp
--- a/Engine/Source/Runtime/Experimental/ChaosSolverEngine/Private/Chaos/ChaosVDEngineEditorBridge.cpp
+++ b/Engine/Source/Runtime/Experimental/ChaosSolverEngine/Private/Chaos/ChaosVDEngineEditorBridge.cpp
@@ -47,13 +47,6 @@ void FChaosVDEngineEditorBridge::RemoveOnScreenRecordingMessage()
 
 void FChaosVDEngineEditorBridge::HandleCVDRecordingStarted()
 {
-	FTSTicker::GetCoreTicker().AddTicker(FTickerDelegate::CreateLambda([](float DeltaTime)
-	{
-		// We need to wait at least one frame to serialize the collision channel names to ensure the Archive header used for the whole session is traced
-		FChaosVDEngineEditorBridge::Get().SerializeCollisionChannelsNames();
-		return false;
-	}));
-
 	AddOnScreenRecordingMessage();
 
 	BroadcastSessionStatus(FApp::GetDeltaTime());
@@ -62,6 +55,11 @@ void FChaosVDEngineEditorBridge::HandleCVDRecordingStarted()
 	RecordingStatusUpdateHandle = FTSTicker::GetCoreTicker().AddTicker(FTickerDelegate::CreateRaw(this, &FChaosVDEngineEditorBridge::BroadcastSessionStatus), UpdateInterval);
 }
 
+void FChaosVDEngineEditorBridge::HandleCVDPostRecordingStarted()
+{
+	SerializeCollisionChannelsNames();
+}
+
 void FChaosVDEngineEditorBridge::HandleCVDRecordingStopped()
 {
 	RemoveOnScreenRecordingMessage();
@@ -87,6 +85,7 @@ void FChaosVDEngineEditorBridge::HandlePIEStarted(UGameInstance* GameInstance)
 	if (FChaosVDRuntimeModule::Get().IsRecording())
 	{
 		HandleCVDRecordingStarted();
+		HandleCVDPostRecordingStarted();
 	}
 }
 
@@ -159,6 +158,7 @@ void FChaosVDEngineEditorBridge::Initialize()
 	RemoteSessionsManager->Initialize();
 
 	RecordingStartedHandle = FChaosVDRuntimeModule::Get().RegisterRecordingStartedCallback(FChaosVDRecordingStateChangedDelegate::FDelegate::CreateRaw(this, &FChaosVDEngineEditorBridge::HandleCVDRecordingStarted));
+	PostRecordingStartedHandle = FChaosVDRuntimeModule::Get().RegisterPostRecordingStartedCallback(FChaosVDRecordingStateChangedDelegate::FDelegate::CreateRaw(this, &FChaosVDEngineEditorBridge::HandleCVDPostRecordingStarted));
 	RecordingStoppedHandle = FChaosVDRuntimeModule::Get().RegisterRecordingStopCallback(FChaosVDRecordingStateChangedDelegate::FDelegate::CreateRaw(this, &FChaosVDEngineEditorBridge::HandleCVDRecordingStopped));
 	RecordingStartFailedHandle = FChaosVDRuntimeModule::Get().RegisterRecordingStartFailedCallback(FChaosVDRecordingStartFailedDelegate::FDelegate::CreateRaw(this, &FChaosVDEngineEditorBridge::HandleCVDRecordingStartFailed));
 
@@ -172,6 +172,7 @@ void FChaosVDEngineEditorBridge::Initialize()
 	if (FChaosVDRuntimeModule::Get().IsRecording())
 	{
 		HandleCVDRecordingStarted();
+		HandleCVDPostRecordingStarted();
 	}
 }
 
@@ -185,6 +186,7 @@ void FChaosVDEngineEditorBridge::TearDown()
 	if (FChaosVDRuntimeModule::IsLoaded())
 	{
 		FChaosVDRuntimeModule::Get().RemoveRecordingStartedCallback(RecordingStartedHandle);
+		FChaosVDRuntimeModule::Get().RemovePostRecordingStartedCallback(PostRecordingStartedHandle);
 		FChaosVDRuntimeModule::Get().RemoveRecordingStopCallback(RecordingStoppedHandle);
 
 #if WITH_EDITOR
diff --git a/Engine/Source/Runtime/Experimental/ChaosSolverEngine/Public/Chaos/ChaosVDEngineEditorBridge.h b/Engine/Source/Runtime/Experimental/ChaosSolverEngine/Public/Chaos/ChaosVDEngineEditorBridge.h
--- a/Engine/Source/Runtime/Experimental/ChaosSolverEngine/Public/Chaos/ChaosVDEngineEditorBridge.h
+++ b/Engine/Source/Runtime/Experimental/ChaosSolverEngine/Public/Chaos/ChaosVDEngineEditorBridge.h
@@ -47,6 +47,7 @@ private:
 	void AddOnScreenRecordingMessage();
 	void RemoveOnScreenRecordingMessage();
 	void HandleCVDRecordingStarted();
+	void HandleCVDPostRecordingStarted();
 	void HandleCVDRecordingStopped();
 	void HandleCVDRecordingStartFailed(const FText& InFailureReason) const;
 	void HandlePIEStarted(UGameInstance* GameInstance);
@@ -58,6 +59,7 @@ private:
 	bool BroadcastSessionStatus(float DeltaTime);
 
 	FDelegateHandle RecordingStartedHandle;
+	FDelegateHandle PostRecordingStartedHandle;
 	FDelegateHandle RecordingStoppedHandle;
 	FDelegateHandle RecordingStartFailedHandle;
 	uint64 CVDRecordingMessageKey = 0;
diff --git a/Engine/Source/Runtime/Experimental/ChaosVDData/Public/ChaosVDRecording.h b/Engine/Source/Runtime/Experimental/ChaosVDData/Public/ChaosVDRecording.h
--- a/Engine/Source/Runtime/Experimental/ChaosVDData/Public/ChaosVDRecording.h
+++ b/Engine/Source/Runtime/Experimental/ChaosVDData/Public/ChaosVDRecording.h
@@ -682,6 +682,19 @@ struct CHAOSVDDATA_API FChaosVDRecording
 		return AvailableTrackIDForRemapping++;
 	}
 
+	FChaosVDCustomFrameData& GetCustomDataHandler()
+	{
+		return CustomData;
+	}
+	
+	const FChaosVDCustomFrameData& GetCustomDataHandler() const
+	{
+		return CustomData;
+	}
+
+private:
+	FChaosVDCustomFrameData CustomData;
+
 protected:
 
 	/** Adds an Implicit Object to the recording and takes ownership of it */
diff --git a/Engine/Source/Runtime/Experimental/ChaosVisualDebugger/Private/ChaosVDRuntimeModule.cpp b/Engine/Source/Runtime/Experimental/ChaosVisualDebugger/Private/ChaosVDRuntimeModule.cpp
--- a/Engine/Source/Runtime/Experimental/ChaosVisualDebugger/Private/ChaosVDRuntimeModule.cpp
+++ b/Engine/Source/Runtime/Experimental/ChaosVisualDebugger/Private/ChaosVDRuntimeModule.cpp
@@ -65,6 +65,7 @@ static FAutoConsoleVariable CVarChaosVDMaxTimeToWaitForDisconnect(
 	TEXT("Max time to wait after attempting to stop an active trace session. After that time has passed if we are still connected, CVD will continue and eventually error out."));
 
 FChaosVDRecordingStateChangedDelegate FChaosVDRuntimeModule::RecordingStartedDelegate = FChaosVDRecordingStateChangedDelegate();
+FChaosVDRecordingStateChangedDelegate FChaosVDRuntimeModule::PostRecordingStartedDelegate = FChaosVDRecordingStateChangedDelegate();
 FChaosVDRecordingStateChangedDelegate FChaosVDRuntimeModule::RecordingStopDelegate = FChaosVDRecordingStateChangedDelegate();
 FChaosVDRecordingStartFailedDelegate FChaosVDRuntimeModule::RecordingStartFailedDelegate = FChaosVDRecordingStartFailedDelegate();
 FChaosVDCaptureRequestDelegate FChaosVDRuntimeModule::PerformFullCaptureDelegate = FChaosVDCaptureRequestDelegate();
@@ -268,6 +269,11 @@ void FChaosVDRuntimeModule::StartRecording(TConstArrayView<FString> Args)
 			FMath::Clamp(ConfiguredTimeBetweenCaptures, MinAllowedTimeInSecondsBetweenCaptures, TNumericLimits<int32>::Max()));
 
 		RecordingTimerHandle = FTSTicker::GetCoreTicker().AddTicker(FTickerDelegate::CreateRaw(this, &FChaosVDRuntimeModule::RecordingTimerTick));
+
+		{
+			UE::TReadScopeLock ReadLock(DelegatesRWLock);
+			PostRecordingStartedDelegate.Broadcast();
+		}
 	}
 	else
 	{
diff --git a/Engine/Source/Runtime/Experimental/ChaosVisualDebugger/Public/ChaosVDRuntimeModule.h b/Engine/Source/Runtime/Experimental/ChaosVisualDebugger/Public/ChaosVDRuntimeModule.h
--- a/Engine/Source/Runtime/Experimental/ChaosVisualDebugger/Public/ChaosVDRuntimeModule.h
+++ b/Engine/Source/Runtime/Experimental/ChaosVisualDebugger/Public/ChaosVDRuntimeModule.h
@@ -60,6 +60,12 @@ public:
 		return RecordingStartedDelegate.Add(InCallback);
 	}
 
+	static FDelegateHandle RegisterPostRecordingStartedCallback(const FChaosVDRecordingStateChangedDelegate::FDelegate& InCallback)
+	{
+		UE::TWriteScopeLock WriteLock(DelegatesRWLock);
+		return PostRecordingStartedDelegate.Add(InCallback);
+	}
+
 	static FDelegateHandle RegisterRecordingStopCallback(const FChaosVDRecordingStateChangedDelegate::FDelegate& InCallback)
 	{
 		UE::TWriteScopeLock WriteLock(DelegatesRWLock);
@@ -83,6 +89,12 @@ public:
 		UE::TWriteScopeLock WriteLock(DelegatesRWLock);
 		return RecordingStartedDelegate.Remove(InDelegateToRemove);
 	}
+	
+	static bool RemovePostRecordingStartedCallback(const FDelegateHandle& InDelegateToRemove)
+	{
+		UE::TWriteScopeLock WriteLock(DelegatesRWLock);
+		return PostRecordingStartedDelegate.Remove(InDelegateToRemove);
+	}
 
 	static bool RemoveRecordingStopCallback(const FDelegateHandle& InDelegateToRemove)
 	{
@@ -148,6 +160,7 @@ private:
 	FTSTicker::FDelegateHandle RecordingTimerHandle;
 
 	static UE_API FChaosVDRecordingStateChangedDelegate RecordingStartedDelegate;
+	static UE_API FChaosVDRecordingStateChangedDelegate PostRecordingStartedDelegate;
 	static UE_API FChaosVDRecordingStateChangedDelegate RecordingStopDelegate;
 	static UE_API FChaosVDRecordingStartFailedDelegate RecordingStartFailedDelegate;
 	static UE_API FChaosVDCaptureRequestDelegate PerformFullCaptureDelegate;
