[ChaosVD] Fixed link error in the DSTM CVD extension. I file with a needed dll export was not include in the original CL by mistake.

diff --git a/Engine/Plugins/ChaosVD/Source/ChaosVD/Public/Widgets/SChaosVDMainTab.h b/Engine/Plugins/ChaosVD/Source/ChaosVD/Public/Widgets/SChaosVDMainTab.h
--- a/Engine/Plugins/ChaosVD/Source/ChaosVD/Public/Widgets/SChaosVDMainTab.h
+++ b/Engine/Plugins/ChaosVD/Source/ChaosVD/Public/Widgets/SChaosVDMainTab.h
@@ -98,8 +98,8 @@ public:
 	};
 	// END ITOOLKITHOST Interface
 
-	TSharedPtr<FComponentVisualizer> FindComponentVisualizer(UClass* ClassPtr);
-	TSharedPtr<FComponentVisualizer> FindComponentVisualizer(FName ClassName);
+	CHAOSVD_API TSharedPtr<FComponentVisualizer> FindComponentVisualizer(UClass* ClassPtr);
+	CHAOSVD_API TSharedPtr<FComponentVisualizer> FindComponentVisualizer(FName ClassName);
 
 	TConstArrayView<TSharedPtr<FComponentVisualizer>> GetAllComponentVisualizers()
 	{
