Fix an issue where the new viewport toolbar wasn't integrated into the ChaosVDPlaybackViewport.

Submitting on behalf taylor.hadden

diff --git a/Engine/Plugins/ChaosVD/Source/ChaosVD/Private/Widgets/SChaosVDPlaybackViewport.cpp b/Engine/Plugins/ChaosVD/Source/ChaosVD/Private/Widgets/SChaosVDPlaybackViewport.cpp
--- a/Engine/Plugins/ChaosVD/Source/ChaosVD/Private/Widgets/SChaosVDPlaybackViewport.cpp
+++ b/Engine/Plugins/ChaosVD/Source/ChaosVD/Private/Widgets/SChaosVDPlaybackViewport.cpp
@@ -104,12 +104,19 @@ void SChaosVDPlaybackViewport::Construct(const FArguments& InArgs, TWeakPtr<FCha
 			CVDEdMode->SetWorld(ScenePtr->GetUnderlyingWorld());
 		}
 	}
-
+	
+	TSharedPtr<SWidget> ViewportToolbar = BuildViewportToolbar();
+	
 	ChildSlot
 	[
 		// 3D Viewport
 		SNew(SVerticalBox)
 		+SVerticalBox::Slot()
+		.AutoHeight()
+		[
+			ViewportToolbar.IsValid() ? ViewportToolbar.ToSharedRef() : SNullWidget::NullWidget
+		]
+		+SVerticalBox::Slot()
 		.FillHeight(0.9f)
 		[
 			ViewportWidget.ToSharedRef()
