Update Chaos Visual Debugger to use the new viewport toolbars.

diff --git a/Engine/Plugins/ChaosVD/Source/ChaosVD/Private/Utils/ChaosVDMenus.cpp b/Engine/Plugins/ChaosVD/Source/ChaosVD/Private/Utils/ChaosVDMenus.cpp
new file mode 100644
--- /dev/null
+++ b/Engine/Plugins/ChaosVD/Source/ChaosVD/Private/Utils/ChaosVDMenus.cpp
@@ -0,0 +1,288 @@
+﻿// Copyright Epic Games, Inc. All Rights Reserved.
+
+#include "ChaosVDMenus.h"
+
+#include "ChaosVDCommands.h"
+#include "ChaosVDPlaybackViewportClient.h"
+#include "ToolMenus.h"
+#include "ShowFlagFilter.h"
+#include "ShowFlagMenuCommands.h"
+#include "ViewportToolbar/UnrealEdViewportToolbar.h"
+#include "Widgets/Layout/SBox.h"
+#include "Widgets/Layout/SBorder.h"
+#include "Widgets/SChaosVDPlaybackViewport.h"
+#include "Widgets/Input/SSpinBox.h"
+
+#define LOCTEXT_NAMESPACE "ChaosVisualDebugger"
+
+namespace Chaos::VisualDebugger::Menus
+{
+
+namespace Private
+{
+
+template <typename TNumericValue>
+TSharedRef<SWidget> GenerateSpinBoxMenuEntryWidget(const FText& ToolTipText, TNumericValue MinValue, TNumericValue MaxValue, typename SSpinBox<TNumericValue>::FOnValueChanged&& ValueChangedDelegate, const TAttribute<TNumericValue>&& ValueAttribute, const TAttribute<bool>&& EnabledAttribute)
+{
+	return SNew(SBox)
+		.HAlign(HAlign_Right)
+		[
+			SNew(SBox)
+			.Padding(FMargin(4.0f, 0.0f, 0.0f, 0.0f))
+			.WidthOverride(100.0f)
+			[
+				SNew ( SBorder )
+				.BorderImage(FAppStyle::Get().GetBrush("Menu.WidgetBorder"))
+				.Padding(FMargin(1.0f))
+				[
+					SNew(SSpinBox<TNumericValue>)
+					.Style(&FAppStyle::Get(), "Menu.SpinBox")
+					.ToolTipText(ToolTipText)
+					.MinValue(MinValue)
+					.MaxValue(MaxValue)
+					.Font(FAppStyle::GetFontStyle(TEXT("MenuItem.Font")))
+					.Value(ValueAttribute)
+					.OnValueChanged(ValueChangedDelegate)
+					.IsEnabled(EnabledAttribute)
+				]
+			]
+		];
+}
+
+FToolMenuEntry CreateFramerateOverrideEntry()
+{
+	return FToolMenuEntry::InitDynamicEntry("DynamicFramerateOverride", FNewToolMenuSectionDelegate::CreateLambda([](FToolMenuSection& Section)
+	{
+		TSharedPtr<SChaosVDPlaybackViewport> Viewport = StaticCastSharedPtr<SChaosVDPlaybackViewport>(UUnrealEdViewportToolbarContext::GetEditorViewport(Section.Context));
+		if (!Viewport)
+		{
+			return;
+		}
+	
+		TAttribute<int32> ValueAttribute = TAttribute<int32>::CreateSP(Viewport.Get(), &SChaosVDPlaybackViewport::GetCurrentTargetFrameRateOverride);
+		TAttribute<bool> EnabledAttribute = TAttribute<bool>::CreateSP(Viewport.Get(), &SChaosVDPlaybackViewport::IsUsingFrameRateOverride);
+		
+		SSpinBox<int32>::FOnValueChanged OnValueChanged = SSpinBox<int32>::FOnValueChanged::CreateSP(Viewport.Get(), &SChaosVDPlaybackViewport::SetCurrentTargetFrameRateOverride);
+		
+		constexpr int32 MinValue = 1;
+		constexpr int32 MaxValue = 1000;
+		
+		Section.AddEntry(FToolMenuEntry::InitWidget(
+			"FramerateOverride",
+			GenerateSpinBoxMenuEntryWidget(
+				LOCTEXT("FramerateOverrideTooltip", "Target framerate we should play the loaded recording at"),
+				MinValue,
+				MaxValue,
+				MoveTemp(OnValueChanged),
+				MoveTemp(ValueAttribute),
+				MoveTemp(EnabledAttribute)
+			),
+			LOCTEXT("FrameRateOverride", "Target Framerate")
+		));
+	}));
+}
+
+FToolMenuEntry CreateTrackingDistanceEntry()
+{
+	return FToolMenuEntry::InitDynamicEntry("DynamicTrackingDistance", FNewToolMenuSectionDelegate::CreateLambda([](FToolMenuSection& Section)
+	{
+		TSharedPtr<SEditorViewport> Viewport = UUnrealEdViewportToolbarContext::GetEditorViewport(Section.Context);
+		if (!Viewport)
+		{
+			return;
+		}
+		
+		TSharedPtr<FChaosVDPlaybackViewportClient> Client = StaticCastSharedPtr<FChaosVDPlaybackViewportClient>(Viewport->GetViewportClient());
+		if (!Client)
+		{
+			return;
+		}
+	
+		TAttribute<float> ValueAttribute = TAttribute<float>::CreateLambda([WeakClient = Client.ToWeakPtr()]
+		{
+			if (TSharedPtr<FChaosVDPlaybackViewportClient> Client = WeakClient.Pin())
+			{
+				return Client->GetAutoTrackingViewDistance();
+			}
+			return 0.0f;
+		});
+		
+		TAttribute<bool> EnabledAttribute = TAttribute<bool>::CreateLambda([WeakClient = Client.ToWeakPtr()]
+		{
+			if (TSharedPtr<FChaosVDPlaybackViewportClient> Client = WeakClient.Pin())
+			{
+				return Client->IsAutoTrackingSelectedObject();
+			}
+			return false;
+		});
+		
+		SSpinBox<float>::FOnValueChanged OnValueChanged = SSpinBox<float>::FOnValueChanged::CreateLambda([WeakClient = Client.ToWeakPtr()](float Value)
+		{
+			if (TSharedPtr<FChaosVDPlaybackViewportClient> Client = WeakClient.Pin())
+			{
+				Client->SetAutoTrackingViewDistance(Value);
+			}
+		});
+		
+		constexpr float MinValue = 1.0f;
+		constexpr float MaxValue = 100000.0f;
+		
+		Section.AddEntry(FToolMenuEntry::InitWidget(
+			"TrackingDistance",
+			GenerateSpinBoxMenuEntryWidget(
+				LOCTEXT("TrackingDistanceTooltip", "Distance from which we want to track the selected object"),
+				MinValue,
+				MaxValue,
+				MoveTemp(OnValueChanged),
+				MoveTemp(ValueAttribute),
+				MoveTemp(EnabledAttribute)
+			),
+			LOCTEXT("TrackingDistance", "Follow Distance")
+		));
+	}));
+}
+
+FToolMenuEntry CreateGoToLocationEntry()
+{
+	return FToolMenuEntry::InitDynamicEntry("DynamicGoToLocation", FNewToolMenuSectionDelegate::CreateLambda([](FToolMenuSection& Section)
+	{
+		TSharedPtr<SChaosVDPlaybackViewport> Viewport = StaticCastSharedPtr<SChaosVDPlaybackViewport>(UUnrealEdViewportToolbarContext::GetEditorViewport(Section.Context));
+		if (!Viewport)
+		{
+			return;
+		}
+		
+		TSharedRef<SWidget> Widget = SNew(SBox)
+			.HAlign(HAlign_Right)
+			[
+				SNew(SBox)
+				.Padding(FMargin(4.0f, 0.0f, 0.0f, 0.0f))
+				.WidthOverride(100.0f)
+				[
+					SNew (SBorder)
+					.BorderImage(FAppStyle::Get().GetBrush("Menu.WidgetBorder"))
+					.Padding(FMargin(1.0f))
+					[
+						SNew(SEditableText)
+						.ToolTipText(LOCTEXT("GoToLocationTooltip", "Location to teleport to."))
+						.Font(FAppStyle::GetFontStyle(TEXT("MenuItem.Font")))
+						.OnTextCommitted(FOnTextCommitted::CreateLambda([WeakViewport = Viewport.ToWeakPtr()](const FText& InLocationAsText, ETextCommit::Type Type)
+						{
+							TSharedPtr<SChaosVDPlaybackViewport> Viewport = WeakViewport.Pin();
+							if (!Viewport || Type != ETextCommit::OnEnter)
+							{
+								return;
+							}
+							
+							FVector Location;
+							Location.InitFromString(InLocationAsText.ToString());
+							Viewport->GoToLocation(Location);
+						}))
+					]
+				]
+			];
+		
+		Section.AddEntry(FToolMenuEntry::InitWidget(
+			"GoToLocation",
+			Widget,
+			LOCTEXT("GoToLocation", "Go to Location")
+		));
+	}));
+}
+
+} // End Private
+
+FToolMenuEntry CreateShowSubmenu()
+{
+	FToolMenuEntry Entry = UE::UnrealEd::CreateShowSubmenu(FNewToolMenuDelegate::CreateLambda([](UToolMenu* ShowMenu)
+	{
+		FToolMenuSection& CommonFlags = ShowMenu->AddSection("CommonViewportFlags", LOCTEXT("ToolbarCommonViewportFlags", "Common Show Flags"));
+		
+		constexpr bool bOpenSubMenuOnClick = false;
+		CommonFlags.AddSubMenu(
+			"CommonViewportFlags",
+			LOCTEXT("CommonShowFlagsMenuLabel", "Common Show Flags"),
+			LOCTEXT("CommonShowFlagsMenuToolTip", "Set of flags to enable/disable specific viewport features"),
+			FNewToolMenuDelegate::CreateLambda([](UToolMenu* Menu)
+			{
+				// Only include the flags that might be helpful.
+				static const FShowFlagFilter ShowFlagFilter = FShowFlagFilter(FShowFlagFilter::ExcludeAllFlagsByDefault)
+					.IncludeFlag(FEngineShowFlags::SF_AntiAliasing)
+					.IncludeFlag(FEngineShowFlags::SF_Grid)
+					.IncludeFlag(FEngineShowFlags::SF_Translucency)
+					.IncludeFlag(FEngineShowFlags::SF_MeshEdges)
+					.IncludeFlag(FEngineShowFlags::SF_HitProxies)
+					.IncludeFlag(FEngineShowFlags::SF_Fog)
+					.IncludeFlag(FEngineShowFlags::SF_Pivot);
+
+				FShowFlagMenuCommands::Get().BuildShowFlagsMenu(Menu, ShowFlagFilter);
+			}),
+			bOpenSubMenuOnClick,
+			FSlateIcon(FAppStyle::Get().GetStyleSetName(), TEXT("Icons.Toolbar.Settings"))
+		);
+	}));
+	
+	// Re-enable the "show" label
+	Entry.ToolBarData.LabelOverride = TAttribute<FText>();
+	
+	return Entry;
+}
+
+FToolMenuEntry CreateSettingsSubmenu()
+{
+	FToolMenuEntry Entry = FToolMenuEntry::InitSubMenu("Settings",
+		LOCTEXT("SettingsMenuLabel", "Settings"),
+		LOCTEXT("SettingsMenuTooltip", "Additional settings for this viewport."),
+		FNewToolMenuDelegate::CreateLambda([](UToolMenu* Menu)
+		{
+			FToolMenuSection& OptionsSection = Menu->AddSection("ViewportOptions", LOCTEXT("ViewportOptionsMenuHeader", "Viewport Options"));
+			{
+				OptionsSection.AddSubMenu(
+					"PlaybackFramerate",
+					LOCTEXT("FrameRateOptionsMenuLabel", "Playback Framerate"),
+					LOCTEXT("FrameRateOptionsMenuToolTip", "Options that control how CVD plays a recording."),
+					FNewToolMenuDelegate::CreateLambda([](UToolMenu* Menu)
+					{
+						FToolMenuSection& Section = Menu->AddSection("Framerate", LOCTEXT("FramerateMenuHeader", "Framerate"));
+						Section.AddMenuEntry(FChaosVDCommands::Get().OverridePlaybackFrameRate);
+						Section.AddEntry(Private::CreateFramerateOverrideEntry());
+					}),
+					false,
+					FSlateIcon(FAppStyle::Get().GetStyleSetName(), TEXT("EditorViewport.ToggleFPS"))
+				);
+			
+				OptionsSection.AddSubMenu(
+					"ObjectTracking",
+					LOCTEXT("ObjectTrackingMenuLabel", "Object Tracking"),
+					LOCTEXT("ObjectTrackingMenuToolTip", "Options that control how objects are tracked in the scene by the camera."),
+					FNewToolMenuDelegate::CreateLambda([](UToolMenu* Menu)
+					{
+						FToolMenuSection& Section = Menu->AddSection("ObjectTracking", LOCTEXT("ObjectTrackingMenuHeader", "Object Tracking"));
+						Section.AddMenuEntry(FChaosVDCommands::Get().ToggleFollowSelectedObject);
+						Section.AddEntry(Private::CreateTrackingDistanceEntry());
+					}),
+					false,
+					FSlateIcon(FAppStyle::Get().GetStyleSetName(), TEXT("AnimViewportMenu.CameraFollow.Small"))
+				);
+				
+				OptionsSection.AddSeparator(NAME_None);
+				
+				OptionsSection.AddMenuEntry(FChaosVDCommands::Get().AllowTranslucentSelection);
+			}
+			
+			FToolMenuSection& UtilsSection = Menu->AddSection("Utilities", LOCTEXT("ViewportUtilMenuHeader", "Utils"));
+			{
+				UtilsSection.AddEntry(Private::CreateGoToLocationEntry());
+			}
+		})
+	);
+	
+	Entry.ToolBarData.LabelOverride = FText::GetEmpty();
+	Entry.Icon = FSlateIcon(FAppStyle::Get().GetStyleSetName(), "EditorViewportToolBar.OptionsDropdown");
+	
+	return Entry;
+}
+
+}
+
+#undef LOCTEXT_NAMESPACE
diff --git a/Engine/Plugins/ChaosVD/Source/ChaosVD/Private/Utils/ChaosVDMenus.h b/Engine/Plugins/ChaosVD/Source/ChaosVD/Private/Utils/ChaosVDMenus.h
new file mode 100644
--- /dev/null
+++ b/Engine/Plugins/ChaosVD/Source/ChaosVD/Private/Utils/ChaosVDMenus.h
@@ -0,0 +1,18 @@
+﻿// Copyright Epic Games, Inc. All Rights Reserved.
+
+#pragma once
+
+#include "UObject/NameTypes.h"
+
+struct FToolMenuEntry;
+
+namespace Chaos::VisualDebugger::Menus
+{
+
+const FName ShowMenuName = "ChaosVDViewportToolbarBase.Show";
+
+FToolMenuEntry CreateShowSubmenu();
+
+FToolMenuEntry CreateSettingsSubmenu();
+
+}
diff --git a/Engine/Plugins/ChaosVD/Source/ChaosVD/Private/Visualizers/ChaosVDCharacterGroundConstraintsDataComponentVisualizer.cpp b/Engine/Plugins/ChaosVD/Source/ChaosVD/Private/Visualizers/ChaosVDCharacterGroundConstraintsDataComponentVisualizer.cpp
--- a/Engine/Plugins/ChaosVD/Source/ChaosVD/Private/Visualizers/ChaosVDCharacterGroundConstraintsDataComponentVisualizer.cpp
+++ b/Engine/Plugins/ChaosVD/Source/ChaosVD/Private/Visualizers/ChaosVDCharacterGroundConstraintsDataComponentVisualizer.cpp
@@ -15,8 +15,8 @@
 #include "ToolMenus.h"
 #include "ToolMenuEntry.h"
 #include "ToolMenuSection.h"
+#include "Utils/ChaosVDMenus.h"
 #include "Visualizers/ChaosVDDebugDrawUtils.h"
-#include "Widgets/SChaosVDViewportToolbar.h"
 
 #define LOCTEXT_NAMESPACE "ChaosVisualDebugger"
 
@@ -43,7 +43,7 @@ void FChaosVDCharacterGroundConstraintDataComponentVisualizer::RegisterVisualize
 	FText SettingsMenuLabel = LOCTEXT("CharacterConstraintDataVisualizationMenuLabel", "Character Ground Constraints Visualization Settings");
 	FText SettingsMenuTooltip = LOCTEXT("CharacterConstraintDataVisualizationMenuToolTip", "Options to change how the recorded Character Constraints data is debug drawn");
 	
-	CreateGenericVisualizerMenu<UChaosVDCharacterConstraintsVisualizationSettings, EChaosVDCharacterGroundConstraintDataVisualizationFlags>(SChaosVDViewportToolbar::ShowMenuName, MenuSection, MenuSectionLabel, FlagsMenuLabel, FlagsMenuTooltip, FlagsMenuIcon, SettingsMenuLabel, SettingsMenuTooltip);
+	CreateGenericVisualizerMenu<UChaosVDCharacterConstraintsVisualizationSettings, EChaosVDCharacterGroundConstraintDataVisualizationFlags>(Chaos::VisualDebugger::Menus::ShowMenuName, MenuSection, MenuSectionLabel, FlagsMenuLabel, FlagsMenuTooltip, FlagsMenuIcon, SettingsMenuLabel, SettingsMenuTooltip);
 }
 
 void FChaosVDCharacterGroundConstraintDataComponentVisualizer::DrawVisualization(const UActorComponent* Component, const FSceneView* View, FPrimitiveDrawInterface* PDI)
diff --git a/Engine/Plugins/ChaosVD/Source/ChaosVD/Private/Visualizers/ChaosVDJointConstraintsDataComponentVisualizer.cpp b/Engine/Plugins/ChaosVD/Source/ChaosVD/Private/Visualizers/ChaosVDJointConstraintsDataComponentVisualizer.cpp
--- a/Engine/Plugins/ChaosVD/Source/ChaosVD/Private/Visualizers/ChaosVDJointConstraintsDataComponentVisualizer.cpp
+++ b/Engine/Plugins/ChaosVD/Source/ChaosVD/Private/Visualizers/ChaosVDJointConstraintsDataComponentVisualizer.cpp
@@ -17,8 +17,8 @@
 #include "ToolMenus.h"
 #include "ToolMenuEntry.h"
 #include "ToolMenuSection.h"
+#include "Utils/ChaosVDMenus.h"
 #include "Visualizers/ChaosVDDebugDrawUtils.h"
-#include "Widgets/SChaosVDViewportToolbar.h"
 
 #define LOCTEXT_NAMESPACE "ChaosVisualDebugger"
 
@@ -122,7 +122,7 @@ void FChaosVDJointConstraintsDataComponentVisualizer::RegisterVisualizerMenus()
 	FText SettingsMenuLabel = LOCTEXT("JointConstraintDataVisualizationMenuLabel", "Joint Constraint Visualization Settings");
 	FText SettingsMenuTooltip = LOCTEXT("JointConstraintDataVisualizationMenuToolTip", "Options to change how the recorded joint constraint data is debug drawn");
 	
-	CreateGenericVisualizerMenu<UChaosVDJointConstraintsVisualizationSettings, EChaosVDJointsDataVisualizationFlags>(SChaosVDViewportToolbar::ShowMenuName, MenuSection, MenuSectionLabel, FlagsMenuLabel, FlagsMenuTooltip, FlagsMenuIcon, SettingsMenuLabel, SettingsMenuTooltip);
+	CreateGenericVisualizerMenu<UChaosVDJointConstraintsVisualizationSettings, EChaosVDJointsDataVisualizationFlags>(Chaos::VisualDebugger::Menus::ShowMenuName, MenuSection, MenuSectionLabel, FlagsMenuLabel, FlagsMenuTooltip, FlagsMenuIcon, SettingsMenuLabel, SettingsMenuTooltip);
 }
 
 void FChaosVDJointConstraintsDataComponentVisualizer::DrawVisualization(const UActorComponent* Component, const FSceneView* View, FPrimitiveDrawInterface* PDI)
diff --git a/Engine/Plugins/ChaosVD/Source/ChaosVD/Private/Visualizers/ChaosVDParticleDataComponentVisualizer.cpp b/Engine/Plugins/ChaosVD/Source/ChaosVD/Private/Visualizers/ChaosVDParticleDataComponentVisualizer.cpp
--- a/Engine/Plugins/ChaosVD/Source/ChaosVD/Private/Visualizers/ChaosVDParticleDataComponentVisualizer.cpp
+++ b/Engine/Plugins/ChaosVD/Source/ChaosVD/Private/Visualizers/ChaosVDParticleDataComponentVisualizer.cpp
@@ -17,10 +17,10 @@
 #include "Utils/ChaosVDUserInterfaceUtils.h"
 #include "Visualizers/ChaosVDDebugDrawUtils.h"
 #include "Widgets/SChaosVDEnumFlagsMenu.h"
-#include "Widgets/SChaosVDViewportToolbar.h"
 #include "Chaos/ImplicitObject.h"
 #include "Chaos/ImplicitObjectScaled.h"
 #include "Chaos/TriangleMeshImplicitObject.h"
+#include "Utils/ChaosVDMenus.h"
 
 #define LOCTEXT_NAMESPACE "ChaosVisualDebugger"
 
@@ -69,7 +69,7 @@ void FChaosVDParticleDataComponentVisualizer::RegisterVisualizerMenus()
 		return;
 	}
 
-	if (UToolMenu* Menu = ToolMenus->ExtendMenu(SChaosVDViewportToolbar::ShowMenuName))
+	if (UToolMenu* Menu = ToolMenus->ExtendMenu(Chaos::VisualDebugger::Menus::ShowMenuName))
 	{
 		FToolMenuSection& Section = Menu->AddSection("ParticleVisualization.Show", LOCTEXT("ParticleVisualizationShowMenuLabel", "Particle Visualization"));
 
diff --git a/Engine/Plugins/ChaosVD/Source/ChaosVD/Private/Visualizers/ChaosVDSceneQueryDataComponentVisualizer.cpp b/Engine/Plugins/ChaosVD/Source/ChaosVD/Private/Visualizers/ChaosVDSceneQueryDataComponentVisualizer.cpp
--- a/Engine/Plugins/ChaosVD/Source/ChaosVD/Private/Visualizers/ChaosVDSceneQueryDataComponentVisualizer.cpp
+++ b/Engine/Plugins/ChaosVD/Source/ChaosVD/Private/Visualizers/ChaosVDSceneQueryDataComponentVisualizer.cpp
@@ -17,8 +17,8 @@
 #include "ToolMenus.h"
 #include "ToolMenuEntry.h"
 #include "ToolMenuSection.h"
+#include "Utils/ChaosVDMenus.h"
 #include "Visualizers/ChaosVDDebugDrawUtils.h"
-#include "Widgets/SChaosVDViewportToolbar.h"
 
 class UChaosVDCoreSettings;
 
@@ -42,7 +42,7 @@ void FChaosVDSceneQueryDataComponentVisualizer::RegisterVisualizerMenus()
 	FText SettingsMenuLabel = LOCTEXT("SceneQuerySettingsMenuLabel", "Scene Query Visualization Settings");
 	FText SettingsMenuTooltip = LOCTEXT("SceneQuerySettingsMenuToolTip", "Options to change how the recorded scene query data is debug drawn");
 	
-	CreateGenericVisualizerMenu<UChaosVDSceneQueriesVisualizationSettings, EChaosVDSceneQueryVisualizationFlags>(SChaosVDViewportToolbar::ShowMenuName, MenuSection, MenuSectionLabel, FlagsMenuLabel, FlagsMenuTooltip, FlagsMenuIcon, SettingsMenuLabel, SettingsMenuTooltip);
+	CreateGenericVisualizerMenu<UChaosVDSceneQueriesVisualizationSettings, EChaosVDSceneQueryVisualizationFlags>(Chaos::VisualDebugger::Menus::ShowMenuName, MenuSection, MenuSectionLabel, FlagsMenuLabel, FlagsMenuTooltip, FlagsMenuIcon, SettingsMenuLabel, SettingsMenuTooltip);
 }
 
 bool FChaosVDSceneQueryDataComponentVisualizer::CanHandleClick(const HChaosVDComponentVisProxy& VisProxy)
diff --git a/Engine/Plugins/ChaosVD/Source/ChaosVD/Private/Visualizers/ChaosVDSolverCollisionDataComponentVisualizer.cpp b/Engine/Plugins/ChaosVD/Source/ChaosVD/Private/Visualizers/ChaosVDSolverCollisionDataComponentVisualizer.cpp
--- a/Engine/Plugins/ChaosVD/Source/ChaosVD/Private/Visualizers/ChaosVDSolverCollisionDataComponentVisualizer.cpp
+++ b/Engine/Plugins/ChaosVD/Source/ChaosVD/Private/Visualizers/ChaosVDSolverCollisionDataComponentVisualizer.cpp
@@ -14,9 +14,9 @@
 #include "ToolMenuEntry.h"
 #include "ToolMenuSection.h"
 #include "Components/ChaosVDParticleDataComponent.h"
+#include "Utils/ChaosVDMenus.h"
 #include "Visualizers/ChaosVDDebugDrawUtils.h"
 #include "Widgets/SChaosVDCollisionDataInspector.h"
-#include "Widgets/SChaosVDViewportToolbar.h"
 
 
 #define LOCTEXT_NAMESPACE "ChaosVisualDebugger"
@@ -43,7 +43,7 @@ void FChaosVDSolverCollisionDataComponentVisualizer::RegisterVisualizerMenus()
 	FText SettingsMenuLabel = LOCTEXT("CollisionVisualizationSettingsMenuLabel", "Collision Visualization Settings");
 	FText SettingsMenuTooltip = LOCTEXT("CollisionVisualizationSettingsMenuToolTip", "Options to change how the recorded collision Data is debug drawn");
 	
-	CreateGenericVisualizerMenu<UChaosVDCollisionDataVisualizationSettings, EChaosVDCollisionVisualizationFlags>(SChaosVDViewportToolbar::ShowMenuName, MenuSection, MenuSectionLabel, FlagsMenuLabel, FlagsMenuTooltip, FlagsMenuIcon, SettingsMenuLabel, SettingsMenuTooltip);
+	CreateGenericVisualizerMenu<UChaosVDCollisionDataVisualizationSettings, EChaosVDCollisionVisualizationFlags>(Chaos::VisualDebugger::Menus::ShowMenuName, MenuSection, MenuSectionLabel, FlagsMenuLabel, FlagsMenuTooltip, FlagsMenuIcon, SettingsMenuLabel, SettingsMenuTooltip);
 }
 
 bool FChaosVDSolverCollisionDataComponentVisualizer::ShouldShowForSelectedSubcomponents(const UActorComponent* Component)
diff --git a/Engine/Plugins/ChaosVD/Source/ChaosVD/Private/Widgets/SChaosVDEditorViewportViewMenu.cpp b/Engine/Plugins/ChaosVD/Source/ChaosVD/Private/Widgets/SChaosVDEditorViewportViewMenu.cpp
deleted file mode 100644
diff --git a/Engine/Plugins/ChaosVD/Source/ChaosVD/Private/Widgets/SChaosVDEditorViewportViewMenu.h b/Engine/Plugins/ChaosVD/Source/ChaosVD/Private/Widgets/SChaosVDEditorViewportViewMenu.h
deleted file mode 100644
diff --git a/Engine/Plugins/ChaosVD/Source/ChaosVD/Private/Widgets/SChaosVDPlaybackViewport.cpp b/Engine/Plugins/ChaosVD/Source/ChaosVD/Private/Widgets/SChaosVDPlaybackViewport.cpp
--- a/Engine/Plugins/ChaosVD/Source/ChaosVD/Private/Widgets/SChaosVDPlaybackViewport.cpp
+++ b/Engine/Plugins/ChaosVD/Source/ChaosVD/Private/Widgets/SChaosVDPlaybackViewport.cpp
@@ -19,11 +19,13 @@
 #include "Framework/Application/SlateApplication.h"
 #include "TEDS/ChaosVDSelectionInterface.h"
 #include "TEDS/ChaosVDStructTypedElementData.h"
+#include "Utils/ChaosVDMenus.h"
+#include "ViewportToolbar/UnrealEdViewportToolbar.h"
 #include "Widgets/SChaosVDTimelineWidget.h"
-#include "Widgets/SChaosVDViewportToolbar.h"
 #include "Widgets/Input/SCheckBox.h"
 #include "Widgets/Layout/SSeparator.h"
 #include "Widgets/Text/STextBlock.h"
+#include "ToolMenus.h"
 
 #define LOCTEXT_NAMESPACE "ChaosVisualDebugger"
 
@@ -478,19 +480,50 @@ TSharedRef<FEditorViewportClient> SChaosVDPlaybackViewport::MakeEditorViewportCl
 
 TSharedPtr<SWidget> SChaosVDPlaybackViewport::BuildViewportToolbar()
 {
-	// Build our toolbar level toolbar
-	TSharedRef< SChaosVDViewportToolbar > ToolBar = SNew(SChaosVDViewportToolbar, SharedThis(this));
+	const FName ToolbarName = "ChaosVDViewportToolbarBase";
+	
+	if (!UToolMenus::Get()->IsMenuRegistered(ToolbarName))
+	{
+		UToolMenu* Toolbar = UToolMenus::Get()->RegisterMenu(ToolbarName, NAME_None, EMultiBoxType::SlimHorizontalToolBar);
+		Toolbar->StyleName = "ViewportToolbar";
+		
+		{
+			FToolMenuSection& LeftSection = Toolbar->AddSection("Left");
+			LeftSection.AddEntry(Chaos::VisualDebugger::Menus::CreateSettingsSubmenu());
+		}
+		
+		{
+			FToolMenuSection& RightSection = Toolbar->AddSection("Right");
+			RightSection.Alignment = EToolMenuSectionAlign::Last;
+			
+			RightSection.AddEntry(UE::UnrealEd::CreateCameraSubmenu(UE::UnrealEd::FViewportCameraMenuOptions().ShowAll()));
+			RightSection.AddEntry(UE::UnrealEd::CreateViewModesSubmenu());
+			RightSection.AddEntry(Chaos::VisualDebugger::Menus::CreateShowSubmenu());
+		}
+	}
 
-	return 
-		SNew(SVerticalBox)
-		.Visibility( EVisibility::SelfHitTestInvisible )
-		+SVerticalBox::Slot()
-		.AutoHeight()
-		.Padding(0, 1.0f, 0, 0)
-		.VAlign(VAlign_Top)
-		[
-			ToolBar
-		];
+	FToolMenuContext Context;
+	{
+		UUnrealEdViewportToolbarContext* ViewportContext = UE::UnrealEd::CreateViewportToolbarDefaultContext(SharedThis(this));
+		ViewportContext->IsViewModeSupported.BindLambda([](EViewModeIndex ViewMode)
+		{
+			switch (ViewMode)
+			{
+			case VMI_Lit:
+			case VMI_Unlit:
+			case VMI_Lit_Wireframe:
+			case VMI_BrushWireframe:
+				return true;
+			default:
+				return false;
+			}
+		});
+		
+		Context.AddObject(ViewportContext);
+		Context.AppendCommandList(GetCommandList());
+	}
+
+	return UToolMenus::Get()->GenerateWidget(ToolbarName, Context);
 }
 
 void SChaosVDPlaybackViewport::HandlePlaybackControllerDataUpdated(TWeakPtr<FChaosVDPlaybackController> InController)
diff --git a/Engine/Plugins/ChaosVD/Source/ChaosVD/Private/Widgets/SChaosVDViewportToolbar.cpp b/Engine/Plugins/ChaosVD/Source/ChaosVD/Private/Widgets/SChaosVDViewportToolbar.cpp
deleted file mode 100644
diff --git a/Engine/Plugins/ChaosVD/Source/ChaosVD/Private/Widgets/SChaosVDViewportToolbar.h b/Engine/Plugins/ChaosVD/Source/ChaosVD/Private/Widgets/SChaosVDViewportToolbar.h
deleted file mode 100644
diff --git a/Engine/Source/Editor/UnrealEd/Private/ViewportToolbar/UnrealEdViewportToolbarContext.cpp b/Engine/Source/Editor/UnrealEd/Private/ViewportToolbar/UnrealEdViewportToolbarContext.cpp
--- a/Engine/Source/Editor/UnrealEd/Private/ViewportToolbar/UnrealEdViewportToolbarContext.cpp
+++ b/Engine/Source/Editor/UnrealEd/Private/ViewportToolbar/UnrealEdViewportToolbarContext.cpp
@@ -5,9 +5,19 @@
 #include "SEditorViewport.h"
 #include "Settings/LevelEditorViewportSettings.h"
 #include "ViewportToolbar/UnrealEdViewportToolbar.h"
+#include "ToolMenuContext.h"
 
 #define LOCTEXT_NAMESPACE "UnrealEdViewportToolbar"
 
+TSharedPtr<SEditorViewport> UUnrealEdViewportToolbarContext::GetEditorViewport(const FToolMenuContext& Context)
+{
+	if (const UUnrealEdViewportToolbarContext* ContextObject = Context.FindContext<UUnrealEdViewportToolbarContext>())
+	{
+		return ContextObject->Viewport.Pin();
+	}
+	return nullptr;
+}
+
 TSharedPtr<IPreviewProfileController> UUnrealEdViewportToolbarContext::GetPreviewProfileController() const
 {
 	if (const TSharedPtr<SEditorViewport> EditorViewport = Viewport.Pin())
diff --git a/Engine/Source/Editor/UnrealEd/Public/ViewportToolbar/UnrealEdViewportToolbarContext.h b/Engine/Source/Editor/UnrealEd/Public/ViewportToolbar/UnrealEdViewportToolbarContext.h
--- a/Engine/Source/Editor/UnrealEd/Public/ViewportToolbar/UnrealEdViewportToolbarContext.h
+++ b/Engine/Source/Editor/UnrealEd/Public/ViewportToolbar/UnrealEdViewportToolbarContext.h
@@ -27,6 +27,7 @@ DECLARE_DELEGATE_RetVal_OneParam(bool, DoesViewModeMenuShowSectionDelegate, EHid
 }
 
 class SEditorViewport;
+struct FToolMenuContext;
 
 UCLASS(MinimalAPI)
 class UUnrealEdViewportToolbarContext : public UObject
@@ -58,6 +59,8 @@ public:
 	 */
 	FTabId PreviewSettingsTabId;
 	
+	static UE_API TSharedPtr<SEditorViewport> GetEditorViewport(const FToolMenuContext& Context);
+	
 	UE_API virtual TSharedPtr<IPreviewProfileController> GetPreviewProfileController() const;
 	
 	bool bShowSurfaceSnap = true;
