[ChaosVD] Fix for an issue where game thread related data, like Scene queries, are not properly drawn and listed when you inspect multiple CVD files at the same time (this is part of the new workflow we shipped in 5.6 to debug packaged client and server builds).

The Problem :

Due to a tech limitation we are removing in 5.7, Game thread data has its own data track, and most of CVD's framework assumes that there only exists one game thread data track.

If you load multiple files, every other data track type have its their IDs remapped to remove collisions and they just work, but the additional game thread data is an issue because we only have one track available.

The workaround for 5.6 was to just inject the game thread data into a the generic data track type currently used for solver data, and implemented a CVD data component that takes care of re-routing the additional game thread data to the appropriate "New Game Thread data" callbacks.

That worked well, but we just found an issue where some of the systems that need to clear the current visualized data if there is no data for a specific solver or world when we scrub to a new frame, it will incorrectly clear valid data because now we call these "new game thread data" callbacks are executed more than once per frame (one per frame, plus once more per new inspected file).

The fix:
There is no simple way to fix this without modifying existing public class layouts, as what we need is to store data to differentiate the source of each game thread data these callbacks receive, and made the CVD's playback system know there could be multiple sources.

But as a workaround, this CL does the follwoing.

- Uses the generic data attachment system for frames, and adds to any simulated game thread frame, the context of which solver track ids are currently available.
- The data component doing the re-routing of this GT data, now just executes the "new game frame" callbacks on the data container actors for whose ID we have data, instead of all of them and let each component decide what to do.
- Given that other tracks can also execute the normal "New Game Frame" callbacks, there is no guarantee that the Data component doing the re-routing is the last being called, therefore data updated by it, can be overwritten incorrectly. To workaround this now we delay the processing the data we need to re-route to after the CVD scene update processing was completed.

Notes: This fix is not perfect. Storing the available list of solvers id at the beginning of each game thread frame event, it means we might miss some data for the first recorded solver frame as these IDs will be populated afterwards.
Also, processing the data in the data re-route component, after the scene update callback is executed means it is not guaranteed that the viewport will fully reflect the change until you interact with it, but in practice I was not able to make it happen locally. Overall this situation is better than the current state where the visualization is misleading.

The good thing about this change is that it will not add overhead or change the behavior inspecting single CVD files, which is the most common use case.

diff --git a/Engine/Plugins/ChaosVD/Source/ChaosVD/Private/Actors/ChaosVDSolverInfoActor.cpp b/Engine/Plugins/ChaosVD/Source/ChaosVD/Private/Actors/ChaosVDSolverInfoActor.cpp
--- a/Engine/Plugins/ChaosVD/Source/ChaosVD/Private/Actors/ChaosVDSolverInfoActor.cpp
+++ b/Engine/Plugins/ChaosVD/Source/ChaosVD/Private/Actors/ChaosVDSolverInfoActor.cpp
@@ -68,7 +68,10 @@ void AChaosVDSolverInfoActor::UpdateFromNewGameFrameData(const FChaosVDGameFrame
 
 void AChaosVDSolverInfoActor::UpdateFromNewSolverFrameData(const FChaosVDSolverFrameData& InSolverFrameData)
 {
-	FindAndUpdateFromCorrectGameFrameData(InSolverFrameData);
+	if (!EnumHasAnyFlags(InSolverFrameData.GetAttributes(), EChaosVDSolverFrameAttributes::HasGTDataToReRoute))
+	{
+		FindAndUpdateFromCorrectGameFrameData(InSolverFrameData);
+	}
 
 	SetSimulationTransform(InSolverFrameData.SimulationTransform);
 	Super::UpdateFromNewSolverFrameData(InSolverFrameData);
diff --git a/Engine/Plugins/ChaosVD/Source/ChaosVD/Private/ChaosVDPlaybackController.cpp b/Engine/Plugins/ChaosVD/Source/ChaosVD/Private/ChaosVDPlaybackController.cpp
--- a/Engine/Plugins/ChaosVD/Source/ChaosVD/Private/ChaosVDPlaybackController.cpp
+++ b/Engine/Plugins/ChaosVD/Source/ChaosVD/Private/ChaosVDPlaybackController.cpp
@@ -1414,7 +1414,12 @@ void FChaosVDPlaybackController::UpdateSolverTracksData()
 		{
 			SolverTrackInfo = CreateTrackInfo();
 			// Solver frame tracks only holding GT Data, these should not be shown in the timeline controls UI
-			SolverTrackInfo->bCanShowTrackControls = !SolverIDPair.Value.IsEmpty() ? !EnumHasAnyFlags(SolverIDPair.Value[0].GetAttributes(), EChaosVDSolverFrameAttributes::HasGTDataToReRoute) : true;
+			// TODO: The following is being commented out because although the comment above is correct, we can end up with two or more source files whose game thread
+			// ticked at different rates. The main game frames timeline controls are populated from the first loaded file, therefore this means that if that recording had its GT ticking slower
+			// , when you scrub forward, the other GT tracks might skip multiple frames (correctly) but you will miss data without having anything in the UI telling you what happened.
+			// This will be addressed when we fully remove the game thread data tracks as unique type in the next release (UE-277464), but for now it means we need to keep these track timeline controls enabled
+			//SolverTrackInfo->bCanShowTrackControls = !SolverIDPair.Value.IsEmpty() ? !EnumHasAnyFlags(SolverIDPair.Value[0].GetAttributes(), EChaosVDSolverFrameAttributes::HasGTDataToReRoute) : true;
+			SolverTrackInfo->bCanShowTrackControls = true;
 			SolverTrackInfo->CurrentFrame = 0;
 			SolverTrackInfo->CurrentStage = 0;
 		}
diff --git a/Engine/Plugins/ChaosVD/Source/ChaosVD/Private/Components/ChaosVDAdditionalGTDataRouterComponent.cpp b/Engine/Plugins/ChaosVD/Source/ChaosVD/Private/Components/ChaosVDAdditionalGTDataRouterComponent.cpp
--- a/Engine/Plugins/ChaosVD/Source/ChaosVD/Private/Components/ChaosVDAdditionalGTDataRouterComponent.cpp
+++ b/Engine/Plugins/ChaosVD/Source/ChaosVD/Private/Components/ChaosVDAdditionalGTDataRouterComponent.cpp
@@ -19,11 +19,7 @@ void UChaosVDAdditionalGTDataRouterComponent::UpdateFromSolverFrameData(const FC
 		{
 			if (TSharedPtr<FChaosVDGameFrameDataWrapper> GTFrameDataWrapper = InSolverFrameData.GetCustomData().GetData<FChaosVDGameFrameDataWrapper>())
 			{
-				for (const TObjectPtr<AChaosVDDataContainerBaseActor>& DataContainerActor : CVDScene->GetDataContainerActorsView())
-				{
-					AChaosVDDataContainerBaseActor::FScopedGameFrameDataReRouting ScopedGTDataUpdate(DataContainerActor.Get());
-					DataContainerActor->UpdateFromNewGameFrameData(*GTFrameDataWrapper->FrameData);
-				}
+				CVDScene->OnSceneUpdated().AddUObject(this, &UChaosVDAdditionalGTDataRouterComponent::ApplyDelayedSolverFrameDataUpdate, GTFrameDataWrapper);
 			}
 		}
 	}
@@ -31,4 +27,52 @@ void UChaosVDAdditionalGTDataRouterComponent::UpdateFromSolverFrameData(const FC
 
 void UChaosVDAdditionalGTDataRouterComponent::ClearData()
 {
+	AChaosVDSolverInfoActor* Owner = Cast<AChaosVDSolverInfoActor>(GetOwner());
+	TSharedPtr<FChaosVDScene> CVDScene = Owner ? Owner->GetScene().Pin() : nullptr;
+	if (!CVDScene)
+	{
+		return;
+	}
+
+	CVDScene->OnSceneUpdated().RemoveAll(this);
+}
+
+void UChaosVDAdditionalGTDataRouterComponent::ApplyDelayedSolverFrameDataUpdate(TSharedPtr<FChaosVDGameFrameDataWrapper> PendingFrameUpdate)
+{
+	AChaosVDSolverInfoActor* Owner = Cast<AChaosVDSolverInfoActor>(GetOwner());
+	TSharedPtr<FChaosVDScene> CVDScene = Owner ? Owner->GetScene().Pin() : nullptr;
+	if (!CVDScene)
+	{
+		return;
+	}
+
+	CVDScene->OnSceneUpdated().RemoveAll(this);
+
+	if (PendingFrameUpdate)
+	{
+		if (TSharedPtr<FChaosVDGameFrameDataWrapperContext> FrameDataContext = PendingFrameUpdate->FrameData->GetCustomDataHandler().GetData<FChaosVDGameFrameDataWrapperContext>())
+		{
+			bool bHasNewData = false;
+			for (int32 SupportedSolverID : FrameDataContext->SupportedSolverIDs)
+			{
+				if (AChaosVDSolverInfoActor* SolverInfoActor = CVDScene->GetSolverInfoActor(SupportedSolverID))
+				{
+					AChaosVDDataContainerBaseActor::FScopedGameFrameDataReRouting ScopedGTDataUpdate(SolverInfoActor);
+					SolverInfoActor->UpdateFromNewGameFrameData(*PendingFrameUpdate->FrameData);
+					bHasNewData = true;
+				}
+				else
+				{
+					AChaosVDDataContainerBaseActor::FScopedGameFrameDataReRouting ScopedGTDataUpdate(Owner);
+					Owner->UpdateFromNewGameFrameData(*PendingFrameUpdate->FrameData);
+					bHasNewData = true;
+				}
+			}
+
+			if (bHasNewData)
+			{
+				CVDScene->RequestUpdate();
+			}
+		}
+	}
 }
diff --git a/Engine/Plugins/ChaosVD/Source/ChaosVD/Private/Components/ChaosVDAdditionalGTDataRouterComponent.h b/Engine/Plugins/ChaosVD/Source/ChaosVD/Private/Components/ChaosVDAdditionalGTDataRouterComponent.h
--- a/Engine/Plugins/ChaosVD/Source/ChaosVD/Private/Components/ChaosVDAdditionalGTDataRouterComponent.h
+++ b/Engine/Plugins/ChaosVD/Source/ChaosVD/Private/Components/ChaosVDAdditionalGTDataRouterComponent.h
@@ -5,6 +5,8 @@
 #include "Components/ChaosVDSolverDataComponent.h"
 #include "ChaosVDAdditionalGTDataRouterComponent.generated.h"
 
+struct FChaosVDGameFrameDataWrapper;
+
 /** Helper component used to re-route GT data loaded as a Solver Frame.
  * This is part of a compatibility layer to support GT data from multiple files
  * Until we refactor CVD to either support multiple GT Tracks, or remove the difference between GT track and solver track
@@ -17,4 +19,6 @@ class UChaosVDAdditionalGTDataRouterComponent : public UChaosVDSolverDataCompone
 	virtual void UpdateFromSolverFrameData(const FChaosVDSolverFrameData& InSolverFrameData) override;
 
 	virtual void ClearData() override;
+
+	void ApplyDelayedSolverFrameDataUpdate(TSharedPtr<FChaosVDGameFrameDataWrapper> PendingFrameUpdate);
 };
diff --git a/Engine/Plugins/ChaosVD/Source/ChaosVD/Private/Components/ChaosVDSceneQueryDataComponent.cpp b/Engine/Plugins/ChaosVD/Source/ChaosVD/Private/Components/ChaosVDSceneQueryDataComponent.cpp
--- a/Engine/Plugins/ChaosVD/Source/ChaosVD/Private/Components/ChaosVDSceneQueryDataComponent.cpp
+++ b/Engine/Plugins/ChaosVD/Source/ChaosVD/Private/Components/ChaosVDSceneQueryDataComponent.cpp
@@ -33,6 +33,8 @@ void UChaosVDSceneQueryDataComponent::ProccessSQData(const TMap<int32, TSharedPt
 {
 	if (!RecordedQueriesByQueryID)
 	{
+		ClearSceneQuerySelection();
+		ClearData();
 		return;
 	}
 	
diff --git a/Engine/Plugins/ChaosVD/Source/ChaosVD/Private/Trace/ChaosVDTraceAnalyzer.cpp b/Engine/Plugins/ChaosVD/Source/ChaosVD/Private/Trace/ChaosVDTraceAnalyzer.cpp
--- a/Engine/Plugins/ChaosVD/Source/ChaosVD/Private/Trace/ChaosVDTraceAnalyzer.cpp
+++ b/Engine/Plugins/ChaosVD/Source/ChaosVD/Private/Trace/ChaosVDTraceAnalyzer.cpp
@@ -45,7 +45,7 @@ void FChaosVDTraceAnalyzer::OnAnalysisEnd()
 
 void FChaosVDTraceAnalyzer::PushSimulatedSolverTrackForGTData(const UE::Trace::IAnalyzer::FOnEventContext& Context, const UE::Trace::IAnalyzer::FEventData& EventData)
 {
-	FChaosVDSolverFrameData NewFrameData;
+	FChaosVDSolverFrameData NewSimulatedSolverFrameData;
 
 	int32 CurrentGTTrackID = ChaosVDTraceProvider->GetCurrentGameThreadTrackID();
 	if (CurrentGTTrackID == INDEX_NONE)
@@ -54,29 +54,38 @@ void FChaosVDTraceAnalyzer::PushSimulatedSolverTrackForGTData(const UE::Trace::I
 		ChaosVDTraceProvider->SetCurrentGameThreadTrackID(CurrentGTTrackID);
 	}
 
-	NewFrameData.SolverID = CurrentGTTrackID;
-	NewFrameData.FrameCycle = EventData.GetValue<uint64>("Cycle");
-	NewFrameData.bIsKeyFrame = true;
-	NewFrameData.bIsResimulated = false;
-	NewFrameData.StartTime = Context.EventTime.AsSeconds(NewFrameData.FrameCycle);
+	NewSimulatedSolverFrameData.SolverID = CurrentGTTrackID;
+	NewSimulatedSolverFrameData.FrameCycle = EventData.GetValue<uint64>("Cycle");
+	NewSimulatedSolverFrameData.bIsKeyFrame = true;
+	NewSimulatedSolverFrameData.bIsResimulated = false;
+	NewSimulatedSolverFrameData.StartTime = Context.EventTime.AsSeconds(NewSimulatedSolverFrameData.FrameCycle);
 
 	static FString GeneratedStageName = TEXT("Stage 0");
-	NewFrameData.SolverSteps.AddDefaulted_GetRef().StepName = GeneratedStageName;
+	NewSimulatedSolverFrameData.SolverSteps.AddDefaulted_GetRef().StepName = GeneratedStageName;
 
-	NewFrameData.GetCustomData().GetOrAddDefaultData<FChaosVDGameFrameDataWrapper>()->FrameData = MakeShared<FChaosVDGameFrameData>();
+	TSharedPtr<FChaosVDGameFrameData> NewGameFrameData = MakeShared<FChaosVDGameFrameData>();
+	NewSimulatedSolverFrameData.GetCustomData().GetOrAddDefaultData<FChaosVDGameFrameDataWrapper>()->FrameData = NewGameFrameData;
 
-	NewFrameData.DebugFName = FName(TEXT("Additional Game Frame Data Track"));
+	if (TSharedPtr<FChaosVDGameFrameDataWrapperContext> FrameDataContext = NewGameFrameData->GetCustomDataHandler().GetOrAddDefaultData<FChaosVDGameFrameDataWrapperContext>())
+	{
+		for (const TPair<int32, int32>& RemappedID : ChaosVDTraceProvider->RemappedSolversIDs)
+		{
+			FrameDataContext->SupportedSolverIDs.Emplace(RemappedID.Value);
+		}
+	}
+
+	NewSimulatedSolverFrameData.DebugFName = FName(TEXT("Additional Game Frame Data Track"));
 	
-	NewFrameData.AddAttributes(EChaosVDSolverFrameAttributes::HasGTDataToReRoute);
+	NewSimulatedSolverFrameData.AddAttributes(EChaosVDSolverFrameAttributes::HasGTDataToReRoute);
 
 	{
 		if (FChaosVDSolverFrameData* PrevFrameData = ChaosVDTraceProvider->GetCurrentSolverFrame(CurrentGTTrackID))
 		{
-			PrevFrameData->EndTime = NewFrameData.StartTime;
+			PrevFrameData->EndTime = NewSimulatedSolverFrameData.StartTime;
 		}
 	}
 					
-	ChaosVDTraceProvider->StartSolverFrame(NewFrameData.SolverID, MoveTemp(NewFrameData));
+	ChaosVDTraceProvider->StartSolverFrame(NewSimulatedSolverFrameData.SolverID, MoveTemp(NewSimulatedSolverFrameData));
 }
 
 bool FChaosVDTraceAnalyzer::OnEvent(uint16 RouteId, EStyle Style, const FOnEventContext& Context)
diff --git a/Engine/Source/Runtime/Experimental/ChaosVDData/Public/ChaosVDRecording.h b/Engine/Source/Runtime/Experimental/ChaosVDData/Public/ChaosVDRecording.h
--- a/Engine/Source/Runtime/Experimental/ChaosVDData/Public/ChaosVDRecording.h
+++ b/Engine/Source/Runtime/Experimental/ChaosVDData/Public/ChaosVDRecording.h
@@ -390,6 +390,14 @@ struct FChaosVDGameFrameDataWrapper
 	TSharedPtr<FChaosVDGameFrameData> FrameData;
 };
 
+USTRUCT()
+struct FChaosVDGameFrameDataWrapperContext
+{
+	GENERATED_BODY()
+
+	TArray<int32, TInlineAllocator<25>> SupportedSolverIDs;
+};
+
 enum class EChaosVDRecordingAttributes : uint8
 {
 	None = 0,
