[ChaosVD] Fix for an issue where the first scene query recorded in a frame is not visualized

The problem is that the scene query data processor is using the wrong getter to obtain the scene query data container for a frame. Data processors need to use GetOrAdd, but the scene query one was using get, therefore it was not creating it if it was the first query.

It works for subsequent queries because, after the scene query metadata, we trace the visited shapes, and that is using the correct getter (so in reality if subsequent queries didn't intersect any shape, multiple queries might be missing, but most of the time is just the first one).

diff --git a/Engine/Plugins/ChaosVD/Source/ChaosVD/Private/Trace/DataProcessors/ChaosVDSceneQueryDataProcessor.cpp b/Engine/Plugins/ChaosVD/Source/ChaosVD/Private/Trace/DataProcessors/ChaosVDSceneQueryDataProcessor.cpp
--- a/Engine/Plugins/ChaosVD/Source/ChaosVD/Private/Trace/DataProcessors/ChaosVDSceneQueryDataProcessor.cpp
+++ b/Engine/Plugins/ChaosVD/Source/ChaosVD/Private/Trace/DataProcessors/ChaosVDSceneQueryDataProcessor.cpp
@@ -30,7 +30,7 @@ bool FChaosVDSceneQueryDataProcessor::ProcessRawData(const TArray<uint8>& InData
 	{
 		if (TSharedPtr<FChaosVDGameFrameData> CurrentFrameData = ProviderSharedPtr->GetCurrentGameFrame().Pin())
 		{
-			if (TSharedPtr<FChaosVDSceneQueriesDataContainer> SQDataContainer = CurrentFrameData->GetCustomDataHandler().GetData<FChaosVDSceneQueriesDataContainer>())
+			if (TSharedPtr<FChaosVDSceneQueriesDataContainer> SQDataContainer = CurrentFrameData->GetCustomDataHandler().GetOrAddDefaultData<FChaosVDSceneQueriesDataContainer>())
 			{
 				QueryData->WorldSolverID = ProviderSharedPtr->GetRemappedSolverID(QueryData->WorldSolverID);
 
@@ -43,7 +43,7 @@ bool FChaosVDSceneQueryDataProcessor::ProcessRawData(const TArray<uint8>& InData
 						{
 							(*ParentQueryData)->SubQueriesIDs.Add(QueryData->ID);
 						}
-					}	
+					}
 				}
 
 				SQDataContainer->RecordedSceneQueriesByQueryID.Add(QueryData->ID, QueryData);
