[ChaosVD] Fix incorrect check for null in ChaosVDSolverCollisionDataComponent.

We should check the ptr after the cast to make sure it is not only valid but also from the correct type (also because the casted ptr is the one used in the rest of the code).

diff --git a/Engine/Plugins/ChaosVD/Source/ChaosVD/Private/Visualizers/ChaosVDSolverCollisionDataComponentVisualizer.cpp b/Engine/Plugins/ChaosVD/Source/ChaosVD/Private/Visualizers/ChaosVDSolverCollisionDataComponentVisualizer.cpp
--- a/Engine/Plugins/ChaosVD/Source/ChaosVD/Private/Visualizers/ChaosVDSolverCollisionDataComponentVisualizer.cpp
+++ b/Engine/Plugins/ChaosVD/Source/ChaosVD/Private/Visualizers/ChaosVDSolverCollisionDataComponentVisualizer.cpp
@@ -127,7 +127,7 @@ void FChaosVDSolverCollisionDataComponentVisualizer::DrawMidPhaseData(const UAct
 	}
 
 	const UChaosVDCollisionDataVisualizationSettings* DebugDrawSettings = Cast<UChaosVDCollisionDataVisualizationSettings>(VisualizationContext.DebugDrawSettings);
-	if (!VisualizationContext.DebugDrawSettings)
+	if (!DebugDrawSettings)
 	{
 		return;
 	}
