[ChaosVD] Fix CIS issue in Shipping Editor Confirgurations.

diff --git a/Engine/Plugins/ChaosVD/Source/ChaosVD/Private/Widgets/SChaosVDRecordingControls.cpp b/Engine/Plugins/ChaosVD/Source/ChaosVD/Private/Widgets/SChaosVDRecordingControls.cpp
--- a/Engine/Plugins/ChaosVD/Source/ChaosVD/Private/Widgets/SChaosVDRecordingControls.cpp
+++ b/Engine/Plugins/ChaosVD/Source/ChaosVD/Private/Widgets/SChaosVDRecordingControls.cpp
@@ -765,7 +765,8 @@ void SChaosVDRecordingControls::SetSessionRecordingState(bool bIsRecording, ECha
 
 			constexpr bool bAppendPort = false;
 			RecordingParams.Target = LocalIP->ToString(bAppendPort);
-	
+
+#if WITH_CHAOS_VISUAL_DEBUGGER
 			Chaos::VisualDebugger::FChaosVDDataChannelsManager::Get().EnumerateChannels([&RecordingParams](const TSharedRef<Chaos::VisualDebugger::FChaosVDOptionalDataChannel>& Channel)
 			{
 				if (Channel->IsChannelEnabled())
@@ -774,7 +775,7 @@ void SChaosVDRecordingControls::SetSessionRecordingState(bool bIsRecording, ECha
 				}
 				return true;
 			});
-
+#endif
 			RemoteSessionManager->SendStartRecordingCommand(SessionInfoRef->Address, RecordingParams);
 
 			// Once the start recording command is issue, we can try to connect to the created session (if everything went well).
@@ -806,6 +807,7 @@ void SChaosVDRecordingControls::SetSessionRecordingState(bool bIsRecording, ECha
 			FChaosVDStartRecordingCommandMessage RecordingParams;
 			RecordingParams.RecordingMode = EChaosVDRecordingMode::File;
 
+#if WITH_CHAOS_VISUAL_DEBUGGER
 			Chaos::VisualDebugger::FChaosVDDataChannelsManager::Get().EnumerateChannels([&RecordingParams](const TSharedRef<Chaos::VisualDebugger::FChaosVDOptionalDataChannel>& Channel)
 			{
 				if (Channel->IsChannelEnabled())
@@ -814,6 +816,7 @@ void SChaosVDRecordingControls::SetSessionRecordingState(bool bIsRecording, ECha
 				}
 				return true;
 			});
+#endif
 			
 			RemoteSessionManager->SendStartRecordingCommand(SessionInfoRef->Address, RecordingParams);
 
