[ChaosVD] Adding Serialization Tests to EngineTest

This CL exposes the required elements in CVD's public API to implement tests that will check that we can load CVD files from previous engine versions without serialization issues. And adds these tests to the SimulationTests plugin

#virtualized

diff --git a/Engine/Plugins/ChaosVD/Source/ChaosVD/Private/Trace/ChaosVDTraceProvider.cpp b/Engine/Plugins/ChaosVD/Source/ChaosVD/Private/Trace/ChaosVDTraceProvider.cpp
--- a/Engine/Plugins/ChaosVD/Source/ChaosVD/Private/Trace/ChaosVDTraceProvider.cpp
+++ b/Engine/Plugins/ChaosVD/Source/ChaosVD/Private/Trace/ChaosVDTraceProvider.cpp
@@ -109,7 +109,7 @@ void FChaosVDTraceProvider::GetAvailablePendingSolverIDsAtGameFrame(const TShare
 	}
 }
 
-FString FChaosVDTraceProvider::GenerateFormattedStringListFromSet(const TSet<FString>& StringsSet)
+FString FChaosVDTraceProvider::GenerateFormattedStringListFromSet(const TSet<FString>& StringsSet) const
 {
 	FString FormattedString = TEXT("");
 	for (const FString& ListEntry : StringsSet)
diff --git a/Engine/Plugins/ChaosVD/Source/ChaosVD/Public/ChaosVDModule.h b/Engine/Plugins/ChaosVD/Source/ChaosVD/Public/ChaosVDModule.h
--- a/Engine/Plugins/ChaosVD/Source/ChaosVD/Public/ChaosVDModule.h
+++ b/Engine/Plugins/ChaosVD/Source/ChaosVD/Public/ChaosVDModule.h
@@ -24,14 +24,14 @@ public:
 
 	using ToolkitHostType = SChaosVDMainTab;
 
-	static FChaosVDModule& Get();
+	CHAOSVD_API static FChaosVDModule& Get();
 
 	/** IModuleInterface implementation */
 	virtual void StartupModule() override;
 	virtual void ShutdownModule() override;
 
 	/** Returns the Trace manager instance used by CVD */
-	TSharedPtr<FChaosVDTraceManager>& GetTraceManager() { return ChaosVDTraceManager; }
+	CHAOSVD_API TSharedPtr<FChaosVDTraceManager>& GetTraceManager() { return ChaosVDTraceManager; }
 
 	/** Spawns a new instance of CVD */
 	void SpawnCVDTab();
diff --git a/Engine/Plugins/ChaosVD/Source/ChaosVD/Public/Trace/ChaosVDTraceProvider.h b/Engine/Plugins/ChaosVD/Source/ChaosVD/Public/Trace/ChaosVDTraceProvider.h
--- a/Engine/Plugins/ChaosVD/Source/ChaosVD/Public/Trace/ChaosVDTraceProvider.h
+++ b/Engine/Plugins/ChaosVD/Source/ChaosVD/Public/Trace/ChaosVDTraceProvider.h
@@ -39,7 +39,7 @@ class FChaosVDTraceProvider : public TraceServices::IProvider, public TraceServi
 {
 public:
 	
-	static FName ProviderName;
+	CHAOSVD_API static FName ProviderName;
 
 	CHAOSVD_API FChaosVDTraceProvider(TraceServices::IAnalysisSession& InSession);
 
@@ -124,9 +124,10 @@ private:
 	/** Gathers any solver id from solver data that is not fully processed yet but that will be valid for the provided game frame data later on */
 	void GetAvailablePendingSolverIDsAtGameFrame(const TSharedRef<FChaosVDGameFrameData>& InProcessedGameFrameData, TArray<int32, TInlineAllocator<16>>& OutSolverIDs);
 
-	FString GenerateFormattedStringListFromSet(const TSet<FString>& StringsSet);
-
 public:
+
+	CHAOSVD_API FString GenerateFormattedStringListFromSet(const TSet<FString>& StringsSet) const;
+
 	CHAOSVD_API int32 RemapSolverID(int32 SolverID);
 
 	virtual void BeginEdit() const override{};
@@ -164,6 +165,11 @@ public:
 		return DataProcessedSoFarCounter;
 	}
 
+	const TSet<FString>& GetTypesFailedToSerialize() const
+	{
+		return TypesFailedToSerialize;
+	}
+
 private:
 	
 	int32 GetCurrentGameThreadTrackID() const
