[ChaosVD] Fix CIS issue in Engine test

diff --git a/Engine/Plugins/ChaosVD/Source/ChaosVD/Private/Trace/ChaosVDTraceManager.h b/Engine/Plugins/ChaosVD/Source/ChaosVD/Private/Trace/ChaosVDTraceManager.h
deleted file mode 100644
diff --git a/Engine/Plugins/ChaosVD/Source/ChaosVD/Private/Trace/ChaosVDTraceModule.cpp b/Engine/Plugins/ChaosVD/Source/ChaosVD/Private/Trace/ChaosVDTraceModule.cpp
--- a/Engine/Plugins/ChaosVD/Source/ChaosVD/Private/Trace/ChaosVDTraceModule.cpp
+++ b/Engine/Plugins/ChaosVD/Source/ChaosVD/Private/Trace/ChaosVDTraceModule.cpp
@@ -3,7 +3,7 @@
 #include "Trace/ChaosVDTraceModule.h"
 
 #include "ChaosVDModule.h"
-#include "ChaosVDTraceManager.h"
+#include "Trace/ChaosVDTraceManager.h"
 #include "Trace/ChaosVDTraceAnalyzer.h"
 #include "Trace/ChaosVDTraceProvider.h"
 #include "TraceServices/Model/AnalysisSession.h"
diff --git a/Engine/Plugins/ChaosVD/Source/ChaosVD/Public/Trace/ChaosVDTraceManager.h b/Engine/Plugins/ChaosVD/Source/ChaosVD/Public/Trace/ChaosVDTraceManager.h
new file mode 100644
--- /dev/null
+++ b/Engine/Plugins/ChaosVD/Source/ChaosVD/Public/Trace/ChaosVDTraceManager.h
@@ -0,0 +1,120 @@
+// Copyright Epic Games, Inc. All Rights Reserved.
+
+#pragma once
+
+#include "ChaosVDModule.h"
+#include "Containers/StringFwd.h"
+#include "Templates/SharedPointer.h"
+#include "Trace/StoreClient.h"
+
+enum class EChaosVDLoadRecordedDataMode : uint8;
+class FChaosVDEngine;
+class FChaosVDTraceModule;
+
+namespace TraceServices { class IAnalysisSession; }
+
+/**
+ * Structure containing info about a trace session used by Chaos Visual Debugger
+ */
+struct FChaosVDTraceSessionDescriptor
+{
+	FString SessionName;
+	bool bIsLiveSession = false;
+
+	bool IsValid() const { return !SessionName.IsEmpty(); }
+};
+
+/** Objects that allows us to use TLS to temporarily store and access a ptr to an existing instance.
+ * This is temporary to workaround the lack of an API method we need in the trace API, and will be removed in the future.
+ * either when we add that to the API, or find another way to pass an existing CVD recording to the trace provider before anaislis starts
+ */
+class FChaosVDTraceManagerThreadContext : public TThreadSingleton<FChaosVDTraceManagerThreadContext>
+{
+public:
+
+	TWeakPtr<FChaosVDRecording> PendingExternalRecordingWeakPtr;
+};
+
+/** Manager class used by Chaos VD to interact/control UE Trace systems */
+class FChaosVDTraceManager
+{
+public:
+	FChaosVDTraceManager();
+	~FChaosVDTraceManager();
+
+	/** Load a trace file and starts analyzing it
+	 * @param InTraceFilename File Name including (Path Included) of the Trace file to load
+	 */
+	CHAOSVD_API FString LoadTraceFile(const FString& InTraceFilename, const TSharedPtr<FChaosVDRecording>& ExistingRecordingPtr = nullptr);
+
+	FString LoadTraceFile(TUniquePtr<IFileHandle>&& InFileHandle, const FString& InTraceSessionName, const TSharedPtr<FChaosVDRecording>& ExistingRecordingPtr);
+	
+	/**
+	 * Connects to a live Trace Session and starts analyzing it.
+	 * @param InSessionHost Trace Store Address for this session
+	 * @param SessionID Trace ID in the Trace Store provided as host
+	 * @return Session Name
+	 */
+	FString ConnectToLiveSession(FStringView InSessionHost, uint32 SessionID, const TSharedPtr<FChaosVDRecording>& ExistingRecordingPtr = nullptr);
+
+	/** Returns the path to the local trace store */
+	FString GetLocalTraceStoreDirPath();
+
+	/** Returns a ptr to the session registered with the provided session name. Null if no session is found */
+	CHAOSVD_API TSharedPtr<const TraceServices::IAnalysisSession> GetSession(const FString& InSessionName);
+
+	/** Stops and de-registers a trace session registered with the provided session name */
+	CHAOSVD_API void CloseSession(const FString& InSessionName);
+
+	/** Stops a trace session registered with the provided session name */
+	void StopSession(const FString& InSessionName);
+
+	template<typename TVisitor>
+	static void EnumerateActiveSessions(FStringView InSessionHost, TVisitor Callback);
+	static const UE::Trace::FStoreClient::FSessionInfo* GetTraceSessionInfo(FStringView InSessionHost, FGuid TraceGuid);
+
+private:
+
+	TUniquePtr<UE::Trace::IInDataStream> CreateFileDataStream(TUniquePtr<IFileHandle>&& InFileHandle);
+
+	/** Temporary workaround method to set an existing recording structure in CVD's trace provider before the trace analysis starts
+	 * in the Trace analysis thread.*/
+	void SetPendingExternalRecordingToProcess(const TSharedPtr<FChaosVDRecording>& ExistingRecordingPtr);
+
+	/** The trace analysis session. */
+	TMap<FString, TSharedPtr<const TraceServices::IAnalysisSession>> AnalysisSessionByName;
+
+	TSharedPtr<FChaosVDTraceModule> ChaosVDTraceModule;
+};
+
+template <typename TCallback>
+void FChaosVDTraceManager::EnumerateActiveSessions(FStringView InSessionHost, TCallback Callback)
+{
+	if (InSessionHost.IsEmpty())
+	{
+		UE_LOG(LogChaosVDEditor, Error, TEXT("[%s] Failed to connect to trace store. Provided session host is empty"), ANSI_TO_TCHAR(__FUNCTION__));
+		return;
+	}
+
+	using namespace UE::Trace;
+	const FStoreClient* StoreClient = FStoreClient::Connect(InSessionHost.GetData());
+
+	if (!StoreClient)
+	{
+		UE_LOG(LogChaosVDEditor, Error, TEXT("[%s] Failed to connect to trace store at [%s]"), ANSI_TO_TCHAR(__FUNCTION__), InSessionHost.GetData())
+		return;
+	}
+
+	const uint32 SessionCount = StoreClient->GetSessionCount();
+
+	for (uint32 SessionIndex = 0; SessionIndex < SessionCount; SessionIndex++)
+	{
+		if (const FStoreClient::FSessionInfo* SessionInfo = StoreClient->GetSessionInfo(SessionIndex))
+		{
+			if (!Callback(*SessionInfo))
+			{
+				return;
+			}
+		}	
+	}
+}
