[ChaosVD] Implemented a scene integrity functional test.

This CL exposes some required objects and methods to the public API so they can be used by the functional test in the Simulation Test Plugins.

It also adds a way to take a playback engine snapshot that is later used by the functional test to create a CVD engine instance, load the target CVD file into it, and scrub to the saved frames. Then the test will generate a new snapshot on the fly, and compare it with the saved one.

If there are any differences between the two, the test fails, otherwise it passes.

#virtualized

diff --git a/Engine/Plugins/ChaosVD/Source/ChaosVD/Private/Actors/ChaosVDDataContainerBaseActor.cpp b/Engine/Plugins/ChaosVD/Source/ChaosVD/Private/Actors/ChaosVDDataContainerBaseActor.cpp
--- a/Engine/Plugins/ChaosVD/Source/ChaosVD/Private/Actors/ChaosVDDataContainerBaseActor.cpp
+++ b/Engine/Plugins/ChaosVD/Source/ChaosVD/Private/Actors/ChaosVDDataContainerBaseActor.cpp
@@ -3,11 +3,17 @@
 #include "Actors/ChaosVDDataContainerBaseActor.h"
 
 #include "ChaosVDModule.h"
+#include "ChaosVDSceneCompositionReport.h"
 #include "Components/ChaosVDSolverDataComponent.h"
 #include "ExtensionsSystem/ChaosVDExtensionsManager.h"
 
 #include UE_INLINE_GENERATED_CPP_BY_NAME(ChaosVDDataContainerBaseActor)
 
+namespace Chaos::VD::Test::SceneObjectTypes
+{
+	FName SolverID = FName("SolverID");
+}
+
 AChaosVDDataContainerBaseActor::AChaosVDDataContainerBaseActor()
 {
 	PrimaryActorTick.bCanEverTick = false;
@@ -96,6 +102,21 @@ void AChaosVDDataContainerBaseActor::SetIsTemporarilyHiddenInEditor(bool bIsHidd
 	Super::SetIsTemporarilyHiddenInEditor(bIsHidden);
 }
 
+void AChaosVDDataContainerBaseActor::AppendSceneCompositionTestData(FChaosVDSceneCompositionTestData& OutStateTestData)
+{
+	int32& CurrentCount = OutStateTestData.ObjectsCountByType.FindOrAdd(Chaos::VD::Test::SceneObjectTypes::SolverID);
+	CurrentCount++;
+
+	TInlineComponentArray<UChaosVDSolverDataComponent*> SolverDataComponents(this);
+	for (UChaosVDSolverDataComponent* Component : SolverDataComponents)
+	{
+		if (Component)
+		{
+			Component->AppendSceneCompositionTestData(OutStateTestData);
+		}
+	}
+}
+
 void AChaosVDDataContainerBaseActor::HandlePostInitializationExtensionRegistered(const TSharedRef<FChaosVDExtension>& NewExtension)
 {
 	TConstArrayView<TSubclassOf<UActorComponent>> CustomDataComponentsClasses = NewExtension->GetSolverDataComponentsClasses();
diff --git a/Engine/Plugins/ChaosVD/Source/ChaosVD/Private/ChaosVDEngine.h b/Engine/Plugins/ChaosVD/Source/ChaosVD/Private/ChaosVDEngine.h
deleted file mode 100644
diff --git a/Engine/Plugins/ChaosVD/Source/ChaosVD/Private/ChaosVDMiscSceneObjectTypeNamesForTesting.cpp b/Engine/Plugins/ChaosVD/Source/ChaosVD/Private/ChaosVDMiscSceneObjectTypeNamesForTesting.cpp
new file mode 100644
--- /dev/null
+++ b/Engine/Plugins/ChaosVD/Source/ChaosVD/Private/ChaosVDMiscSceneObjectTypeNamesForTesting.cpp
@@ -0,0 +1,11 @@
+// Copyright Epic Games, Inc. All Rights Reserved.
+
+#include "ChaosVDMiscSceneObjectTypeNamesForTesting.h"
+
+namespace Chaos::VD::Test::SceneObjectTypes
+{
+	FName SceneQuery = FName("SceneQuery");
+	FName CollisionConstraint = FName("CollisionConstraint");
+	FName JointConstraint = FName("JointConstraint");
+	FName CharacterGroundConstraint = FName("CharacterGroundConstraint");
+}
diff --git a/Engine/Plugins/ChaosVD/Source/ChaosVD/Private/ChaosVDModule.cpp b/Engine/Plugins/ChaosVD/Source/ChaosVD/Private/ChaosVDModule.cpp
--- a/Engine/Plugins/ChaosVD/Source/ChaosVD/Private/ChaosVDModule.cpp
+++ b/Engine/Plugins/ChaosVD/Source/ChaosVD/Private/ChaosVDModule.cpp
@@ -4,10 +4,14 @@
 
 #include "ChaosVDCommands.h"
 #include "ChaosVDEngine.h"
+#include "ChaosVDPlaybackController.h"
+#include "ChaosVDSceneCompositionReport.h"
 #include "ChaosVDSceneParticleCustomization.h"
 #include "ChaosVDSettingsManager.h"
 #include "ChaosVDStyle.h"
 #include "ChaosVDTabsIDs.h"
+#include "DesktopPlatformModule.h"
+#include "IDesktopPlatform.h"
 #include "Misc/App.h"
 #include "Misc/Guid.h"
 #include "Trace/ChaosVDTraceManager.h"
@@ -17,11 +21,22 @@
 #include "WorkspaceMenuStructureModule.h"
 #include "DetailsCustomizations/ChaosVDShapeDataCustomization.h"
 #include "ExtensionsSystem/ChaosVDExtensionsManager.h"
+#include "Framework/Application/SlateApplication.h"
+#include "HAL/FileManager.h"
 
 #define LOCTEXT_NAMESPACE "ChaosVisualDebugger"
 
 DEFINE_LOG_CATEGORY(LogChaosVDEditor);
 
+FAutoConsoleCommand ChaosVDTakePlaybackEngineSnapshot(
+	TEXT("p.Chaos.VD.TakePlaybackEngineSnapshot"),
+	TEXT("Take a snapshot of CVD's playback engine at the current frame to be used in functional tests"),
+	FConsoleCommandWithArgsDelegate::CreateLambda([](const TArray<FString>& Args)
+	{
+		FChaosVDModule::Get().TakePlaybackEngineStateSnapshot();
+	})
+);
+
 FString FChaosVDModule::ChaosVisualDebuggerProgramName = TEXT("ChaosVisualDebugger");
 
 FChaosVDModule& FChaosVDModule::Get()
@@ -225,6 +240,67 @@ void FChaosVDModule::ReloadInstanceUI(FGuid InstanceGUID)
 	}
 }
 
+
+void FChaosVDModule::TakePlaybackEngineStateSnapshot()
+{
+	if (ActiveChaosVDInstances.IsEmpty())
+	{
+		UE_LOG(LogChaosVDEditor, Error, TEXT("[%hs]: There are no active CVD Instances."), __func__);
+		return;
+	}
+
+	if (ActiveChaosVDInstances.Num() > 1)
+	{
+		UE_LOG(LogChaosVDEditor, Error, TEXT("[%hs]: Snapshots can only be taken when only one CVD instance is active."), __func__);
+		return;
+	}
+	
+	TSharedPtr<FChaosVDEngine> CVDEngineInstance;
+	for (const TPair<FGuid, TSharedPtr<FChaosVDEngine>>& EngineInstanceWithID : ActiveChaosVDInstances)
+	{
+		if (EngineInstanceWithID.Value)
+		{
+			CVDEngineInstance = EngineInstanceWithID.Value;
+			break;
+		}
+	}
+	
+	FChaosVDPlaybackEngineSnapshot StateData = CVDEngineInstance->GetPlaybackController()->GeneratePlaybackEngineSnapshot();
+
+	TArray<FString> OutSelectedFilenames;
+	if (IDesktopPlatform* DesktopPlatform = FDesktopPlatformModule::Get())
+	{
+		FString ExtensionStr;
+		ExtensionStr += TEXT("CVD Playback Engine State|*.cvdplaystate");
+		
+		DesktopPlatform->SaveFileDialog(
+			FSlateApplication::Get().FindBestParentWindowHandleForDialogs(nullptr),
+			LOCTEXT("SaveDialogTitle", "Dump Playback Engine Sate").ToString(),
+			*FPaths::ProfilingDir(),
+			TEXT(""),
+			*ExtensionStr,
+			EFileDialogFlags::None,
+			OutSelectedFilenames
+		);
+	}
+
+	if (!OutSelectedFilenames.IsEmpty())
+	{
+		const FString& TargetFilePath = OutSelectedFilenames[0];
+
+		IFileManager& FileManager = IFileManager::Get();
+		FArchive* FileWriter = FileManager.CreateFileWriter(*TargetFilePath);
+
+		TArray<uint8> RawData;
+		FMemoryWriter MemWriter(RawData);
+		FChaosVDPlaybackEngineSnapshot::StaticStruct()->SerializeTaggedProperties(MemWriter, reinterpret_cast<uint8*>(&StateData),FChaosVDPlaybackEngineSnapshot::StaticStruct(), nullptr);
+
+		FileWriter->Serialize(RawData.GetData(), RawData.Num());
+		FileWriter->Close();
+		FPlatformProcess::ExploreFolder(*TargetFilePath);
+	}
+}
+
 #undef LOCTEXT_NAMESPACE
 	
 IMPLEMENT_MODULE(FChaosVDModule, ChaosVD)
diff --git a/Engine/Plugins/ChaosVD/Source/ChaosVD/Private/ChaosVDPlaybackController.cpp b/Engine/Plugins/ChaosVD/Source/ChaosVD/Private/ChaosVDPlaybackController.cpp
--- a/Engine/Plugins/ChaosVD/Source/ChaosVD/Private/ChaosVDPlaybackController.cpp
+++ b/Engine/Plugins/ChaosVD/Source/ChaosVD/Private/ChaosVDPlaybackController.cpp
@@ -7,11 +7,13 @@
 #include "ChaosVDRecording.h"
 #include "ChaosVDRuntimeModule.h"
 #include "ChaosVDScene.h"
+#include "ChaosVDSceneCompositionReport.h"
 #include "ExtensionsSystem/ChaosVDExtensionsManager.h"
 #include "Misc/MessageDialog.h"
 #include "Trace/ChaosVDTraceManager.h"
 #include "Trace/ChaosVDTraceProvider.h"
 #include "TraceServices/Model/AnalysisSession.h"
+#include "Widgets/SChaosVDTimelineWidget.h"
 
 #include UE_INLINE_GENERATED_CPP_BY_NAME(ChaosVDPlaybackController)
 
@@ -1391,6 +1393,41 @@ void FChaosVDPlaybackController::TrySetActiveTrack(const TSharedRef<const FChaos
 	});
 }
 
+FChaosVDPlaybackEngineSnapshot FChaosVDPlaybackController::GeneratePlaybackEngineSnapshot()
+{
+	FChaosVDPlaybackEngineSnapshot PlaybackEngineState;
+
+	PlaybackEngineState.TimelineSyncMode = GetTimelineSyncMode();
+
+	if (const TSharedPtr<const FChaosVDTrackInfo> GameTrackInfo = GetTrackInfo(EChaosVDTrackType::Game, GameTrackID))
+	{
+		PlaybackEngineState.FrameIndexDataByTrackID.Add(GameTrackID, { GameTrackInfo->CurrentFrame, GameTrackInfo->CurrentStage });
+		TArray<TSharedPtr<const FChaosVDTrackInfo>> TrackInfoArray;
+		GetAvailableTrackInfosAtTrackFrame(EChaosVDTrackType::Solver, GameTrackInfo.ToSharedRef(), TrackInfoArray);
+
+		for (const TSharedPtr<const FChaosVDTrackInfo>& SolversTrackInfo : TrackInfoArray)
+		{
+			if (SolversTrackInfo)
+			{
+				PlaybackEngineState.FrameIndexDataByTrackID.Add(SolversTrackInfo->TrackID, { SolversTrackInfo->CurrentFrame, SolversTrackInfo->CurrentStage });
+			}
+		}
+	}
+
+	FChaosVDExtensionsManager::Get().EnumerateExtensions([&PlaybackEngineState](const TSharedRef<FChaosVDExtension>& Extension)
+	{
+		PlaybackEngineState.InstalledExtensionsNames.Add(Extension->GetExtensionType());
+		return true;
+	});
+
+	if (TSharedPtr<FChaosVDScene> CVDScenePtr = SceneToControl.Pin())
+	{
+		CVDScenePtr->AppendSceneCompositionTestData(PlaybackEngineState);
+	}
+
+	return PlaybackEngineState;
+}
+
 void FChaosVDPlaybackController::UpdateSolverTracksData()
 {
 	if (!LoadedRecording.IsValid())
diff --git a/Engine/Plugins/ChaosVD/Source/ChaosVD/Private/ChaosVDScene.cpp b/Engine/Plugins/ChaosVD/Source/ChaosVD/Private/ChaosVDScene.cpp
--- a/Engine/Plugins/ChaosVD/Source/ChaosVD/Private/ChaosVDScene.cpp
+++ b/Engine/Plugins/ChaosVD/Source/ChaosVD/Private/ChaosVDScene.cpp
@@ -8,6 +8,7 @@
 #include "ChaosVDSceneParticle.h"
 #include "Chaos/ImplicitObject.h"
 #include "ChaosVDRecording.h"
+#include "ChaosVDSceneCompositionReport.h"
 #include "ChaosVDSelectionCustomization.h"
 #include "ChaosVDSettingsManager.h"
 #include "ChaosVDSkySphereInterface.h"
@@ -649,6 +650,17 @@ void FChaosVDScene::UpdateWorldStreamingLocation(const FVector& InLocation)
 	}
 }
 
+void FChaosVDScene::AppendSceneCompositionTestData(FChaosVDPlaybackEngineSnapshot& OutStateTestData)
+{
+	for (TObjectPtr<AChaosVDDataContainerBaseActor>& DataContainerActor : AvailableDataContainerActors)
+	{
+		if (DataContainerActor)
+		{
+			DataContainerActor->AppendSceneCompositionTestData(OutStateTestData.SceneComposition);
+		}
+	}
+}
+
 void FChaosVDScene::InitializeSelectionSets()
 {
 	SelectionSet = NewObject<UTypedElementSelectionSet>(GetTransientPackage(), NAME_None, RF_Transactional);
diff --git a/Engine/Plugins/ChaosVD/Source/ChaosVD/Private/Components/ChaosVDConstraintDataComponent.cpp b/Engine/Plugins/ChaosVD/Source/ChaosVD/Private/Components/ChaosVDConstraintDataComponent.cpp
--- a/Engine/Plugins/ChaosVD/Source/ChaosVD/Private/Components/ChaosVDConstraintDataComponent.cpp
+++ b/Engine/Plugins/ChaosVD/Source/ChaosVD/Private/Components/ChaosVDConstraintDataComponent.cpp
@@ -3,6 +3,8 @@
 #include "ChaosVDConstraintDataComponent.h"
 
 #include "ChaosVDConstraintDataHelpers.h"
+#include "ChaosVDMiscSceneObjectTypeNamesForTesting.h"
+#include "ChaosVDSceneCompositionReport.h"
 
 #include UE_INLINE_GENERATED_CPP_BY_NAME(ChaosVDConstraintDataComponent)
 
@@ -36,3 +38,11 @@ void UChaosVDConstraintDataComponent::ClearData()
 	ConstraintByConstraintIndex.Empty();
 }
 
+void UChaosVDConstraintDataComponent::AppendSceneCompositionTestData(FChaosVDSceneCompositionTestData& OutStateTestData)
+{
+	Super::AppendSceneCompositionTestData(OutStateTestData);
+
+	int32& CurrentCount = OutStateTestData.ObjectsCountByType.FindOrAdd(Chaos::VD::Test::SceneObjectTypes::CollisionConstraint);
+	CurrentCount+= AllConstraints.Num();
+}
+
diff --git a/Engine/Plugins/ChaosVD/Source/ChaosVD/Private/Components/ChaosVDConstraintDataComponent.h b/Engine/Plugins/ChaosVD/Source/ChaosVD/Private/Components/ChaosVDConstraintDataComponent.h
--- a/Engine/Plugins/ChaosVD/Source/ChaosVD/Private/Components/ChaosVDConstraintDataComponent.h
+++ b/Engine/Plugins/ChaosVD/Source/ChaosVD/Private/Components/ChaosVDConstraintDataComponent.h
@@ -28,6 +28,8 @@ public:
 
 	virtual void ClearData() override;
 
+	virtual void AppendSceneCompositionTestData(FChaosVDSceneCompositionTestData& OutStateTestData) override;
+
 protected:
 
 	TSharedPtr<FChaosVDConstraintDataWrapperBase> GetConstraintByIndex(int32 ConstraintIndex);
diff --git a/Engine/Plugins/ChaosVD/Source/ChaosVD/Private/Components/ChaosVDParticleDataComponent.cpp b/Engine/Plugins/ChaosVD/Source/ChaosVD/Private/Components/ChaosVDParticleDataComponent.cpp
--- a/Engine/Plugins/ChaosVD/Source/ChaosVD/Private/Components/ChaosVDParticleDataComponent.cpp
+++ b/Engine/Plugins/ChaosVD/Source/ChaosVD/Private/Components/ChaosVDParticleDataComponent.cpp
@@ -4,6 +4,7 @@
 
 #include "ChaosVDRecording.h"
 #include "ChaosVDScene.h"
+#include "ChaosVDSceneCompositionReport.h"
 #include "ChaosVDSceneParticle.h"
 #include "ChaosVDSettingsManager.h"
 #include "Actors/ChaosVDSolverInfoActor.h"
@@ -40,6 +41,11 @@ namespace Chaos::VisualDebugger::Cvars
 		TEXT("If false, CVD will only rely on the particle destroyed events to figure out what needs to be removed from the visualization, instead of doing a diff based on keyframe data"));
 }
 
+namespace Chaos::VD::Test::SceneObjectTypes
+{
+	FName ActiveParticles = FName("ActiveParticles");
+}
+
 UChaosVDParticleDataComponent::UChaosVDParticleDataComponent()
 {
 	if (HasAnyFlags(RF_ClassDefaultObject))
@@ -141,6 +147,22 @@ void UChaosVDParticleDataComponent::HandleWorldStreamingLocationUpdated(const FV
 	StreamingSystem.UpdateStreamingSourceLocation(InLocation);
 }
 
+void UChaosVDParticleDataComponent::AppendSceneCompositionTestData(FChaosVDSceneCompositionTestData& OutStateTestData)
+{
+	Super::AppendSceneCompositionTestData(OutStateTestData);
+
+	int32& CurrentCount = OutStateTestData.ObjectsCountByType.FindOrAdd(Chaos::VD::Test::SceneObjectTypes::ActiveParticles);
+
+	for (const TPair<int32, TSharedRef<FChaosVDSceneParticle>>& ParticleWithID : SolverParticlesByID)
+	{
+		// Only take into account active particles because inactive particles are destroyed ones we keep instantiated as part of a pseudo pool. 
+		if (ParticleWithID.Value->IsActive())
+		{
+			CurrentCount++;
+		}
+	}
+}
+
 int32 UChaosVDParticleDataComponent::GetMaxElementCountForBatchTEDSUpdate_AssumesLocked() const
 {
 	// All Particles + their container objects + this component owner
diff --git a/Engine/Plugins/ChaosVD/Source/ChaosVD/Private/Components/ChaosVDSceneQueryDataComponent.cpp b/Engine/Plugins/ChaosVD/Source/ChaosVD/Private/Components/ChaosVDSceneQueryDataComponent.cpp
--- a/Engine/Plugins/ChaosVD/Source/ChaosVD/Private/Components/ChaosVDSceneQueryDataComponent.cpp
+++ b/Engine/Plugins/ChaosVD/Source/ChaosVD/Private/Components/ChaosVDSceneQueryDataComponent.cpp
@@ -2,8 +2,10 @@
 
 #include "Components/ChaosVDSceneQueryDataComponent.h"
 
+#include "ChaosVDMiscSceneObjectTypeNamesForTesting.h"
 #include "ChaosVDRecording.h"
 #include "ChaosVDScene.h"
+#include "ChaosVDSceneCompositionReport.h"
 #include "Actors/ChaosVDSolverInfoActor.h"
 
 #include UE_INLINE_GENERATED_CPP_BY_NAME(ChaosVDSceneQueryDataComponent)
@@ -104,3 +106,11 @@ void UChaosVDSceneQueryDataComponent::ClearData()
 	RecordedQueriesByID.Reset();
 	RecordedQueries.Reset();
 }
+
+void UChaosVDSceneQueryDataComponent::AppendSceneCompositionTestData(FChaosVDSceneCompositionTestData& OutStateTestData)
+{
+	Super::AppendSceneCompositionTestData(OutStateTestData);
+
+	int32& CurrentCount = OutStateTestData.ObjectsCountByType.FindOrAdd(Chaos::VD::Test::SceneObjectTypes::SceneQuery);
+	CurrentCount+= RecordedQueries.Num();
+}
diff --git a/Engine/Plugins/ChaosVD/Source/ChaosVD/Private/Components/ChaosVDSceneQueryDataComponent.h b/Engine/Plugins/ChaosVD/Source/ChaosVD/Private/Components/ChaosVDSceneQueryDataComponent.h
--- a/Engine/Plugins/ChaosVD/Source/ChaosVD/Private/Components/ChaosVDSceneQueryDataComponent.h
+++ b/Engine/Plugins/ChaosVD/Source/ChaosVD/Private/Components/ChaosVDSceneQueryDataComponent.h
@@ -28,6 +28,8 @@ public:
 
 	virtual void ClearData() override;
 
+	virtual void AppendSceneCompositionTestData(FChaosVDSceneCompositionTestData& OutStateTestData) override;
+
 protected:
 
 	void ProccessSQData(const TMap<int32, TSharedPtr<FChaosVDQueryDataWrapper>>* RecordedQueriesByQueryID);
diff --git a/Engine/Plugins/ChaosVD/Source/ChaosVD/Private/Components/ChaosVDSolverCharacterGroundConstraintDataComponent.cpp b/Engine/Plugins/ChaosVD/Source/ChaosVD/Private/Components/ChaosVDSolverCharacterGroundConstraintDataComponent.cpp
--- a/Engine/Plugins/ChaosVD/Source/ChaosVD/Private/Components/ChaosVDSolverCharacterGroundConstraintDataComponent.cpp
+++ b/Engine/Plugins/ChaosVD/Source/ChaosVD/Private/Components/ChaosVDSolverCharacterGroundConstraintDataComponent.cpp
@@ -3,8 +3,10 @@
 #include "Components/ChaosVDSolverCharacterGroundConstraintDataComponent.h"
 
 #include "ChaosVDCharacterGroundConstraintDataProviderInterface.h"
+#include "ChaosVDMiscSceneObjectTypeNamesForTesting.h"
 #include "Components/ChaosVDParticleDataComponent.h"
 #include "ChaosVDScene.h"
+#include "ChaosVDSceneCompositionReport.h"
 #include "ChaosVDSettingsManager.h"
 #include "Selection.h"
 #include "Actors/ChaosVDSolverInfoActor.h"
@@ -70,6 +72,14 @@ void UChaosVDSolverCharacterGroundConstraintDataComponent::BeginDestroy()
 	ScenePtr->OnSceneUpdated().RemoveAll(this);
 }
 
+void UChaosVDSolverCharacterGroundConstraintDataComponent::AppendSceneCompositionTestData(FChaosVDSceneCompositionTestData& OutStateTestData)
+{
+	Super::AppendSceneCompositionTestData(OutStateTestData);
+
+	int32& CurrentCount = OutStateTestData.ObjectsCountByType.FindOrAdd(Chaos::VD::Test::SceneObjectTypes::CharacterGroundConstraint);
+	CurrentCount+= AllConstraints.Num();
+}
+
 void UChaosVDSolverCharacterGroundConstraintDataComponent::SetScene(const TWeakPtr<FChaosVDScene>& InSceneWeakPtr)
 {
 	Super::SetScene(InSceneWeakPtr);
diff --git a/Engine/Plugins/ChaosVD/Source/ChaosVD/Private/Components/ChaosVDSolverCharacterGroundConstraintDataComponent.h b/Engine/Plugins/ChaosVD/Source/ChaosVD/Private/Components/ChaosVDSolverCharacterGroundConstraintDataComponent.h
--- a/Engine/Plugins/ChaosVD/Source/ChaosVD/Private/Components/ChaosVDSolverCharacterGroundConstraintDataComponent.h
+++ b/Engine/Plugins/ChaosVD/Source/ChaosVD/Private/Components/ChaosVDSolverCharacterGroundConstraintDataComponent.h
@@ -19,6 +19,8 @@ public:
 	UE_API void HandleSceneUpdated();
 
 	UE_API virtual void BeginDestroy() override;
+
+	virtual void AppendSceneCompositionTestData(FChaosVDSceneCompositionTestData& OutStateTestData) override;
 };
 
 #undef UE_API
diff --git a/Engine/Plugins/ChaosVD/Source/ChaosVD/Private/Components/ChaosVDSolverJointConstraintDataComponent.cpp b/Engine/Plugins/ChaosVD/Source/ChaosVD/Private/Components/ChaosVDSolverJointConstraintDataComponent.cpp
--- a/Engine/Plugins/ChaosVD/Source/ChaosVD/Private/Components/ChaosVDSolverJointConstraintDataComponent.cpp
+++ b/Engine/Plugins/ChaosVD/Source/ChaosVD/Private/Components/ChaosVDSolverJointConstraintDataComponent.cpp
@@ -2,7 +2,9 @@
 
 #include "Components/ChaosVDSolverJointConstraintDataComponent.h"
 
+#include "ChaosVDMiscSceneObjectTypeNamesForTesting.h"
 #include "ChaosVDRecording.h"
+#include "ChaosVDSceneCompositionReport.h"
 
 #include UE_INLINE_GENERATED_CPP_BY_NAME(ChaosVDSolverJointConstraintDataComponent)
 
@@ -13,3 +15,11 @@ void UChaosVDSolverJointConstraintDataComponent::UpdateFromNewSolverStageData(co
 		UpdateConstraintData(InSolverFrameStageData.RecordedJointConstraints);
 	}
 }
+
+void UChaosVDSolverJointConstraintDataComponent::AppendSceneCompositionTestData(FChaosVDSceneCompositionTestData& OutStateTestData)
+{
+	Super::AppendSceneCompositionTestData(OutStateTestData);
+
+	int32& CurrentCount = OutStateTestData.ObjectsCountByType.FindOrAdd(Chaos::VD::Test::SceneObjectTypes::JointConstraint);
+	CurrentCount+= AllConstraints.Num();
+}
diff --git a/Engine/Plugins/ChaosVD/Source/ChaosVD/Private/Components/ChaosVDSolverJointConstraintDataComponent.h b/Engine/Plugins/ChaosVD/Source/ChaosVD/Private/Components/ChaosVDSolverJointConstraintDataComponent.h
--- a/Engine/Plugins/ChaosVD/Source/ChaosVD/Private/Components/ChaosVDSolverJointConstraintDataComponent.h
+++ b/Engine/Plugins/ChaosVD/Source/ChaosVD/Private/Components/ChaosVDSolverJointConstraintDataComponent.h
@@ -14,6 +14,8 @@ class UChaosVDSolverJointConstraintDataComponent : public UChaosVDConstraintData
 	GENERATED_BODY()
 
 	UE_API virtual void UpdateFromNewSolverStageData(const FChaosVDSolverFrameData& InSolverFrameData, const FChaosVDFrameStageData& InSolverFrameStageData) override;
+
+	virtual void AppendSceneCompositionTestData(FChaosVDSceneCompositionTestData& OutStateTestData) override;
 };
 
 #undef UE_API
diff --git a/Engine/Plugins/ChaosVD/Source/ChaosVD/Public/Actors/ChaosVDDataContainerBaseActor.h b/Engine/Plugins/ChaosVD/Source/ChaosVD/Public/Actors/ChaosVDDataContainerBaseActor.h
--- a/Engine/Plugins/ChaosVD/Source/ChaosVD/Public/Actors/ChaosVDDataContainerBaseActor.h
+++ b/Engine/Plugins/ChaosVD/Source/ChaosVD/Public/Actors/ChaosVDDataContainerBaseActor.h
@@ -5,6 +5,7 @@
 #include "GameFramework/Actor.h"
 #include "ChaosVDDataContainerBaseActor.generated.h"
 
+struct FChaosVDSceneCompositionTestData;
 struct FChaosVDFrameStageData;
 class FChaosVDExtension;
 class FChaosVDScene;
@@ -12,6 +13,11 @@ class FChaosVDScene;
 struct FChaosVDSolverFrameData;
 struct FChaosVDGameFrameData;
 
+namespace Chaos::VD::Test::SceneObjectTypes
+{
+	extern FName SolverID;
+}
+
 /** Base class for any CVD actor that will contain frame related data (either solver frame or game frame) */
 UCLASS(MinimalAPI, Abstract, NotBlueprintable, NotPlaceable)
 class AChaosVDDataContainerBaseActor : public AActor
@@ -90,6 +96,14 @@ public:
 #if WITH_EDITOR
 	CHAOSVD_API void SetIsTemporarilyHiddenInEditor(bool bIsHidden) override;
 #endif
+
+	/**
+	 * Appends Scene Data relevant for functional tests
+	 * @param OutStateTestData Structure instance where to add the data
+	 * @note If you change the data being added here, you need to re-generate the snapshots used by the
+	 * scene integrity playback tests in the Simulation Tests Plugin
+	 */
+	CHAOSVD_API virtual void AppendSceneCompositionTestData(FChaosVDSceneCompositionTestData& OutStateTestData);
 	
 protected:
 
diff --git a/Engine/Plugins/ChaosVD/Source/ChaosVD/Public/ChaosVDEngine.h b/Engine/Plugins/ChaosVD/Source/ChaosVD/Public/ChaosVDEngine.h
new file mode 100644
--- /dev/null
+++ b/Engine/Plugins/ChaosVD/Source/ChaosVD/Public/ChaosVDEngine.h
@@ -0,0 +1,119 @@
+// Copyright Epic Games, Inc. All Rights Reserved.
+
+#pragma once
+
+#include "Templates/SharedPointer.h"
+#include "Containers/Ticker.h"
+#include "Misc/Guid.h"
+#include "Trace/ChaosVDTraceManager.h"
+
+#include "ChaosVDEngine.generated.h"
+
+namespace Chaos::VisualDebugger
+{
+	struct FChaosVDOptionalDataChannel;
+}
+
+class FChaosVDTraceManager;
+class FChaosVDPlaybackController;
+class FChaosVDScene;
+class FChaosVisualDebuggerMainUI;
+
+/** Enumeration of the available modes controlling how data is loaded into CVD */
+UENUM()
+enum class EChaosVDLoadRecordedDataMode : uint8
+{
+	/** This mode will unload any CVD recording currently loaded before loading the selected file */
+	SingleSource,
+	/** CVD will load and merge the data of the selected recording into the currently loaded recording */
+	MultiSource
+};
+
+DECLARE_MULTICAST_DELEGATE_OneParam(FSessionStateChangedDelegate, const FChaosVDTraceSessionDescriptor& InSessionDescriptor)
+
+/** Core Implementation of the visual debugger - Owns the systems that are not UI */
+class FChaosVDEngine : public FTSTickerObjectBase, public TSharedFromThis<FChaosVDEngine>
+{
+public:
+	FChaosVDEngine()
+	{
+		InstanceGUID = FGuid::NewGuid();
+	}
+
+	CHAOSVD_API void Initialize();
+	CHAOSVD_API void CloseActiveTraceSessions();
+
+	void StopActiveTraceSessions();
+
+	CHAOSVD_API void DeInitialize();
+
+	CHAOSVD_API virtual bool Tick(float DeltaTime) override;
+
+	const FGuid& GetInstanceGuid() const
+	{
+		return InstanceGUID;
+	}
+
+	TSharedPtr<FChaosVDScene>& GetCurrentScene()
+	{
+		return CurrentScene;
+	}
+
+	TSharedPtr<FChaosVDPlaybackController>& GetPlaybackController()
+	{
+		return PlaybackController;
+	}
+
+	TArrayView<FChaosVDTraceSessionDescriptor> GetCurrentSessionDescriptors()
+	{
+		return CurrentSessionDescriptors;
+	}
+
+	bool HasAnyLiveSessionActive() const;
+
+	CHAOSVD_API void LoadRecording(const FString& FilePath, EChaosVDLoadRecordedDataMode LoadingMode = EChaosVDLoadRecordedDataMode::MultiSource);
+	void LoadCombinedMultiRecording(const FString& FilePath);
+
+	bool ConnectToLiveSession(uint32 SessionID, const FString& InSessionAddress, EChaosVDLoadRecordedDataMode LoadingMode = EChaosVDLoadRecordedDataMode::MultiSource);
+
+	void OpenSession(const FChaosVDTraceSessionDescriptor& SessionDescriptor, EChaosVDLoadRecordedDataMode LoadingMode = EChaosVDLoadRecordedDataMode::MultiSource);
+
+	bool SaveOpenSessionToCombinedFile(const FString& InTargetFilePath = FString());
+
+	bool CanCombineOpenSessions() const;
+	
+	FSessionStateChangedDelegate& OnSessionOpened()
+	{
+		return OnSessionOpenedDelegate;
+	}
+
+	FSessionStateChangedDelegate& OnSessionClosed()
+	{
+		return OnSessionClosedDelegate;
+	}
+
+private:
+
+	void LoadRecording_Internal(const TFunction<FString(const TSharedPtr<FChaosVDRecording>&)>&, EChaosVDLoadRecordedDataMode LoadingMode = EChaosVDLoadRecordedDataMode::MultiSource);
+
+	void RestoreDataChannelsEnabledStateFromSave();
+	void UpdateSavedDataChannelsEnabledState(TWeakPtr<Chaos::VisualDebugger::FChaosVDOptionalDataChannel> DataChannelChanged);
+
+	void UpdateRecentFilesList(const FString& InFilename);
+
+	FGuid InstanceGUID;
+
+	TArray<FChaosVDTraceSessionDescriptor> CurrentSessionDescriptors;
+
+	TSharedPtr<FChaosVDScene> CurrentScene;
+	TSharedPtr<FChaosVDPlaybackController> PlaybackController;
+	
+	bool bIsInitialized = false;
+
+	FDelegateHandle LiveSessionStoppedDelegateHandle;
+
+	FDelegateHandle DataChannelStateUpdatedHandle;
+	
+	FSessionStateChangedDelegate OnSessionOpenedDelegate;
+	FSessionStateChangedDelegate OnSessionClosedDelegate;
+};
diff --git a/Engine/Plugins/ChaosVD/Source/ChaosVD/Public/ChaosVDMiscSceneObjectTypeNamesForTesting.h b/Engine/Plugins/ChaosVD/Source/ChaosVD/Public/ChaosVDMiscSceneObjectTypeNamesForTesting.h
new file mode 100644
--- /dev/null
+++ b/Engine/Plugins/ChaosVD/Source/ChaosVD/Public/ChaosVDMiscSceneObjectTypeNamesForTesting.h
@@ -0,0 +1,13 @@
+// Copyright Epic Games, Inc. All Rights Reserved.
+
+#pragma once
+
+#include "UObject/NameTypes.h"
+
+namespace Chaos::VD::Test::SceneObjectTypes
+{
+	extern FName SceneQuery;
+	extern FName CollisionConstraint;
+	extern FName JointConstraint;
+	extern FName CharacterGroundConstraint;
+}
diff --git a/Engine/Plugins/ChaosVD/Source/ChaosVD/Public/ChaosVDModule.h b/Engine/Plugins/ChaosVD/Source/ChaosVD/Public/ChaosVDModule.h
--- a/Engine/Plugins/ChaosVD/Source/ChaosVD/Public/ChaosVDModule.h
+++ b/Engine/Plugins/ChaosVD/Source/ChaosVD/Public/ChaosVDModule.h
@@ -41,6 +41,8 @@ public:
 	/** Re-instantiates CVD's widgets without restarting the underlying CVD engine instance */
 	void ReloadInstanceUI(FGuid InstanceGUID);
 
+	void TakePlaybackEngineStateSnapshot();
+
 private:
 
 	TSharedRef<SDockTab> SpawnMainTab(const FSpawnTabArgs& Args);
diff --git a/Engine/Plugins/ChaosVD/Source/ChaosVD/Public/ChaosVDPlaybackController.h b/Engine/Plugins/ChaosVD/Source/ChaosVD/Public/ChaosVDPlaybackController.h
--- a/Engine/Plugins/ChaosVD/Source/ChaosVD/Public/ChaosVDPlaybackController.h
+++ b/Engine/Plugins/ChaosVD/Source/ChaosVD/Public/ChaosVDPlaybackController.h
@@ -8,9 +8,9 @@
 #include "Templates/SharedPointer.h"
 #include "Delegates/DelegateCombinations.h"
 #include "Delegates/Delegate.h"
-#include "Widgets/SChaosVDTimelineWidget.h"
 #include "ChaosVDPlaybackController.generated.h"
 
+struct FChaosVDPlaybackEngineSnapshot;
 class IChaosVDPlaybackControllerInstigator;
 class UChaosVDCoreSettings;
 struct FChaosVDTraceSessionDescriptor;
@@ -20,6 +20,8 @@ struct FChaosVDRecording;
 class FChaosVDPlaybackController;
 class FString;
 
+enum class EChaosVDPlaybackButtonsID : uint8;
+
 DECLARE_MULTICAST_DELEGATE_OneParam(FChaosVDPlaybackControllerUpdated, TWeakPtr<FChaosVDPlaybackController>)
 DECLARE_MULTICAST_DELEGATE_ThreeParams(FChaosVDPlaybackControllerFrameUpdated, TWeakPtr<FChaosVDPlaybackController>, TWeakPtr<const FChaosVDTrackInfo>, FGuid);
 
@@ -145,7 +147,7 @@ public:
 	 * @param FrameNumber Frame number to go to
 	 * @param StageNumber Stage number to go to
 	 */
-	void GoToTrackFrame(FGuid InstigatorID, EChaosVDTrackType TrackType, int32 InTrackID, int32 FrameNumber, int32 StageNumber);
+	CHAOSVD_API void GoToTrackFrame(FGuid InstigatorID, EChaosVDTrackType TrackType, int32 InTrackID, int32 FrameNumber, int32 StageNumber);
 	void GoToTrackFrame_AssumesLocked(FGuid InstigatorID, EChaosVDTrackType TrackType, int32 InTrackID, int32 FrameNumber, int32 StageNumber);
 
 	/**
@@ -156,7 +158,7 @@ public:
 	 * @param FrameNumber Frame number to go to
 	 * @param StageNumber Stage number to go to
 	 */
-	void GoToTrackFrameAndSync(FGuid InstigatorID, EChaosVDTrackType TrackType, int32 InTrackID, int32 FrameNumber, int32 StageNumber);
+	CHAOSVD_API void GoToTrackFrameAndSync(FGuid InstigatorID, EChaosVDTrackType TrackType, int32 InTrackID, int32 FrameNumber, int32 StageNumber);
 	void GoToTrackFrame_AssumesLockedAndSync(FGuid InstigatorID, EChaosVDTrackType TrackType, int32 InTrackID, int32 FrameNumber, int32 StageNumber);
 
 	/**
@@ -285,7 +287,7 @@ public:
 		return ControllerFrameUpdatedDelegate;
 	}
 	
-	virtual bool Tick(float DeltaTime) override;
+	CHAOSVD_API virtual bool Tick(float DeltaTime) override;
 
 	/** Returns true if we are playing a live debugging session */
 	bool IsPlayingLiveSession() const;
@@ -425,7 +427,7 @@ public:
 	bool IsCompatibleWithSyncMode(const TSharedRef<const FChaosVDTrackInfo>& InTrackRef, EChaosVDSyncTimelinesMode SyncMode);
 
 	/** Sets a new sync mode to be used between tracks */
-	void SetTimelineSyncMode(EChaosVDSyncTimelinesMode SyncMode);
+	CHAOSVD_API void SetTimelineSyncMode(EChaosVDSyncTimelinesMode SyncMode);
 
 	/** Attempts to set the track on the specified slot as active.
 	 * An active track is the track to all key shortcuts will be applied
@@ -439,6 +441,15 @@ public:
 	 */
 	void TrySetActiveTrack(const TSharedRef<const FChaosVDTrackInfo>& NewActiveTrack);
 
+	/**
+	 * Creates and populates an object containing information about the state of the playback engine and scene.
+	 * This is used by functional tests to compare a known good state agains ones generated in future versions
+	 * @return Playback engine state data
+	 * @note If you change the data being added here, you need to re-generate the snapshots used by the
+	 * scene integrity playback tests in the Simulation Tests Plugin
+	 */
+	CHAOSVD_API FChaosVDPlaybackEngineSnapshot GeneratePlaybackEngineSnapshot();
+
 protected:
 	/** Updates (or adds) solvers data from the loaded recording to the solver tracks */
 	void UpdateSolverTracksData();
diff --git a/Engine/Plugins/ChaosVD/Source/ChaosVD/Public/ChaosVDScene.h b/Engine/Plugins/ChaosVD/Source/ChaosVD/Public/ChaosVDScene.h
--- a/Engine/Plugins/ChaosVD/Source/ChaosVD/Public/ChaosVDScene.h
+++ b/Engine/Plugins/ChaosVD/Source/ChaosVD/Public/ChaosVDScene.h
@@ -16,6 +16,7 @@
 #include "UObject/ObjectPtr.h"
 #include "ChaosVDScene.generated.h"
 
+struct FChaosVDPlaybackEngineSnapshot;
 class AChaosVDDataContainerBaseActor;
 class AChaosVDGeometryContainer;
 class IChaosVDGeometryOwnerInterface;
@@ -225,7 +226,7 @@ public:
 
 	void RequestUpdate();
 
-	virtual bool Tick(float DeltaTime) override;
+	CHAOSVD_API virtual bool Tick(float DeltaTime) override;
 
 	void UpdateWorldStreamingLocation(const FVector& InLocation);
 
@@ -234,6 +235,14 @@ public:
 		return WorldStreamingLocation;
 	}
 
+	/**
+	 * Appends Scene Data relevant for functional tests
+	 * @param OutStateTestData Structure instance where to add the data
+	 * @note If you change the data being added here, you need to re-generate the snapshots used by the
+	 * scene integrity playback tests in the Simulation Tests Plugin
+	 */
+	void AppendSceneCompositionTestData(FChaosVDPlaybackEngineSnapshot& OutStateTestData);
+
 private:
 
 	FVector WorldStreamingLocation = FVector::ZeroVector; 
diff --git a/Engine/Plugins/ChaosVD/Source/ChaosVD/Public/ChaosVDSceneCompositionReport.h b/Engine/Plugins/ChaosVD/Source/ChaosVD/Public/ChaosVDSceneCompositionReport.h
new file mode 100644
--- /dev/null
+++ b/Engine/Plugins/ChaosVD/Source/ChaosVD/Public/ChaosVDSceneCompositionReport.h
@@ -0,0 +1,62 @@
+// Copyright Epic Games, Inc. All Rights Reserved.
+
+#pragma once
+
+#include "Containers/Array.h"
+#include "Containers/Map.h"
+#include "Containers/UnrealString.h"
+#include "UObject/ObjectMacros.h"
+#include "UObject/NameTypes.h"
+#include "ChaosVDPlaybackController.h"
+
+#include "ChaosVDSceneCompositionReport.generated.h"
+
+USTRUCT()
+struct FChaosVDFrameIndexTestData
+{
+	GENERATED_BODY()
+
+	UPROPERTY(VisibleAnywhere, Category=SceneComposition)
+	int32 FrameNumber = INDEX_NONE;
+
+	UPROPERTY(VisibleAnywhere, Category=SceneComposition)
+	int32 StageNumber = INDEX_NONE;
+};
+
+/**
+ * Structure containing data about how a CVD scene is composed.
+ * This is used for functional testing, to verify if a loaded file at a specific frame keeps the expected composition
+ */
+USTRUCT()
+struct FChaosVDSceneCompositionTestData
+{
+	GENERATED_BODY()
+
+	UPROPERTY(VisibleAnywhere, Category=SceneComposition)
+	TMap<FName, int32> ObjectsCountByType;
+};
+
+/**
+ * Structure containing data about the playback engine state at the time when this struct is created.
+ * This is used for functional testing, to verify if a loaded file at a specific frame keeps the expected composition
+ *
+ * @note If you change the composition of this structure or of the structures used by it, you need to re-generate the snapshots used by the
+ * scene integrity playback tests in the Simulation Tests Plugin
+ */
+USTRUCT()
+struct FChaosVDPlaybackEngineSnapshot
+{
+	GENERATED_BODY()
+
+	UPROPERTY(VisibleAnywhere, Category=SceneComposition)
+	TArray<FName> InstalledExtensionsNames;
+
+	UPROPERTY(VisibleAnywhere, Category=SceneComposition)
+	TMap<int32, FChaosVDFrameIndexTestData> FrameIndexDataByTrackID;
+
+	UPROPERTY(VisibleAnywhere, Category=SceneComposition)
+	EChaosVDSyncTimelinesMode TimelineSyncMode = EChaosVDSyncTimelinesMode::RecordedTimestamp;
+
+	UPROPERTY(VisibleAnywhere, Category=SceneComposition)
+	FChaosVDSceneCompositionTestData SceneComposition;
+};
diff --git a/Engine/Plugins/ChaosVD/Source/ChaosVD/Public/Components/ChaosVDParticleDataComponent.h b/Engine/Plugins/ChaosVD/Source/ChaosVD/Public/Components/ChaosVDParticleDataComponent.h
--- a/Engine/Plugins/ChaosVD/Source/ChaosVD/Public/Components/ChaosVDParticleDataComponent.h
+++ b/Engine/Plugins/ChaosVD/Source/ChaosVD/Public/Components/ChaosVDParticleDataComponent.h
@@ -16,6 +16,11 @@ class FChaosVDSceneStreaming;
 struct FChaosVDParticleDataWrapper;
 struct FChaosVDSceneParticle;
 
+namespace Chaos::VD::Test::SceneObjectTypes
+{
+	extern FName ActiveParticles;
+}
+
 /**
  * Component that references all particle data for a specific solver for the current frame, and handles how the visualization
  * is updated based on that data
@@ -78,6 +83,8 @@ public:
 	}
 	// END IChaosVDStreamingDataSource
 
+	virtual void AppendSceneCompositionTestData(FChaosVDSceneCompositionTestData& OutStateTestData) override;
+
 protected:
 
 	int32 GetMaxElementCountForBatchTEDSUpdate_AssumesLocked() const;
diff --git a/Engine/Plugins/ChaosVD/Source/ChaosVD/Public/Components/ChaosVDSolverDataComponent.h b/Engine/Plugins/ChaosVD/Source/ChaosVD/Public/Components/ChaosVDSolverDataComponent.h
--- a/Engine/Plugins/ChaosVD/Source/ChaosVD/Public/Components/ChaosVDSolverDataComponent.h
+++ b/Engine/Plugins/ChaosVD/Source/ChaosVD/Public/Components/ChaosVDSolverDataComponent.h
@@ -5,6 +5,7 @@
 #include "Components/ActorComponent.h"
 #include "ChaosVDSolverDataComponent.generated.h"
 
+struct FChaosVDSceneCompositionTestData;
 struct FChaosVDFrameStageData;
 struct FChaosVDSolverFrameData;
 class FChaosVDScene;
@@ -38,6 +39,14 @@ public:
 		SolverID = InSolverID;
 	}
 
+	/**
+	 * Appends Scene Data relevant for functional tests
+	 * @param OutStateTestData Structure instance where to add the data
+	 * @note If you change the data being added here, you need to re-generate the snapshots used by the
+	 * scene integrity playback tests in the Simulation Tests Plugin
+	 */
+	virtual void AppendSceneCompositionTestData(FChaosVDSceneCompositionTestData& OutStateTestData) {}
+
 protected:
 
 	TWeakPtr<FChaosVDScene> SceneWeakPtr;
