[ChaosVD] Removing duplicated localization key in ChaosVDModule.cpp

diff --git a/Engine/Plugins/ChaosVD/Source/ChaosVD/Private/ChaosVDModule.cpp b/Engine/Plugins/ChaosVD/Source/ChaosVD/Private/ChaosVDModule.cpp
--- a/Engine/Plugins/ChaosVD/Source/ChaosVD/Private/ChaosVDModule.cpp
+++ b/Engine/Plugins/ChaosVD/Source/ChaosVD/Private/ChaosVDModule.cpp
@@ -275,7 +275,7 @@ void FChaosVDModule::TakePlaybackEngineStateSnapshot()
 		
 		DesktopPlatform->SaveFileDialog(
 			FSlateApplication::Get().FindBestParentWindowHandleForDialogs(nullptr),
-			LOCTEXT("SaveDialogTitle", "Dump Playback Engine Sate").ToString(),
+			LOCTEXT("PlaybackSnapshotSaveDialogTitle", "Dump Playback Engine Sate").ToString(),
 			*FPaths::ProfilingDir(),
 			TEXT(""),
 			*ExtensionStr,
