[ChaosVD] Adding a pop up notification to let the user know if we are still processing data for the selected trace file

diff --git a/Engine/Plugins/ChaosVD/Source/ChaosVD/Private/Widgets/SChaosVDMainTab.cpp b/Engine/Plugins/ChaosVD/Source/ChaosVD/Private/Widgets/SChaosVDMainTab.cpp
--- a/Engine/Plugins/ChaosVD/Source/ChaosVD/Private/Widgets/SChaosVDMainTab.cpp
+++ b/Engine/Plugins/ChaosVD/Source/ChaosVD/Private/Widgets/SChaosVDMainTab.cpp
@@ -69,6 +69,7 @@
 #include "Widgets/Docking/SDockTab.h"
 #include "Widgets/Input/SComboButton.h"
 #include "Widgets/Layout/SSeparator.h"
+#include "Widgets/Notifications/SNotificationList.h"
 #include "Widgets/SToolTip.h"
 
 #include UE_INLINE_GENERATED_CPP_BY_NAME(SChaosVDMainTab)
@@ -330,6 +331,96 @@ const TSharedPtr<FChaosVDIndependentDetailsPanelManager>& SChaosVDMainTab::GetIn
 	return IndependentDetailsPanelManager;
 }
 
+void SChaosVDMainTab::UpdateTraceProcessingNotificationUpdate()
+{
+	bool bIsTraceAnalysisInProgress = false;
+	TSharedRef<FChaosVDEngine> CVDEngineInstance = GetChaosVDEngineInstance();
+
+	// Live sessions should not show the notification pop up. When CVD is connected in live debugging mode the UI already has feedback to inform the user.
+	if (CVDEngineInstance->GetCurrentSessionDescriptors().Num() > 0 && !CVDEngineInstance->HasAnyLiveSessionActive())
+	{
+		FChaosVDModule::Get().GetTraceManager()->EnumerateActiveSessions([&bIsTraceAnalysisInProgress](const TSharedRef<const TraceServices::IAnalysisSession>& InSession)
+		{
+			// TODO: We should implement a way to gather the actual progress so we can report that in the UI
+			// but the Trace API don't have that data yet.
+			if (!InSession->IsAnalysisComplete())
+			{
+				bIsTraceAnalysisInProgress = true;
+				return false;
+			}
+
+			return true;
+		});
+	}
+
+	if (bIsTraceAnalysisInProgress)
+	{
+		ShowTraceDataProcessingNotification();
+	}
+	else
+	{
+		HideTraceDataProcessingNotification();
+	}
+}
+
+void SChaosVDMainTab::Tick(const FGeometry& AllottedGeometry, const double InCurrentTime, const float InDeltaTime)
+{
+	UpdateTraceProcessingNotificationUpdate();
+
+	SCompoundWidget::Tick(AllottedGeometry, InCurrentTime, InDeltaTime);
+}
+
+TSharedPtr<SNotificationItem> SChaosVDMainTab::CreateProcessingTraceDataNotification()
+{
+	FNotificationInfo Info(LOCTEXT("ProcessingTraceSessionMessage", "Processing Trace Data ..."));
+	Info.bFireAndForget = false;
+	Info.FadeOutDuration = 3.0f;
+	Info.ExpireDuration = 0.0f;
+
+	{
+		TSharedPtr<SDockTab> OwnerTabPtr = OwnerTab.Pin();
+		Info.ForWindow = OwnerTabPtr ? OwnerTabPtr->GetParentWindow() : nullptr;
+	}
+
+	TSharedPtr<SNotificationItem> ProcessingDataNotification = FSlateNotificationManager::Get().AddNotification(Info);
+
+	if (ProcessingDataNotification.IsValid())
+	{
+		ProcessingDataNotification->SetCompletionState(SNotificationItem::CS_Pending);
+		return ProcessingDataNotification;
+	}
+
+	return nullptr;
+}
+
+void SChaosVDMainTab::ShowTraceDataProcessingNotification()
+{
+	if (ActiveProcessingTraceDataNotification)
+	{
+		if (ActiveProcessingTraceDataNotification->GetCompletionState() != SNotificationItem::CS_Pending)
+		{
+			// The existing notification already expired, we need to create a new one
+			ActiveProcessingTraceDataNotification = CreateProcessingTraceDataNotification();
+		}
+	}
+	else
+	{
+		ActiveProcessingTraceDataNotification = CreateProcessingTraceDataNotification();
+	}
+
+	ActiveProcessingTraceDataNotification->SetCompletionState(SNotificationItem::CS_Pending);
+}
+
+void SChaosVDMainTab::HideTraceDataProcessingNotification()
+{
+	if (ActiveProcessingTraceDataNotification)
+	{
+		ActiveProcessingTraceDataNotification->SetCompletionState(SNotificationItem::CS_Success);
+		ActiveProcessingTraceDataNotification->ExpireAndFadeout();
+		ActiveProcessingTraceDataNotification = nullptr;
+	}
+}
+
 void SChaosVDMainTab::HandlePersistLayout(const TSharedRef<FTabManager::FLayout>& InLayoutToSave)
 {
 	if (!bCanTabManagerPersistLayout)
diff --git a/Engine/Plugins/ChaosVD/Source/ChaosVD/Public/Trace/ChaosVDTraceManager.h b/Engine/Plugins/ChaosVD/Source/ChaosVD/Public/Trace/ChaosVDTraceManager.h
--- a/Engine/Plugins/ChaosVD/Source/ChaosVD/Public/Trace/ChaosVDTraceManager.h
+++ b/Engine/Plugins/ChaosVD/Source/ChaosVD/Public/Trace/ChaosVDTraceManager.h
@@ -71,8 +71,12 @@ public:
 
 	template<typename TVisitor>
 	static void EnumerateActiveSessions(FStringView InSessionHost, TVisitor Callback);
+
 	static const UE::Trace::FStoreClient::FSessionInfo* GetTraceSessionInfo(FStringView InSessionHost, FGuid TraceGuid);
 
+	template<typename TVisitor>
+	void EnumerateActiveSessions(TVisitor Callback);
+
 private:
 
 	TUniquePtr<UE::Trace::IInDataStream> CreateFileDataStream(TUniquePtr<IFileHandle>&& InFileHandle);
@@ -118,3 +122,18 @@ void FChaosVDTraceManager::EnumerateActiveSessions(FStringView InSessionHost, TC
 		}	
 	}
 }
+
+template <typename TVisitor>
+void FChaosVDTraceManager::EnumerateActiveSessions(TVisitor Callback)
+{
+	for (const TPair<FString, TSharedPtr<const TraceServices::IAnalysisSession>>& SessionWithName : AnalysisSessionByName)
+	{
+		if (SessionWithName.Value)
+		{
+			if (!Callback(SessionWithName.Value.ToSharedRef()))
+			{
+				return;
+			}
+		}
+	}
+}
diff --git a/Engine/Plugins/ChaosVD/Source/ChaosVD/Public/Widgets/SChaosVDMainTab.h b/Engine/Plugins/ChaosVD/Source/ChaosVD/Public/Widgets/SChaosVDMainTab.h
--- a/Engine/Plugins/ChaosVD/Source/ChaosVD/Public/Widgets/SChaosVDMainTab.h
+++ b/Engine/Plugins/ChaosVD/Source/ChaosVD/Public/Widgets/SChaosVDMainTab.h
@@ -156,8 +156,16 @@ public:
 
 	const TSharedPtr<FChaosVDIndependentDetailsPanelManager>& GetIndependentDetailsPanelManager();
 
+	virtual void Tick(const FGeometry& AllottedGeometry, const double InCurrentTime, const float InDeltaTime) override;
+
 private:
 
+	TSharedPtr<SNotificationItem> CreateProcessingTraceDataNotification();
+
+	void UpdateTraceProcessingNotificationUpdate();
+	void ShowTraceDataProcessingNotification();
+	void HideTraceDataProcessingNotification();
+
 	void HandlePersistLayout(const TSharedRef<FTabManager::FLayout>& InLayoutToSave);
 
 	void HandlePostInitializationExtensionRegistered(const TSharedRef<FChaosVDExtension>& NewExtension);
@@ -195,6 +203,8 @@ private:
 
 	bool ShouldDisableCPUThrottling() const;
 
+	TSharedPtr<SNotificationItem> ActiveProcessingTraceDataNotification;
+
 	bool bCanTabManagerPersistLayout = true;
 	
 	TArray<TWeakPtr<IDetailsView>> CustomizedDetailsPanels;
