[ChaosVD] Added an async progress bar notification for the CVD file loading process.

This also works for CVD multi file. In that case, CVD will show a separate progress bar per trace inside the multi file container, in the status bar.

diff --git a/Engine/Plugins/ChaosVD/Source/ChaosVD/Private/Trace/ChaosVDTraceManager.cpp b/Engine/Plugins/ChaosVD/Source/ChaosVD/Private/Trace/ChaosVDTraceManager.cpp
--- a/Engine/Plugins/ChaosVD/Source/ChaosVD/Private/Trace/ChaosVDTraceManager.cpp
+++ b/Engine/Plugins/ChaosVD/Source/ChaosVD/Private/Trace/ChaosVDTraceManager.cpp
@@ -15,6 +15,72 @@
 #include "TraceServices/ITraceServicesModule.h"
 #include "TraceServices/Model/AnalysisSession.h"
 
+namespace Chaos::VD
+{
+	FAsyncProgressNotification::~FAsyncProgressNotification()
+	{
+		if (NotificationHandle.IsValid())
+		{
+			// FSlateNotificationManager is not thread safe, so we need to make sure this runs in the game thread
+			if (IsInGameThread())
+			{
+				FSlateNotificationManager::Get().CancelProgressNotification(NotificationHandle);
+			}
+			else
+			{
+				constexpr float Delay = 0.0f;
+                FTSTicker::GetCoreTicker().AddTicker(TEXT("CancelCVDAsyncNotification"), Delay, [HandleCopy = NotificationHandle](float DeltaTime)
+                {
+                	FSlateNotificationManager::Get().CancelProgressNotification(HandleCopy);
+                	return false;
+                });
+			}
+		}
+	}
+
+	void FAsyncProgressNotification::EnterProgress(int32 InCurrentProgress)
+	{
+		CurrentProgress = InCurrentProgress;
+	}
+
+	bool FAsyncProgressNotification::Tick(float DeltaTime)
+	{
+		constexpr int32 MaxProgress = 100;
+		constexpr float TimeBetweenUpdatesSeconds = 0.5f;
+
+		ElapsedTimeSinceLastUpdate += DeltaTime;
+		if (ElapsedTimeSinceLastUpdate < TimeBetweenUpdatesSeconds)
+		{
+			return true;
+		}
+
+		ElapsedTimeSinceLastUpdate = 0.0f;
+		
+		if (CurrentProgress >= MaxProgress)
+		{
+			if (NotificationHandle.IsValid())
+			{
+				FSlateNotificationManager::Get().UpdateProgressNotification(NotificationHandle, MaxProgress, MaxProgress, NotificationTitle);
+			}
+			NotificationHandle = FProgressNotificationHandle();
+
+			// We are done and we can stop ticking
+			return false;
+		}
+
+		if (!NotificationHandle.IsValid())
+		{
+			NotificationHandle = FSlateNotificationManager::Get().StartProgressNotification(NotificationTitle, MaxProgress);
+		}
+		else
+		{
+			FSlateNotificationManager::Get().UpdateProgressNotification(NotificationHandle, CurrentProgress, MaxProgress, NotificationTitle);
+		}
+
+		return true;
+	}
+}
+
 FChaosVDTraceManager::FChaosVDTraceManager() 
 {
 	ChaosVDTraceModule = MakeShared<FChaosVDTraceModule>();
@@ -33,24 +99,25 @@ FString FChaosVDTraceManager::LoadTraceFile(const FString& InTraceFilename, cons
 {
 	CloseSession(InTraceFilename);
 
-	ITraceServicesModule& TraceServicesModule = FModuleManager::LoadModuleChecked<ITraceServicesModule>("TraceServices");
-	if (const TSharedPtr<TraceServices::IAnalysisService> TraceAnalysisService = TraceServicesModule.GetAnalysisService())
+	IPlatformFile& FileSystem = IPlatformFile::GetPlatformPhysical();
+	
+	constexpr bool bAllowWrite = false;
+	IFileHandle* Handle = FileSystem.OpenRead(*InTraceFilename, bAllowWrite);
+	if (!Handle)
 	{
-		SetPendingExternalRecordingToProcess(ExistingRecordingPtr);
-
-		if (const TSharedPtr<const TraceServices::IAnalysisSession> NewSession = TraceAnalysisService->StartAnalysis(*InTraceFilename))
-		{
-			AnalysisSessionByName.Add(InTraceFilename, NewSession.ToSharedRef());
-
-			return NewSession->GetName();
-		}
+		return FString();
 	}
 
-	return FString();
+	return LoadTraceFile(TUniquePtr<IFileHandle>(Handle), InTraceFilename, ExistingRecordingPtr);
 }
 
 FString FChaosVDTraceManager::LoadTraceFile(TUniquePtr<IFileHandle>&& InFileHandle, const FString& InTraceSessionName, const TSharedPtr<FChaosVDRecording>& ExistingRecordingPtr)
 {
+	if (InFileHandle->Size() == 0)
+	{
+		return FString();
+	}
+
 	ITraceServicesModule& TraceServicesModule = FModuleManager::LoadModuleChecked<ITraceServicesModule>("TraceServices");
 	if (const TSharedPtr<TraceServices::IAnalysisService> TraceAnalysisService = TraceServicesModule.GetAnalysisService())
 	{
@@ -91,19 +158,39 @@ TUniquePtr<UE::Trace::IInDataStream> FChaosVDTraceManager::CreateFileDataStream(
 {
 	check(InFileHandle);
 
-	struct FFileDataStream : public UE::Trace::FFileDataStream
+	struct FFileDataStream : public UE::Trace::IInDataStream
 	{
+		explicit FFileDataStream(TUniquePtr<IFileHandle>&& FileHandle) :
+			Handle(MoveTemp(FileHandle)),
+			AsyncProgressNotification(Chaos::VD::FAsyncProgressNotification(NSLOCTEXT("ChaosVisualDebugger", "FileLoadingProgressNotification","Loading File")))
+		{
+			check(Handle);
+			FileSize = Handle->Size();
+			RemainingDataSize = FileSize;
+			constexpr int32 FileFullyLoadedPercentage = 100;
+
+			check(FileSize > 0);
+			ProgressPercentagePerByte = static_cast<double>(FileFullyLoadedPercentage) / static_cast<double>(FileSize);
+		}
+		
 		virtual int32 Read(void* Data, uint32 Size) override
 		{
-			if (Remaining <= 0)
+			if (RemainingDataSize <= 0)
 			{
 				return 0;
 			}
-			if (Size > Remaining)
+
+			if (Size > RemainingDataSize)
 			{
-				Size = static_cast<uint32>(Remaining);
+				Size = static_cast<uint32>(RemainingDataSize);
 			}
-			Remaining -= Size;
+		
+			RemainingDataSize -= Size;
+
+			double ProgressPercentage = static_cast<double>(FileSize - RemainingDataSize) * ProgressPercentagePerByte;
+			AsyncProgressNotification.EnterProgress(FMath::CeilToInt32(ProgressPercentage));
+
+			check(Handle);
 			if (!Handle->Read(static_cast<uint8*>(Data), Size))
 			{
 				return 0;
@@ -111,13 +198,22 @@ TUniquePtr<UE::Trace::IInDataStream> FChaosVDTraceManager::CreateFileDataStream(
 			return Size;
 		}
 
+		virtual void Close() override
+		{
+			Handle.Reset();
+		}
+
+	private:
+		
 		TUniquePtr<IFileHandle> Handle;
-		uint64 Remaining;
+		uint64 RemainingDataSize = 0;
+		uint64 FileSize = 0;
+		double ProgressPercentagePerByte = 0.0;
+	
+		Chaos::VD::FAsyncProgressNotification AsyncProgressNotification;
 	};
 
-	FFileDataStream* FileStream = new FFileDataStream();
-	FileStream->Handle = MoveTemp(InFileHandle);
-	FileStream->Remaining = FileStream->Handle->Size();
+	FFileDataStream* FileStream = new FFileDataStream(MoveTemp(InFileHandle));
 
 	TUniquePtr<UE::Trace::IInDataStream> DataStream(FileStream);
 
diff --git a/Engine/Plugins/ChaosVD/Source/ChaosVD/Public/Trace/ChaosVDTraceManager.h b/Engine/Plugins/ChaosVD/Source/ChaosVD/Public/Trace/ChaosVDTraceManager.h
--- a/Engine/Plugins/ChaosVD/Source/ChaosVD/Public/Trace/ChaosVDTraceManager.h
+++ b/Engine/Plugins/ChaosVD/Source/ChaosVD/Public/Trace/ChaosVDTraceManager.h
@@ -4,6 +4,7 @@
 
 #include "ChaosVDModule.h"
 #include "Containers/StringFwd.h"
+#include "Framework/Notifications/NotificationManager.h"
 #include "Templates/SharedPointer.h"
 #include "Trace/StoreClient.h"
 
@@ -13,6 +14,35 @@ class FChaosVDTraceModule;
 
 namespace TraceServices { class IAnalysisSession; }
 
+namespace Chaos::VD
+{
+	/** Object that allows showing a progress bar in the status bar of the editor,
+	 * where the progress can be updated from any thread */
+	struct FAsyncProgressNotification : FTSTickerObjectBase
+	{
+		FAsyncProgressNotification(const FText& InMessage) : NotificationTitle(InMessage)
+		{
+		}
+
+
+		virtual ~FAsyncProgressNotification() override;
+
+		/** Updates the current progress
+		 * @param InCurrentProgress Value between 0 and 100, representing the progress we want to display
+		 */
+		void EnterProgress(int32 InCurrentProgress);
+
+	private:
+		virtual bool Tick(float DeltaTime) override;
+
+		float ElapsedTimeSinceLastUpdate = 0.0f;
+
+		FProgressNotificationHandle NotificationHandle;
+		std::atomic<int32> CurrentProgress = 0;
+		FText NotificationTitle;
+	};
+}
+
 /**
  * Structure containing info about a trace session used by Chaos Visual Debugger
  */
