[ChaosVD] Combining the two existing TEDS batch operation if we need to set the Active tag and the SyncFromWorld tag at the same time.

Before this change, we would batch apply either tag to a collection of particles. One batch per tag. But if you need to apply both tags it means you can end up issuing double of the required commands in TEDS' command buffer.

As TEDS operations support setting multiple tags at the same time, we can cut in half the number of commands we queue if we combine both existing batches

diff --git a/Engine/Plugins/ChaosVD/Source/ChaosVD/Private/Components/ChaosVDParticleDataComponent.cpp b/Engine/Plugins/ChaosVD/Source/ChaosVD/Private/Components/ChaosVDParticleDataComponent.cpp
--- a/Engine/Plugins/ChaosVD/Source/ChaosVD/Private/Components/ChaosVDParticleDataComponent.cpp
+++ b/Engine/Plugins/ChaosVD/Source/ChaosVD/Private/Components/ChaosVDParticleDataComponent.cpp
@@ -261,8 +261,31 @@ void UChaosVDParticleDataComponent::BatchTEDSSyncFromWorldTag(FTedsHandlesForBat
 	}
 }
 
+void UChaosVDParticleDataComponent::BatchTEDSCombinedSyncFromWorldAndActiveTag(FTedsHandlesForBatchArray& ParticlesPendingCombinedTEDSTagSync, const UE::Editor::DataStorage::ICompatibilityProvider* CompatibilityStorage, UE::Editor::DataStorage::ICoreProvider* DataStorage)
+{
+	FMemMark StackMarker(FMemStack::Get());
+	FTedsHandlesForBatchArray OwnersPendingTEDSWorldSyncTag;
+	if (CompatibilityStorage && DataStorage)
+	{
+		// We need to mark our owner as pending sync to keep the hierarchy up to date because our owner is not a CVDScene Object, we also need to manually flag
+		// The containers we use to group particles by type, as pending sync.
+		// This is only needed for the World Sync tag, for the Active tag we don't need to update our owner's and container tags
+		OwnersPendingTEDSWorldSyncTag.Emplace(CompatibilityStorage->FindRowWithCompatibleObject(GetOwner()));
+
+		for (const TPair<EChaosVDParticleType, TSharedPtr<FChaosVDBaseSceneObject>>& ParticleContainer : ParticleSceneContainersByType)
+		{
+			OwnersPendingTEDSWorldSyncTag.Emplace(ParticleContainer.Value->GetTedsRowHandle());
+		}
+
+		BatchTEDSAddRemoveTags(DataStorage, OwnersPendingTEDSWorldSyncTag, { FTypedElementSyncFromWorldTag::StaticStruct() }, {});
+
+		TArray<const UScriptStruct*, TInlineAllocator<2>> ColumnsToAdd = { FTypedElementSyncFromWorldTag::StaticStruct(),  FChaosVDActiveObjectTag::StaticStruct() };
+		BatchTEDSAddRemoveTags(DataStorage, ParticlesPendingCombinedTEDSTagSync, ColumnsToAdd, {});
+	}
+}
+
 void UChaosVDParticleDataComponent::BatchTEDSAddRemoveTags(UE::Editor::DataStorage::ICoreProvider* EditorDataStorage, FTedsHandlesForBatchArray& ParticlesPendingTEDSTagSync,
-	TConstArrayView<const UScriptStruct*> ColumnsToAdd, TConstArrayView<const UScriptStruct*> ColumnsToRemove)
+                                                           TConstArrayView<const UScriptStruct*> ColumnsToAdd, TConstArrayView<const UScriptStruct*> ColumnsToRemove)
 {
 	if (!EditorDataStorage)
 	{
@@ -366,6 +389,7 @@ void UChaosVDParticleDataComponent::UpdateFromNewSolverStageData(const FChaosVDS
 	}
 
 	FTedsHandlesForBatchArray ParticlesPendingTEDSActiveStateTagAddition;
+	FTedsHandlesForBatchArray ParticlesPendingTEDSActiveStateAndSyncTagAddition;
 	
 	{
 		const FText UpdatingSceneMessage = FText::Format(FTextFormat(LOCTEXT("UpdatingSceneMessage", "Updating Scene for {0} Solver with ID {1} ...")), FText::FromString(GetSolverName()), FText::AsNumber(SolverID));
@@ -386,13 +410,6 @@ void UChaosVDParticleDataComponent::UpdateFromNewSolverStageData(const FChaosVDS
 
 			ParticleToProcessInGT->ProcessPendingParticleDataUpdates();
 
-			if (EnumHasAnyFlags(ParticleToProcessInGT->DirtyFlags, EChaosVDSceneParticleDirtyFlags::Active))
-			{
-				// We could check if the particle is active before doing this, but at this point only newly active particles can be in the queue
-				ParticlesPendingTEDSActiveStateTagAddition.Emplace(ParticleToProcessInGT->GetTedsRowHandle());
-				EnumRemoveFlags(ParticleToProcessInGT->DirtyFlags, EChaosVDSceneParticleDirtyFlags::Active);
-			}
-
 			if (Chaos::VisualDebugger::Cvars::bUnloadParticleDataUsingKeyFrameDiff)
 			{
 				if (TSharedPtr<const FChaosVDParticleDataWrapper> ParticleData = ParticleToProcessInGT->GetParticleData())
@@ -400,12 +417,26 @@ void UChaosVDParticleDataComponent::UpdateFromNewSolverStageData(const FChaosVDS
 					ParticlesUpdatedIDsForKeyFrameDiff.Emplace(ParticleData->ParticleIndex);
 				}
 			}
-
-			if (EnumHasAnyFlags(ParticleToProcessInGT->DirtyFlags, EChaosVDSceneParticleDirtyFlags::TEDS))
+			
+			if (EnumHasAllFlags(ParticleToProcessInGT->DirtyFlags, EChaosVDSceneParticleDirtyFlags::Active | EChaosVDSceneParticleDirtyFlags::TEDS))
+			{
+				ParticlesPendingTEDSActiveStateAndSyncTagAddition.Emplace(ParticleToProcessInGT->GetTedsRowHandle());
+				EnumRemoveFlags(ParticleToProcessInGT->DirtyFlags, EChaosVDSceneParticleDirtyFlags::Active | EChaosVDSceneParticleDirtyFlags::TEDS);
+			}
+			else
 			{
-				ParticlesPendingTEDSSyncFromWorldTagSync.Emplace(ParticleToProcessInGT->GetTedsRowHandle());
+				if (EnumHasAnyFlags(ParticleToProcessInGT->DirtyFlags, EChaosVDSceneParticleDirtyFlags::Active))
+				{
+					// We could check if the particle is active before doing this, but at this point only newly active particles can be in the queue
+					ParticlesPendingTEDSActiveStateTagAddition.Emplace(ParticleToProcessInGT->GetTedsRowHandle());
+					EnumRemoveFlags(ParticleToProcessInGT->DirtyFlags, EChaosVDSceneParticleDirtyFlags::Active);
+				}
 
-				EnumRemoveFlags(ParticleToProcessInGT->DirtyFlags, EChaosVDSceneParticleDirtyFlags::TEDS);
+				if (EnumHasAnyFlags(ParticleToProcessInGT->DirtyFlags, EChaosVDSceneParticleDirtyFlags::TEDS))
+				{
+					ParticlesPendingTEDSSyncFromWorldTagSync.Emplace(ParticleToProcessInGT->GetTedsRowHandle());
+					EnumRemoveFlags(ParticleToProcessInGT->DirtyFlags, EChaosVDSceneParticleDirtyFlags::TEDS);
+				}
 			}
 		}
 	}
@@ -415,9 +446,10 @@ void UChaosVDParticleDataComponent::UpdateFromNewSolverStageData(const FChaosVDS
 
 	const ICompatibilityProvider* EditorDataStorageCompatibility = GetDataStorageFeature<ICompatibilityProvider>(CompatibilityFeatureName);
 	ICoreProvider* EditorDataStorage = GetMutableDataStorageFeature<ICoreProvider>(StorageFeatureName);
-	
+
+	BatchTEDSCombinedSyncFromWorldAndActiveTag(ParticlesPendingTEDSActiveStateAndSyncTagAddition, EditorDataStorageCompatibility, EditorDataStorage);
 	BatchTEDSSyncFromWorldTag(ParticlesPendingTEDSSyncFromWorldTagSync, EditorDataStorageCompatibility, EditorDataStorage);
-	BatchTEDSAddRemoveTags(EditorDataStorage, ParticlesPendingTEDSActiveStateTagAddition, { FChaosVDActiveObjectTag::StaticStruct() } , {});
+	BatchTEDSAddRemoveTags(EditorDataStorage, ParticlesPendingTEDSActiveStateAndSyncTagAddition, { FChaosVDActiveObjectTag::StaticStruct() } , {});
 	BatchTEDSAddRemoveTags(EditorDataStorage, ParticlesPendingTEDSActiveStateTagRemoval, {}, { FChaosVDActiveObjectTag::StaticStruct() });
 }
 
diff --git a/Engine/Plugins/ChaosVD/Source/ChaosVD/Public/Components/ChaosVDParticleDataComponent.h b/Engine/Plugins/ChaosVD/Source/ChaosVD/Public/Components/ChaosVDParticleDataComponent.h
--- a/Engine/Plugins/ChaosVD/Source/ChaosVD/Public/Components/ChaosVDParticleDataComponent.h
+++ b/Engine/Plugins/ChaosVD/Source/ChaosVD/Public/Components/ChaosVDParticleDataComponent.h
@@ -93,6 +93,10 @@ protected:
 		const UE::Editor::DataStorage::ICompatibilityProvider* CompatibilityStorage,
 		UE::Editor::DataStorage::ICoreProvider* DataStorage);
 
+	void BatchTEDSCombinedSyncFromWorldAndActiveTag(FTedsHandlesForBatchArray& ParticlesPendingCombinedTEDSTagSync,
+	const UE::Editor::DataStorage::ICompatibilityProvider* CompatibilityStorage,
+	UE::Editor::DataStorage::ICoreProvider* DataStorage);
+
 	void BatchTEDSAddRemoveTags(UE::Editor::DataStorage::ICoreProvider* EditorDataStorage, FTedsHandlesForBatchArray& ParticlesPendingTEDSTagSync,
 		TConstArrayView<const UScriptStruct*> ColumnsToAdd, TConstArrayView<const UScriptStruct*> ColumnsToRemove);
 
