[ChaosVD] Ground work to support different transport modes in the Liver Recording and Record to file modes.

Currently, only the direct trace to file are supported, but this will allow in the future support Trace direct to memory, or via a relay service.

diff --git a/Engine/Plugins/ChaosVD/Source/ChaosVD/Private/ChaosVDEngine.cpp b/Engine/Plugins/ChaosVD/Source/ChaosVD/Private/ChaosVDEngine.cpp
--- a/Engine/Plugins/ChaosVD/Source/ChaosVD/Private/ChaosVDEngine.cpp
+++ b/Engine/Plugins/ChaosVD/Source/ChaosVD/Private/ChaosVDEngine.cpp
@@ -297,6 +297,31 @@ bool FChaosVDEngine::ConnectToLiveSession(uint32 SessionID, const FString& InSes
 	return bSuccess;
 }
 
+bool FChaosVDEngine::ConnectToLiveSession_Direct(EChaosVDLoadRecordedDataMode LoadingMode)
+{
+	FChaosVDTraceSessionDescriptor NewSessionFromFileDescriptor;
+
+	NewSessionFromFileDescriptor.bIsLiveSession = true;
+	
+	TSharedPtr<FChaosVDRecording> ExistingRecordingInstance = nullptr;
+	if (LoadingMode == EChaosVDLoadRecordedDataMode::MultiSource && PlaybackController)
+	{
+		ExistingRecordingInstance = PlaybackController->GetCurrentRecording().Pin();
+	}
+
+	NewSessionFromFileDescriptor.SessionName = FChaosVDModule::Get().GetTraceManager()->ConnectToLiveSession_Direct(NewSessionFromFileDescriptor.SessionPort, ExistingRecordingInstance);
+
+	bool bSuccess = false;
+
+	if (NewSessionFromFileDescriptor.IsValid())
+	{
+		OpenSession(NewSessionFromFileDescriptor, LoadingMode);
+		bSuccess = true;
+	}
+
+	return bSuccess;
+}
+
 void FChaosVDEngine::OpenSession(const FChaosVDTraceSessionDescriptor& SessionDescriptor, EChaosVDLoadRecordedDataMode LoadingMode)
 {
 	if (LoadingMode == EChaosVDLoadRecordedDataMode::SingleSource)
diff --git a/Engine/Plugins/ChaosVD/Source/ChaosVD/Private/Settings/ChaosVDGeneralSettings.h b/Engine/Plugins/ChaosVD/Source/ChaosVD/Private/Settings/ChaosVDGeneralSettings.h
--- a/Engine/Plugins/ChaosVD/Source/ChaosVD/Private/Settings/ChaosVDGeneralSettings.h
+++ b/Engine/Plugins/ChaosVD/Source/ChaosVD/Private/Settings/ChaosVDGeneralSettings.h
@@ -3,6 +3,7 @@
 #pragma once
 
 #include "Settings/ChaosVDCoreSettings.h"
+#include "ChaosVDRecordingDetails.h"
 #include "ChaosVDGeneralSettings.generated.h"
 
 /**
@@ -42,4 +43,12 @@ public:
 	/** How many times CVD will attempt to connect to a live trace or load from file session if the first attempt failed */
 	UPROPERTY(Config, EditAnywhere, Category = Connection)
 	int32 MaxConnectionRetries = 14;
+
+	/** [DEBUG Setting] Data Transport Override - Used to change the transport mode for the trace data for testing. This is not saved by design */
+	UPROPERTY(EditAnywhere, Category = Connection)
+	EChaosVDTransportMode DataTransportModeOverride = EChaosVDTransportMode::Invalid;
+
+	/** [DEBUG Setting] If True, any traces done to memory will also be saved to disk at the same time in the Profiling folder. This is not saved by design */
+	UPROPERTY(EditAnywhere, Category = Connection)
+	bool bSaveMemoryTracesToDisk = true;
 };
diff --git a/Engine/Plugins/ChaosVD/Source/ChaosVD/Private/Trace/ChaosVDTraceManager.cpp b/Engine/Plugins/ChaosVD/Source/ChaosVD/Private/Trace/ChaosVDTraceManager.cpp
--- a/Engine/Plugins/ChaosVD/Source/ChaosVD/Private/Trace/ChaosVDTraceManager.cpp
+++ b/Engine/Plugins/ChaosVD/Source/ChaosVD/Private/Trace/ChaosVDTraceManager.cpp
@@ -4,11 +4,14 @@
 
 #include "ChaosVDEngine.h"
 #include "ChaosVDModule.h"
+#include "ChaosVDRuntimeModule.h"
+#include "ChaosVDSettingsManager.h"
 #include "Features/IModularFeatures.h"
-#include "GenericPlatform/GenericPlatformFile.h"
+#include "HAL/PlatformFileManager.h"
 #include "Modules/ModuleManager.h"
 #include "Misc/Paths.h"
 #include "ProfilingDebugging/TraceAuxiliary.h"
+#include "Settings/ChaosVDGeneralSettings.h"
 #include "Trace/ChaosVDTraceModule.h"
 #include "Trace/StoreClient.h"
 #include "TraceServices/AnalysisService.h"
@@ -154,6 +157,55 @@ const UE::Trace::FStoreClient::FSessionInfo* FChaosVDTraceManager::GetTraceSessi
 	return StoreClient->GetSessionInfoByGuid(TraceGuid);
 }
 
+FString FChaosVDTraceManager::GetTraceFileNameFromStoreForSession(FStringView InSessionHost, uint32 SessionID)
+{
+	using namespace UE::Trace;
+	FStoreClient* StoreClient = FStoreClient::Connect(InSessionHost.GetData());
+
+	if (!StoreClient)
+	{
+		return FString();
+	}
+
+	const FStoreClient::FTraceInfo* TraceInfo = StoreClient->GetTraceInfoById(SessionID);
+	if (!TraceInfo)
+	{
+		return FString();
+	}
+
+	const FStoreClient::FStatus* StoreStatus = StoreClient->GetStatus();
+	if (!StoreStatus)
+	{
+		return FString();
+	}
+
+	const FUtf8StringView Utf8NameView = TraceInfo->GetName();
+	FString FileName(Utf8NameView);
+	if (!FileName.EndsWith(TEXT(".utrace")))
+	{
+		FileName += TEXT(".utrace");
+	}
+	
+	FString FullFileName = FPaths::Combine(FString(StoreStatus->GetStoreDir()), FileName);
+	FPaths::NormalizeFilename(FullFileName);
+
+	return FullFileName;
+}
+
+
+UE::Trace::FStoreClient::FTraceData FChaosVDTraceManager::GetTraceDataStreamFromStore(FStringView InSessionHost, uint32 SessionID)
+{
+	using namespace UE::Trace;
+	FStoreClient* StoreClient = FStoreClient::Connect(InSessionHost.GetData());
+	if (!StoreClient)
+	{
+		return UE::Trace::FStoreClient::FTraceData();
+	}
+
+	return StoreClient->ReadTrace(SessionID);
+}
+
+
 TUniquePtr<UE::Trace::IInDataStream> FChaosVDTraceManager::CreateFileDataStream(TUniquePtr<IFileHandle>&& InFileHandle)
 {
 	check(InFileHandle);
@@ -227,60 +279,126 @@ void FChaosVDTraceManager::SetPendingExternalRecordingToProcess(const TSharedPtr
 	PendingRecording = ExistingRecordingPtr;
 }
 
-FString FChaosVDTraceManager::ConnectToLiveSession(FStringView InSessionHost, uint32 SessionID, const TSharedPtr<FChaosVDRecording>& ExistingRecordingPtr)
+bool FChaosVDTraceManager::ConnectToLiveSession_Internal(uint32 SessionID, const TSharedPtr<FChaosVDRecording>& ExistingRecordingPtr, const FString& InRequestedSessionName, UE::Trace::FStoreClient::FTraceData&& InTraceDataStream)
 {
-	FString SessionName;
+	ITraceServicesModule& TraceServicesModule = FModuleManager::LoadModuleChecked<ITraceServicesModule>("TraceServices");
+	if (const TSharedPtr<TraceServices::IAnalysisService> TraceAnalysisService = TraceServicesModule.GetAnalysisService())
+	{
+		// Close this session in case we were already analysing it
+		CloseSession(InRequestedSessionName);
+		
+		SetPendingExternalRecordingToProcess(ExistingRecordingPtr);
+	
+		if (const TSharedPtr<const TraceServices::IAnalysisSession> NewSession = TraceAnalysisService->StartAnalysis(SessionID, *InRequestedSessionName, MoveTemp(InTraceDataStream)))
+		{
+			AnalysisSessionByName.Add(InRequestedSessionName, NewSession.ToSharedRef());
+			return true;
+		}
+	}
 
+	return false;
+}
+
+
+FString FChaosVDTraceManager::ConnectToLiveSession(FStringView InSessionHost, uint32 SessionID, const TSharedPtr<FChaosVDRecording>& ExistingRecordingPtr)
+{
 	if (InSessionHost.IsEmpty())
 	{
-		return SessionName;
+		return FString();
 	}
 
 	using namespace UE::Trace;
-	FStoreClient* StoreClient = FStoreClient::Connect(InSessionHost.GetData());
 
-	if (!StoreClient)
+	FString TraceName = GetTraceFileNameFromStoreForSession(InSessionHost, SessionID);
+	if (TraceName.IsEmpty())
 	{
-		return SessionName;
+		UE_LOG(LogChaosVDEditor, Error, TEXT("[%hs] Failed to get the trace file name from the trace store | Host [%s] | Session ID [%u]"), __func__, InSessionHost.GetData(), SessionID);
+		return FString();
 	}
 
-	FStoreClient::FTraceData TraceData = StoreClient->ReadTrace(SessionID);
+	FStoreClient::FTraceData TraceData = GetTraceDataStreamFromStore(InSessionHost, SessionID);
 	if (!TraceData.IsValid())
 	{
-		return SessionName;
+		UE_LOG(LogChaosVDEditor, Error, TEXT("[%hs] Failed to get trace data stream from the trace store | Host [%s] | Session ID [%u]"), __func__, InSessionHost.GetData(), SessionID);
+		return FString();
 	}
 
-	FString TraceName(StoreClient->GetStatus()->GetStoreDir());
-	const FStoreClient::FTraceInfo* TraceInfo = StoreClient->GetTraceInfoById(SessionID);
-	if (TraceInfo != nullptr)
+	if (ConnectToLiveSession_Internal(SessionID, ExistingRecordingPtr, TraceName, MoveTemp(TraceData)))
+	{
+		 return TraceName;
+	}
+
+	UE_LOG(LogChaosVDEditor, Error, TEXT("[%hs] Failed connect to session | Host [%s] | Session ID [%u]"), __func__, InSessionHost.GetData(), SessionID);
+
+	return FString();
+}
+
+TUniquePtr<FArchiveFileWriterGeneric> FChaosVDTraceManager::OpenTraceFileForWrite()
+{
+	TUniquePtr<FArchiveFileWriterGeneric> FileWriter = nullptr;
+
+	IPlatformFile& PlatformFile = FPlatformFileManager::Get().GetPlatformFile();
+	FChaosVDRuntimeModule& CVDRuntimeModule = FChaosVDRuntimeModule::Get();
+	constexpr int32 MaxAttempts = 10;
+	int32 CurrentAttempts = 0;
+	while (CurrentAttempts < MaxAttempts)
 	{
-		const FUtf8StringView Utf8NameView = TraceInfo->GetName();
-		FString Name(Utf8NameView);
-		if (!Name.EndsWith(TEXT(".utrace")))
+		FString OutFileName;
+		CVDRuntimeModule.GenerateRecordingFileName(OutFileName);
+		OutFileName = FPaths::Combine(FPaths::ProfilingDir(), OutFileName);
+
+		IFileHandle* FileHandle = FileHandle = PlatformFile.OpenWrite(*OutFileName);
+		if (FileHandle)
 		{
-			Name += TEXT(".utrace");
+			FileWriter = MakeUnique<FArchiveFileWriterGeneric>(FileHandle, *OutFileName, FileHandle->Tell());
+			break;
 		}
-		TraceName = FPaths::Combine(TraceName, Name);
-		FPaths::NormalizeFilename(TraceName);
+	
+		CurrentAttempts++;
+
+		FPlatformProcess::Sleep(0.1f);
 	}
 
-	ITraceServicesModule& TraceServicesModule = FModuleManager::LoadModuleChecked<ITraceServicesModule>("TraceServices");
-	if (const TSharedPtr<TraceServices::IAnalysisService> TraceAnalysisService = TraceServicesModule.GetAnalysisService())
+	return MoveTemp(FileWriter);
+}
+
+FString FChaosVDTraceManager::ConnectToLiveSession_Direct(uint16& OutSessionPort, const TSharedPtr<FChaosVDRecording>& ExistingRecordingPtr)
+{
+	TUniquePtr<FArchiveFileWriterGeneric> TraceFileWriter = nullptr;
+
+	if (UChaosVDGeneralSettings* Settings = FChaosVDSettingsManager::Get().GetSettingsObject<UChaosVDGeneralSettings>())
 	{
-		// Close this session in case we were already analysing it
-		CloseSession(TraceName);
-		
-		SetPendingExternalRecordingToProcess(ExistingRecordingPtr);
-	
-		if (const TSharedPtr<const TraceServices::IAnalysisSession> NewSession = TraceAnalysisService->StartAnalysis(SessionID, *TraceName, MoveTemp(TraceData)))
+		if (Settings->bSaveMemoryTracesToDisk)
 		{
-			AnalysisSessionByName.Add(TraceName, NewSession.ToSharedRef());
-
-			SessionName = NewSession->GetName();
+			TraceFileWriter = OpenTraceFileForWrite();
+			if (!TraceFileWriter)
+			{
+				UE_LOG(LogChaosVDEditor, Error, TEXT("[%hs] Failed to open open file for direct trace stream"), __func__);
+				return FString();
+			}
 		}
 	}
 
-	return SessionName;
+	TUniquePtr<UE::Trace::FDirectSocketStream> DirectSocketStream = MakeUnique<UE::Trace::FDirectSocketStream>(MoveTemp(TraceFileWriter));
+	OutSessionPort = DirectSocketStream->StartListening();
+
+	if (OutSessionPort == 0)
+	{
+		UE_LOG(LogChaosVDEditor, Error, TEXT("[%hs] Failed to open direct trace stream socket"), __func__);
+		return FString();
+	}
+
+	FString DirectSessionName = FString::Printf(TEXT("DirectSession[127.0.0.1:%u]"), OutSessionPort);
+	
+	if (ConnectToLiveSession_Internal(~0u, ExistingRecordingPtr, DirectSessionName, MoveTemp(DirectSocketStream)))
+	{
+		return DirectSessionName;
+	}
+	else
+	{
+		UE_LOG(LogChaosVDEditor, Error, TEXT("[%hs] Failed start trace analysis with direct trace stream socket | [%s]"), __func__, *DirectSessionName);
+		return FString();
+	}
 }
 
 FString FChaosVDTraceManager::GetLocalTraceStoreDirPath()
diff --git a/Engine/Plugins/ChaosVD/Source/ChaosVD/Private/Widgets/SChaosVDMainTab.cpp b/Engine/Plugins/ChaosVD/Source/ChaosVD/Private/Widgets/SChaosVDMainTab.cpp
--- a/Engine/Plugins/ChaosVD/Source/ChaosVD/Private/Widgets/SChaosVDMainTab.cpp
+++ b/Engine/Plugins/ChaosVD/Source/ChaosVD/Private/Widgets/SChaosVDMainTab.cpp
@@ -909,6 +909,11 @@ bool SChaosVDMainTab::ConnectToLiveSession(int32 SessionID, const FString& InSes
 	return GetChaosVDEngineInstance()->ConnectToLiveSession(SessionID, InSessionAddress, LoadingMode);
 }
 
+bool SChaosVDMainTab::ConnectToLiveSession_Direct(EChaosVDLoadRecordedDataMode LoadingMode) const
+{
+	return GetChaosVDEngineInstance()->ConnectToLiveSession_Direct(LoadingMode);
+}
+
 bool SChaosVDMainTab::IsSupportedFile(const FString& InFilename)
 {
 	return InFilename.EndsWith(TEXT(".utrace")) || InFilename.EndsWith(TEXT(".cvdmulti"));
diff --git a/Engine/Plugins/ChaosVD/Source/ChaosVD/Private/Widgets/SChaosVDRecordingControls.cpp b/Engine/Plugins/ChaosVD/Source/ChaosVD/Private/Widgets/SChaosVDRecordingControls.cpp
--- a/Engine/Plugins/ChaosVD/Source/ChaosVD/Private/Widgets/SChaosVDRecordingControls.cpp
+++ b/Engine/Plugins/ChaosVD/Source/ChaosVD/Private/Widgets/SChaosVDRecordingControls.cpp
@@ -124,7 +124,7 @@ FText SChaosVDRecordingControls::GetCurrentSelectedSessionName() const
 		return FText::AsCultureInvariant(CurrentSessionPtr->SessionName);
 	}
 
-	static FText InvalidSessionName = LOCTEXT("InvalidSessionLabel", "Invalid Session");
+	static FText InvalidSessionName = LOCTEXT("InvalidSessionLabel", "Select a Session");
 
 	return InvalidSessionName;
 }
@@ -551,7 +551,28 @@ void SChaosVDRecordingControls::HandleRecordingStop(TWeakPtr<FChaosVDSessionInfo
 		return;
 	}
 
-	FText CurrentTraceTarget = FText::AsCultureInvariant(SessionInfoPtr->LastKnownRecordingState.TraceDetails.TraceTarget);
+	FString TraceTarget;
+
+	switch (SessionInfoPtr->LastKnownRecordingState.TraceDetails.TransportMode)
+	{
+		case EChaosVDTransportMode::FileSystem:
+			{
+				TraceTarget = SessionInfoPtr->LastKnownRecordingState.TraceDetails.TraceTarget;
+				break;
+			}
+		case EChaosVDTransportMode::TraceServer:
+			{
+				FTraceStoreConnectionSettings ConnectionDetails = GetTraceStoreConnectionSettings(SessionInfoPtr.ToSharedRef());
+				TraceTarget = FChaosVDModule::Get().GetTraceManager()->GetTraceFileNameFromStoreForSession(ConnectionDetails.SessionAddress, ConnectionDetails.TraceID);
+				break;
+			}
+		case EChaosVDTransportMode::Memory:
+		case EChaosVDTransportMode::Relay:
+		case EChaosVDTransportMode::Invalid:
+		default:
+		UE_LOG(LogChaosVDEditor, Error, TEXT("[%hs] Transport mode is not supported | [%s]"), __func__, *UEnum::GetValueAsString(SessionInfoPtr->LastKnownRecordingState.TraceDetails.TransportMode));
+			break;
+	}
 
 	const bool bIsLiveSession = SessionInfoPtr->GetRecordingMode() == EChaosVDRecordingMode::Live;
 
@@ -566,7 +587,7 @@ void SChaosVDRecordingControls::HandleRecordingStop(TWeakPtr<FChaosVDSessionInfo
 		}
 		else
 		{
-			const FText RecordingPathMessage = FText::Format(LOCTEXT("RecordingSavedPathMessage"," Recording saved at {0} "), CurrentTraceTarget);
+			const FText RecordingPathMessage = FText::Format(LOCTEXT("RecordingSavedPathMessage"," Recording saved at {0} "), FText::AsCultureInvariant(TraceTarget));
 			RecordingPathMessageHandle = StatusBarSubsystem->PushStatusBarMessage(StatusBarID, RecordingPathMessage);
 		}
 	}
@@ -575,7 +596,7 @@ void SChaosVDRecordingControls::HandleRecordingStop(TWeakPtr<FChaosVDSessionInfo
 	{
 		if (FMessageDialog::Open(EAppMsgType::YesNo, LOCTEXT("OpenLastRecordingMessage", "Do you want to load the recorded file now? ")) == EAppReturnType::Yes)
 		{
-			MainTabSharedPtr->GetChaosVDEngineInstance()->LoadRecording(CurrentTraceTarget.ToString(), EChaosVDLoadRecordedDataMode::SingleSource);
+			MainTabSharedPtr->GetChaosVDEngineInstance()->LoadRecording(TraceTarget, EChaosVDLoadRecordedDataMode::SingleSource);			
 		}
 	}
 }
@@ -660,9 +681,9 @@ void SChaosVDRecordingControls::ExecutePostRecordingAsyncTaskWithRetry_Internal(
 	}), IntervalBetweenAutoplayConnectionAttemptsSeconds);
 }
 
-SChaosVDRecordingControls::FLiveConnectionSettings SChaosVDRecordingControls::GetTargetLiveConnectionSettings(const TSharedRef<FChaosVDSessionInfo>& InSessionInfoRef)
+SChaosVDRecordingControls::FTraceStoreConnectionSettings SChaosVDRecordingControls::GetTraceStoreConnectionSettings(const TSharedRef<FChaosVDSessionInfo>& InSessionInfoRef)
 {
-	FLiveConnectionSettings ConnectionSettings;
+	FTraceStoreConnectionSettings ConnectionSettings;
 
 	if (!ensure(!EnumHasAnyFlags(InSessionInfoRef->GetSessionTypeAttributes(), EChaosVDRemoteSessionAttributes::IsMultiSessionWrapper)))
 	{
@@ -671,16 +692,13 @@ SChaosVDRecordingControls::FLiveConnectionSettings SChaosVDRecordingControls::Ge
 		return ConnectionSettings;
 	}
 
-	if (InSessionInfoRef->IsConnected() && InSessionInfoRef->GetRecordingMode() == EChaosVDRecordingMode::Live)
-	{
-		FChaosVDTraceDetails& RecordingSessionDetails = InSessionInfoRef->LastKnownRecordingState.TraceDetails;
+	FChaosVDTraceDetails& RecordingSessionDetails = InSessionInfoRef->LastKnownRecordingState.TraceDetails;
 
-		if (const UE::Trace::FStoreClient::FSessionInfo* TraceSessionInfo = FChaosVDTraceManager::GetTraceSessionInfo(RecordingSessionDetails.TraceTarget, RecordingSessionDetails.TraceGuid))
-		{
-			ConnectionSettings.SessionAddress = RecordingSessionDetails.TraceTarget;
-			ConnectionSettings.TraceID = TraceSessionInfo->GetTraceId();
-			return ConnectionSettings;
-		}
+	if (const UE::Trace::FStoreClient::FSessionInfo* TraceSessionInfo = FChaosVDTraceManager::GetTraceSessionInfo(RecordingSessionDetails.TraceTarget, RecordingSessionDetails.TraceGuid))
+	{
+		ConnectionSettings.SessionAddress = RecordingSessionDetails.TraceTarget;
+		ConnectionSettings.TraceID = TraceSessionInfo->GetTraceId();
+		return ConnectionSettings;
 	}
 
 	return ConnectionSettings;
@@ -728,10 +746,12 @@ void SChaosVDRecordingControls::SetSessionRecordingState(bool bIsRecording, ECha
 		TArray<FString, TInlineAllocator<1>> RecordingArgs;
 
 		int32 RemainingRetries = 4;
+		EChaosVDTransportMode DataTransportOverrideMode = EChaosVDTransportMode::Invalid;
 
 		if (UChaosVDGeneralSettings* Settings = FChaosVDSettingsManager::Get().GetSettingsObject<UChaosVDGeneralSettings>())
 		{
 			RemainingRetries = Settings->MaxConnectionRetries;
+			DataTransportOverrideMode = Settings->DataTransportModeOverride;
 		}
 		else
 		{
@@ -757,6 +777,7 @@ void SChaosVDRecordingControls::SetSessionRecordingState(bool bIsRecording, ECha
 		{
 			FChaosVDStartRecordingCommandMessage RecordingParams;
 			RecordingParams.RecordingMode = EChaosVDRecordingMode::Live;
+			RecordingParams.TransportMode = DataTransportOverrideMode != EChaosVDTransportMode::Invalid ? DataTransportOverrideMode : EChaosVDTransportMode::TraceServer;
 
 			check(GLog);
 			bool bOutCanBindAll = false;
@@ -776,36 +797,65 @@ void SChaosVDRecordingControls::SetSessionRecordingState(bool bIsRecording, ECha
 				return true;
 			});
 #endif
-			RemoteSessionManager->SendStartRecordingCommand(SessionInfoRef->Address, RecordingParams);
 
-			// Once the start recording command is issue, we can try to connect to the created session (if everything went well).
-			// If it didn't go well, this will take care of update the UI to notify the user
-			ExecuteAsyncConnectionAttemptTaskWithRetry(RemainingRetries, [WeakThis = AsWeak(), SessionInstanceId = SessionInfoRef->InstanceId]()
+			if (RecordingParams.TransportMode == EChaosVDTransportMode::Memory)
 			{
-				TSharedPtr<FChaosVDRemoteSessionsManager> RemoteSessionManagerPtr = FChaosVDEngineEditorBridge::Get().GetRemoteSessionsManager();
-				TSharedPtr<FChaosVDSessionInfo> SessionInfoPtr = RemoteSessionManagerPtr ? RemoteSessionManagerPtr->GetSessionInfo(SessionInstanceId).Pin() : nullptr;
-
-				if (!SessionInfoPtr)
+				// Using the memory transport or direct trace 
+				const TSharedPtr<SChaosVDMainTab> MainTabSharedPtr = MainTabWeakPtr.Pin();
+				if (!MainTabSharedPtr)
 				{
-					return false;
+					return;
 				}
-				
-				TSharedPtr<SChaosVDRecordingControls> Controls =  StaticCastSharedPtr<SChaosVDRecordingControls>(WeakThis.Pin());
-				const TSharedPtr<SChaosVDMainTab> MainTabSharedPtr = Controls ? Controls->MainTabWeakPtr.Pin() : nullptr;
-				if (!MainTabSharedPtr)
+
+				// Memory Tracing mode requires opening a trace session in the editor side first before starting the recording
+				// so targets have something to connect to (Usually targets connect to a running trace server)
+				if (!MainTabSharedPtr->ConnectToLiveSession_Direct(CurrentLoadingMode))
 				{
-					return false;
+					UE_LOG(LogChaosVDEditor, Error, TEXT("[%hs] Failed to connect to live session using direct trace mode"), __func__);
+					return;
 				}
 
-				FLiveConnectionSettings ConnectionDetails = Controls->GetTargetLiveConnectionSettings(SessionInfoPtr.ToSharedRef());
-				return MainTabSharedPtr->ConnectToLiveSession(ConnectionDetails.TraceID, ConnectionDetails.SessionAddress, Controls->CurrentLoadingMode);
-			},
-			RecordingAttemptFailedCallback);
+				// If we get here it means we have at least one valid session
+				const FChaosVDTraceSessionDescriptor& SessionDesc = MainTabSharedPtr->GetChaosVDEngineInstance()->GetCurrentSessionDescriptors().Last();
+
+				RecordingParams.Target.Appendf(TEXT(":%u"), SessionDesc.SessionPort);
+
+				RemoteSessionManager->SendStartRecordingCommand(SessionInfoRef->Address, RecordingParams);
+			}
+			else
+			{
+				RemoteSessionManager->SendStartRecordingCommand(SessionInfoRef->Address, RecordingParams);
+                
+                // Once the start recording command is issue, we can try to connect to the created session (if everything went well).
+                // If it didn't go well, this will take care of update the UI to notify the user
+                ExecuteAsyncConnectionAttemptTaskWithRetry(RemainingRetries, [WeakThis = AsWeak(), SessionInstanceId = SessionInfoRef->InstanceId]()
+                {
+                	TSharedPtr<FChaosVDRemoteSessionsManager> RemoteSessionManagerPtr = FChaosVDEngineEditorBridge::Get().GetRemoteSessionsManager();
+                	TSharedPtr<FChaosVDSessionInfo> SessionInfoPtr = RemoteSessionManagerPtr ? RemoteSessionManagerPtr->GetSessionInfo(SessionInstanceId).Pin() : nullptr;
+    
+                	if (!SessionInfoPtr)
+                	{
+                		return false;
+                	}
+                	
+                	TSharedPtr<SChaosVDRecordingControls> Controls =  StaticCastSharedPtr<SChaosVDRecordingControls>(WeakThis.Pin());
+                	const TSharedPtr<SChaosVDMainTab> MainTabSharedPtr = Controls ? Controls->MainTabWeakPtr.Pin() : nullptr;
+                	if (!MainTabSharedPtr)
+                	{
+                		return false;
+                	}
+    
+                	FTraceStoreConnectionSettings ConnectionDetails = Controls->GetTraceStoreConnectionSettings(SessionInfoPtr.ToSharedRef());
+                	return MainTabSharedPtr->ConnectToLiveSession(ConnectionDetails.TraceID, ConnectionDetails.SessionAddress, Controls->CurrentLoadingMode);
+                },
+                RecordingAttemptFailedCallback);
+			}
 		}
 		else
 		{
 			FChaosVDStartRecordingCommandMessage RecordingParams;
 			RecordingParams.RecordingMode = EChaosVDRecordingMode::File;
+			RecordingParams.TransportMode = DataTransportOverrideMode != EChaosVDTransportMode::Invalid ? DataTransportOverrideMode : EChaosVDTransportMode::FileSystem;
 
 #if WITH_CHAOS_VISUAL_DEBUGGER
 			Chaos::VisualDebugger::FChaosVDDataChannelsManager::Get().EnumerateChannels([&RecordingParams](const TSharedRef<Chaos::VisualDebugger::FChaosVDOptionalDataChannel>& Channel)
diff --git a/Engine/Plugins/ChaosVD/Source/ChaosVD/Private/Widgets/SChaosVDRecordingControls.h b/Engine/Plugins/ChaosVD/Source/ChaosVD/Private/Widgets/SChaosVDRecordingControls.h
--- a/Engine/Plugins/ChaosVD/Source/ChaosVD/Private/Widgets/SChaosVDRecordingControls.h
+++ b/Engine/Plugins/ChaosVD/Source/ChaosVD/Private/Widgets/SChaosVDRecordingControls.h
@@ -120,13 +120,13 @@ protected:
 
 	void HandleConnectionAttemptResult(EChaosVDLiveConnectionAttemptResult Result, const TSharedPtr<SNotificationItem>& InNotification);
 
-	struct FLiveConnectionSettings
+	struct FTraceStoreConnectionSettings
 	{
 		int32 TraceID = INDEX_NONE;
 		FString SessionAddress;
 	};
 
-	FLiveConnectionSettings GetTargetLiveConnectionSettings(const TSharedRef<FChaosVDSessionInfo>& InSessionInfoRef);
+	FTraceStoreConnectionSettings GetTraceStoreConnectionSettings(const TSharedRef<FChaosVDSessionInfo>& InSessionInfoRef);
 
 	FText GetCurrentSelectedSessionName() const;
 
diff --git a/Engine/Plugins/ChaosVD/Source/ChaosVD/Public/ChaosVDEngine.h b/Engine/Plugins/ChaosVD/Source/ChaosVD/Public/ChaosVDEngine.h
--- a/Engine/Plugins/ChaosVD/Source/ChaosVD/Public/ChaosVDEngine.h
+++ b/Engine/Plugins/ChaosVD/Source/ChaosVD/Public/ChaosVDEngine.h
@@ -75,6 +75,7 @@ public:
 	void LoadCombinedMultiRecording(const FString& FilePath);
 
 	bool ConnectToLiveSession(uint32 SessionID, const FString& InSessionAddress, EChaosVDLoadRecordedDataMode LoadingMode = EChaosVDLoadRecordedDataMode::MultiSource);
+	bool ConnectToLiveSession_Direct(EChaosVDLoadRecordedDataMode LoadingMode = EChaosVDLoadRecordedDataMode::MultiSource);
 
 	void OpenSession(const FChaosVDTraceSessionDescriptor& SessionDescriptor, EChaosVDLoadRecordedDataMode LoadingMode = EChaosVDLoadRecordedDataMode::MultiSource);
 
diff --git a/Engine/Plugins/ChaosVD/Source/ChaosVD/Public/Trace/ChaosVDTraceManager.h b/Engine/Plugins/ChaosVD/Source/ChaosVD/Public/Trace/ChaosVDTraceManager.h
--- a/Engine/Plugins/ChaosVD/Source/ChaosVD/Public/Trace/ChaosVDTraceManager.h
+++ b/Engine/Plugins/ChaosVD/Source/ChaosVD/Public/Trace/ChaosVDTraceManager.h
@@ -49,6 +49,7 @@ namespace Chaos::VD
 struct FChaosVDTraceSessionDescriptor
 {
 	FString SessionName;
+	uint16 SessionPort = 0;
 	bool bIsLiveSession = false;
 
 	bool IsValid() const { return !SessionName.IsEmpty(); }
@@ -83,10 +84,19 @@ public:
 	 * Connects to a live Trace Session and starts analyzing it.
 	 * @param InSessionHost Trace Store Address for this session
 	 * @param SessionID Trace ID in the Trace Store provided as host
+	 * @param ExistingRecordingPtr Ptr to an existing recording instance to append the new CVD data to, if any
 	 * @return Session Name
 	 */
 	FString ConnectToLiveSession(FStringView InSessionHost, uint32 SessionID, const TSharedPtr<FChaosVDRecording>& ExistingRecordingPtr = nullptr);
 
+	/**
+	 * Opens up a direct trace stream and waits for data
+	 * @param OutSessionPort Number of the port we open for this session
+	 * @param ExistingRecordingPtr Ptr to an existing recording instance to append the new CVD data to, if any
+	 * @return Session Name
+	 */
+	FString ConnectToLiveSession_Direct(uint16& OutSessionPort, const TSharedPtr<FChaosVDRecording>& ExistingRecordingPtr = nullptr);
+
 	/** Returns the path to the local trace store */
 	FString GetLocalTraceStoreDirPath();
 
@@ -107,8 +117,25 @@ public:
 	template<typename TVisitor>
 	void EnumerateActiveSessions(TVisitor Callback);
 
+	/** Access the trace store at the provided host address, and returns the file name for the trace file
+	 * @param InSessionHost Address of the Trace Store
+	 * @param SessionID Session id of the trace data from which get the file name from
+	 * @return Full Filename with path to the trace file if exist. Returns an empty string if a trace file cannot be found
+	 */
+	FString GetTraceFileNameFromStoreForSession(FStringView InSessionHost, uint32 SessionID);
+
+	/** Access the trace store at the provided host address, and returns a trace data stream ready to use
+	 * @param InSessionHost Address of the Trace Store
+	 * @param SessionID Session id of the trace data
+	 */
+	UE::Trace::FStoreClient::FTraceData GetTraceDataStreamFromStore(FStringView InSessionHost, uint32 SessionID);
+
 private:
 
+	TUniquePtr<FArchiveFileWriterGeneric>  OpenTraceFileForWrite();
+
+	bool ConnectToLiveSession_Internal(uint32 SessionID, const TSharedPtr<FChaosVDRecording>& ExistingRecordingPtr, const FString& InRequestedSessionName, UE::Trace::FStoreClient::FTraceData&& InTraceDataStream);
+
 	TUniquePtr<UE::Trace::IInDataStream> CreateFileDataStream(TUniquePtr<IFileHandle>&& InFileHandle);
 
 	/** Temporary workaround method to set an existing recording structure in CVD's trace provider before the trace analysis starts
diff --git a/Engine/Plugins/ChaosVD/Source/ChaosVD/Public/Widgets/SChaosVDMainTab.h b/Engine/Plugins/ChaosVD/Source/ChaosVD/Public/Widgets/SChaosVDMainTab.h
--- a/Engine/Plugins/ChaosVD/Source/ChaosVD/Public/Widgets/SChaosVDMainTab.h
+++ b/Engine/Plugins/ChaosVD/Source/ChaosVD/Public/Widgets/SChaosVDMainTab.h
@@ -107,6 +107,7 @@ public:
 	}
 	
 	bool ConnectToLiveSession(int32 SessionID, const FString& InSessionAddress, EChaosVDLoadRecordedDataMode LoadingMode) const;
+	bool ConnectToLiveSession_Direct(EChaosVDLoadRecordedDataMode LoadingMode) const;
 
 	/**
 	 * Evaluates a filename to determine if it is a supported filetype
diff --git a/Engine/Source/Runtime/Experimental/ChaosSolverEngine/Private/Chaos/ChaosVDRemoteSessionsManager.cpp b/Engine/Source/Runtime/Experimental/ChaosSolverEngine/Private/Chaos/ChaosVDRemoteSessionsManager.cpp
--- a/Engine/Source/Runtime/Experimental/ChaosSolverEngine/Private/Chaos/ChaosVDRemoteSessionsManager.cpp
+++ b/Engine/Source/Runtime/Experimental/ChaosSolverEngine/Private/Chaos/ChaosVDRemoteSessionsManager.cpp
@@ -94,6 +94,13 @@ FChaosVDRemoteSessionsManager::FChaosVDRemoteSessionsManager()
 {
 }
 
+void FChaosVDRemoteSessionsManager::Initialize(const TSharedPtr<IMessageBus>& InMessageBus)
+{
+	ensureMsgf(!InMessageBus, TEXT("Initialize(MessageBusPtr) is deprecated!. The provided message buss will be ignored"));
+
+	Initialize();
+}
+
 constexpr bool FChaosVDRemoteSessionsManager::IsController()
 {
 #if WITH_EDITOR
@@ -144,17 +151,21 @@ TSharedPtr<FMessageEndpoint> FChaosVDRemoteSessionsManager::CreateEndPoint(const
 	return EndPointBuilder.Build();
 }
 
-void FChaosVDRemoteSessionsManager::Initialize()
+void FChaosVDRemoteSessionsManager::ReInitializeMessagingSystem(const TSharedPtr<IMessageBus>& InMessageBus)
 {
-	TSharedPtr<IMessageBus> CVDMessageBus = IMessagingModule::Get().GetDefaultBus();
+	ShutdownMessagingSystem();
+	InitializeMessagingSystem(InMessageBus);
+}
 
-	if (!ensure(CVDMessageBus))
+void FChaosVDRemoteSessionsManager::InitializeMessagingSystem(const TSharedPtr<IMessageBus>& InMessageBus)
+{
+	if (!ensure(InMessageBus))
 	{
 		return;
 	}
 
-	MessageBusPtr = CVDMessageBus;
-	MessageEndpoint = CreateEndPoint(CVDMessageBus.ToSharedRef());
+	MessageBusPtr = InMessageBus;
+	MessageEndpoint = CreateEndPoint(InMessageBus.ToSharedRef());
 	
 	if (!ensure(MessageEndpoint))
 	{
@@ -174,15 +185,9 @@ void FChaosVDRemoteSessionsManager::Initialize()
 	MessageEndpoint->Subscribe<FChaosVDStopRecordingCommandMessage>();
 	MessageEndpoint->Subscribe<FChaosVDChannelStateChangeCommandMessage>();
 	MessageEndpoint->Subscribe<FChaosVDFullSessionInfoRequestMessage>();
-
-	ActiveSessionsByInstanceId.Add(AllRemoteSessionsWrapperGUID, CreatedWrapperSessionInfo(AllRemoteSessionsWrapperGUID, AllRemoteSessionsTargetName));
-	ActiveSessionsByInstanceId.Add(AllRemoteServersWrapperGUID, CreatedWrapperSessionInfo(AllRemoteServersWrapperGUID, AllRemoteServersTargetName));
-	ActiveSessionsByInstanceId.Add(AllRemoteClientsWrapperGUID, CreatedWrapperSessionInfo(AllRemoteClientsWrapperGUID, AllRemoteClientsTargetName));
-	ActiveSessionsByInstanceId.Add(AllSessionsWrapperGUID, CreatedWrapperSessionInfo(AllSessionsWrapperGUID, AllSessionsTargetName));
-	ActiveSessionsByInstanceId.Add(CustomSessionsWrapperGUID, CreatedWrapperSessionInfo(CustomSessionsWrapperGUID, CustomSessionsTargetName));
 }
 
-void FChaosVDRemoteSessionsManager::Shutdown()
+void FChaosVDRemoteSessionsManager::ShutdownMessagingSystem()
 {
 	if (MessageEndpoint)
 	{
@@ -196,6 +201,40 @@ void FChaosVDRemoteSessionsManager::Shutdown()
 		MessageEndpoint->Unsubscribe<FChaosVDFullSessionInfoResponseMessage>();
 		MessageEndpoint->Unsubscribe<FChaosVDChannelStateChangeResponseMessage>();
 	}
+
+	MessageBusPtr.Reset();
+	MessageEndpoint = nullptr;
+}
+
+void FChaosVDRemoteSessionsManager::Initialize()
+{
+#if !WITH_CHAOS_VISUAL_DEBUGGER_EXTERNAL_MESSAGING
+	InitializeMessagingSystem(IMessagingModule::Get().GetDefaultBus());
+#endif
+
+	ActiveSessionsByInstanceId.Add(AllRemoteSessionsWrapperGUID, CreatedWrapperSessionInfo(AllRemoteSessionsWrapperGUID, AllRemoteSessionsTargetName));
+	ActiveSessionsByInstanceId.Add(AllRemoteServersWrapperGUID, CreatedWrapperSessionInfo(AllRemoteServersWrapperGUID, AllRemoteServersTargetName));
+	ActiveSessionsByInstanceId.Add(AllRemoteClientsWrapperGUID, CreatedWrapperSessionInfo(AllRemoteClientsWrapperGUID, AllRemoteClientsTargetName));
+	ActiveSessionsByInstanceId.Add(AllSessionsWrapperGUID, CreatedWrapperSessionInfo(AllSessionsWrapperGUID, AllSessionsTargetName));
+	ActiveSessionsByInstanceId.Add(CustomSessionsWrapperGUID, CreatedWrapperSessionInfo(CustomSessionsWrapperGUID, CustomSessionsTargetName));
+}
+
+void FChaosVDRemoteSessionsManager::Shutdown()
+{
+	ShutdownMessagingSystem();
+}
+
+void FChaosVDRemoteSessionsManager::EnumerateMessageTypes(const FVisitorFunction& InVisitor)
+{
+	InVisitor(FChaosVDSessionPong::StaticStruct());
+	InVisitor(FChaosVDRecordingStatusMessage::StaticStruct());
+	InVisitor(FChaosVDSessionPing::StaticStruct());
+	InVisitor(FChaosVDStartRecordingCommandMessage::StaticStruct());
+	InVisitor(FChaosVDStopRecordingCommandMessage::StaticStruct());
+	InVisitor(FChaosVDChannelStateChangeCommandMessage::StaticStruct());
+	InVisitor(FChaosVDFullSessionInfoRequestMessage::StaticStruct());
+	InVisitor(FChaosVDFullSessionInfoResponseMessage::StaticStruct());
+	InVisitor(FChaosVDChannelStateChangeResponseMessage::StaticStruct());
 }
 
 void FChaosVDRemoteSessionsManager::StartSessionDiscovery()
@@ -337,7 +376,7 @@ void FChaosVDRemoteSessionsManager::SendPong(const FChaosVDSessionPing& InMessag
 	{
 		return;
 	}
-	
+
 	if (FChaosVDSessionPong* PongMessage = FMessageEndpoint::MakeMessage<FChaosVDSessionPong>())
 	{
 		PongMessage->InstanceId = FApp::GetInstanceId();
@@ -412,7 +451,10 @@ void FChaosVDRemoteSessionsManager::HandleSessionPongMessage(const FChaosVDSessi
 
 void FChaosVDRemoteSessionsManager::HandleSessionPingMessage(const FChaosVDSessionPing& InMessage, const TSharedRef<IMessageContext>& InContext)
 {
+	// If this instance is not running trace, don't answer to the ping as the CVD instance will not be able to do anything useful with it
+#if	!CHAOS_VISUAL_DEBUGGER_WITHOUT_TRACE
 	SendPong(InMessage);
+#endif
 }
 
 void FChaosVDRemoteSessionsManager::HandleRecordingStatusUpdateMessage(const FChaosVDRecordingStatusMessage& Message, const TSharedRef<IMessageContext>& InContext)
@@ -447,24 +489,8 @@ void FChaosVDRemoteSessionsManager::HandleRecordingStartCommandMessage(const FCh
 #if WITH_CHAOS_VISUAL_DEBUGGER
 	UE_AUTORTFM_ONCOMMIT(InMessage)
 	{
-		TArray<FString> RecordingParams;
-
-		if (InMessage.RecordingMode == EChaosVDRecordingMode::Live)
-		{
-			RecordingParams.Add(TEXT("Server"));
-
-			if (!InMessage.Target.IsEmpty())
-			{
-				RecordingParams.Add(InMessage.Target);
-			}
-		}
-		else if (InMessage.RecordingMode == EChaosVDRecordingMode::File)
-		{
-			RecordingParams.Add(TEXT("File"));
-		}
-
 		FChaosVisualDebuggerTrace::OverrideDefaultEnabledDataChannels(InMessage.DataChannelsEnabledOverrideList);
-		FChaosVDRuntimeModule::Get().StartRecording(RecordingParams);
+		FChaosVDRuntimeModule::Get().StartRecording(InMessage);
 	};
 
 #endif
diff --git a/Engine/Source/Runtime/Experimental/ChaosSolverEngine/Public/Chaos/ChaosVDEngineEditorBridge.h b/Engine/Source/Runtime/Experimental/ChaosSolverEngine/Public/Chaos/ChaosVDEngineEditorBridge.h
--- a/Engine/Source/Runtime/Experimental/ChaosSolverEngine/Public/Chaos/ChaosVDEngineEditorBridge.h
+++ b/Engine/Source/Runtime/Experimental/ChaosSolverEngine/Public/Chaos/ChaosVDEngineEditorBridge.h
@@ -22,6 +22,10 @@ namespace Chaos
 	}
 }
 
+#ifndef WITH_CHAOS_VISUAL_DEBUGGER_EXTERNAL_MESSAGING
+	#define WITH_CHAOS_VISUAL_DEBUGGER_EXTERNAL_MESSAGING 0
+#endif
+
 /** Object that bridges the gap between the ChaosVD runtime module and the Engine & CVD editor.
  * As the ChaosVDRuntime module does not have access to the engine module, this object reacts to events and performs necessary operation the runtime module cannot do directly
  */
diff --git a/Engine/Source/Runtime/Experimental/ChaosSolverEngine/Public/Chaos/ChaosVDRemoteSessionsManager.h b/Engine/Source/Runtime/Experimental/ChaosSolverEngine/Public/Chaos/ChaosVDRemoteSessionsManager.h
--- a/Engine/Source/Runtime/Experimental/ChaosSolverEngine/Public/Chaos/ChaosVDRemoteSessionsManager.h
+++ b/Engine/Source/Runtime/Experimental/ChaosSolverEngine/Public/Chaos/ChaosVDRemoteSessionsManager.h
@@ -43,21 +43,6 @@ struct FChaosVDSessionPong
 	uint8 BuildTargetType = static_cast<uint8>(EBuildTargetType::Unknown);
 };
 
-USTRUCT()
-struct FChaosVDStartRecordingCommandMessage
-{
-	GENERATED_BODY()
-
-	UPROPERTY()
-	EChaosVDRecordingMode RecordingMode = EChaosVDRecordingMode::Invalid;
-	
-	UPROPERTY()
-	FString Target;
-
-	UPROPERTY()
-	TArray<FString> DataChannelsEnabledOverrideList;
-};
-
 USTRUCT()
 struct FChaosVDStopRecordingCommandMessage
 {
@@ -253,14 +238,26 @@ public:
 
 	FChaosVDRemoteSessionsManager();
 
-	void Initialize();
-	void Shutdown();
+	UE_DEPRECATED(5.7, "Use Initialize()")
+	CHAOSSOLVERENGINE_API void Initialize(const TSharedPtr<IMessageBus>& InMessageBus);
+	CHAOSSOLVERENGINE_API void Initialize();
+	CHAOSSOLVERENGINE_API void Shutdown();
+
+	CHAOSSOLVERENGINE_API void ReInitializeMessagingSystem(const TSharedPtr<IMessageBus>& InMessageBus);
 
 	FSimpleMulticastDelegate& OnSessionsUpdated()
 	{
 		return SessionsUpdatedDelegate;
 	}
 
+	typedef TFunction<void(const UScriptStruct*)> FVisitorFunction;
+	
+	/**
+	 * Enumerate messages types allowed by this module.
+	 * @param InVisitor
+	 */
+	CHAOSSOLVERENGINE_API static void EnumerateMessageTypes(const FVisitorFunction& InVisitor);
+
 
 	/**
 	 * Starts sending ping requests that potential targets (client, servers or other editors) can answer to and report themselves as available
@@ -364,6 +361,9 @@ public:
 	CHAOSSOLVERENGINE_API static const FGuid LocalEditorSessionID;
 
 private:
+
+	void InitializeMessagingSystem(const TSharedPtr<IMessageBus>& InMessageBus);
+	void ShutdownMessagingSystem();
 	
 	/**
 	 * Returns true if this instance has controller capabilities (is either an editor or CVD Standalone, which is also an editor)
diff --git a/Engine/Source/Runtime/Experimental/ChaosVisualDebugger/Private/ChaosVDRuntimeModule.cpp b/Engine/Source/Runtime/Experimental/ChaosVisualDebugger/Private/ChaosVDRuntimeModule.cpp
--- a/Engine/Source/Runtime/Experimental/ChaosVisualDebugger/Private/ChaosVDRuntimeModule.cpp
+++ b/Engine/Source/Runtime/Experimental/ChaosVisualDebugger/Private/ChaosVDRuntimeModule.cpp
@@ -91,6 +91,7 @@ void FChaosVDRuntimeModule::StartupModule()
 			FString CVDHostAddress;
 			if (FParse::Value(FCommandLine::Get(), TEXT("CVDHost="), CVDHostAddress))
 			{
+				CVDOptions.Emplace(TEXT("Server"));
 				CVDOptions.Emplace(MoveTemp(CVDHostAddress));
 			}
 		}
@@ -117,6 +118,24 @@ void FChaosVDRuntimeModule::ShutdownModule()
 	FTraceAuxiliary::OnTraceStopped.RemoveAll(this);
 }
 
+void FChaosVDRuntimeModule::StartRecording(TConstArrayView<FString> Args)
+{
+	FChaosVDStartRecordingCommandMessage StartRecordingCommand;
+	if (Args.Num() == 0 || Args[0] == TEXT("File"))
+	{
+		StartRecordingCommand.RecordingMode = EChaosVDRecordingMode::File;
+		StartRecordingCommand.TransportMode = EChaosVDTransportMode::FileSystem;
+	}
+	else if(Args[0] == TEXT("Server"))
+	{
+		StartRecordingCommand.RecordingMode = EChaosVDRecordingMode::Live;
+		StartRecordingCommand.TransportMode = EChaosVDTransportMode::TraceServer;
+		StartRecordingCommand.Target = Args.IsValidIndex(1) ? Args[1] : TEXT("127.0.0.1");
+	}
+
+	StartRecording(StartRecordingCommand);
+}
+
 int32 FChaosVDRuntimeModule::GenerateUniqueID()
 {
 	int32 NewID = 0;
@@ -130,7 +149,7 @@ int32 FChaosVDRuntimeModule::GenerateUniqueID()
 
 FString FChaosVDRuntimeModule::GetLastRecordingFileNamePath() const
 {
-	return LastRecordingFileNamePath;
+	return FString();
 }
 
 FChaosVDTraceDetails FChaosVDRuntimeModule::GetCurrentTraceSessionDetails() const
@@ -139,7 +158,8 @@ FChaosVDTraceDetails FChaosVDRuntimeModule::GetCurrentTraceSessionDetails() cons
 	if (bool bIsConnected = FTraceAuxiliary::IsConnected(Details.SessionGuid, Details.TraceGuid))
 	{
 		Details.TraceTarget = FTraceAuxiliary::GetTraceDestinationString();
-		Details.Mode = FTraceAuxiliary::GetConnectionType() == FTraceAuxiliary::EConnectionType::File ? EChaosVDRecordingMode::File : EChaosVDRecordingMode::Live;
+		Details.Mode = GetCurrentRecordingMode();
+		Details.TransportMode = CurrentTransportMode;
 		Details.bIsConnected = bIsConnected;
 		return Details;
 	}
@@ -217,7 +237,7 @@ void FChaosVDRuntimeModule::HandleTraceConnectionEstablished()
 }
 
 
-void FChaosVDRuntimeModule::StartRecording(TConstArrayView<FString> Args)
+void FChaosVDRuntimeModule::StartRecording(const FChaosVDStartRecordingCommandMessage& InRecordingStartCommand)
 {
 	if (!ensureAlwaysMsgf(!bIsRecording, TEXT("Received a recording start command while a recording is already active")))
 	{
@@ -254,29 +274,44 @@ void FChaosVDRuntimeModule::StartRecording(TConstArrayView<FString> Args)
 	FTraceAuxiliary::FOptions TracingOptions;
 	TracingOptions.bExcludeTail = true;
 
-	if (Args.Num() == 0 || Args[0] == TEXT("File"))
+	if (InRecordingStartCommand.RecordingMode == EChaosVDRecordingMode::File)
 	{
-		LastRecordingFileNamePath.Empty();
-		GenerateRecordingFileName(LastRecordingFileNamePath);
-
-		UE_LOG(LogChaosVDRuntime, Log, TEXT("[%s] Generated trace file name [%s]"), ANSI_TO_TCHAR(__FUNCTION__), *LastRecordingFileNamePath);
+		FString FileTarget;
+		GenerateRecordingFileName(FileTarget);
 
-		bIsRecording = FTraceAuxiliary::Start(FTraceAuxiliary::EConnectionType::File, *LastRecordingFileNamePath, nullptr, &TracingOptions);
+		UE_LOG(LogChaosVDRuntime, Log, TEXT("[%hs] Generated trace file name [%s]"), __func__, *FileTarget);
 
-		LastRecordingFileNamePath = bIsRecording ? FTraceAuxiliary::GetTraceDestinationString() : TEXT("");
+		switch (InRecordingStartCommand.TransportMode)
+		{
+		case EChaosVDTransportMode::FileSystem:
+			bIsRecording = FTraceAuxiliary::Start(FTraceAuxiliary::EConnectionType::File, *FileTarget, nullptr, &TracingOptions);
+			break;
+		case EChaosVDTransportMode::TraceServer:
+			bIsRecording = FTraceAuxiliary::Start(FTraceAuxiliary::EConnectionType::Network, *InRecordingStartCommand.Target, nullptr, &TracingOptions);
+			break;
+		case EChaosVDTransportMode::Memory:
+		case EChaosVDTransportMode::Relay:
+		case EChaosVDTransportMode::Invalid:
+			UE_LOG(LogChaosVDRuntime, Error, TEXT("[%hs] Unsupported Transport mode [%s]"), __func__, *UEnum::GetValueAsString(InRecordingStartCommand.TransportMode))
+			break;
+		}
 		
 		CurrentRecordingMode = EChaosVDRecordingMode::File;
 	}
-	else if(Args[0] == TEXT("Server"))
+	else if(InRecordingStartCommand.RecordingMode == EChaosVDRecordingMode::Live)
 	{
-		const FString Target = Args.IsValidIndex(1) ? Args[1] : TEXT("127.0.0.1");
-
-		LastRecordingFileNamePath = Target;
-
-		bIsRecording = FTraceAuxiliary::Start(
-		FTraceAuxiliary::EConnectionType::Network,
-		*Target,
-		nullptr, &TracingOptions);
+		switch (InRecordingStartCommand.TransportMode)
+		{
+		case EChaosVDTransportMode::Memory:
+		case EChaosVDTransportMode::TraceServer:
+			bIsRecording = FTraceAuxiliary::Start(FTraceAuxiliary::EConnectionType::Network, *InRecordingStartCommand.Target, nullptr, &TracingOptions);
+			break;
+		case EChaosVDTransportMode::Relay:
+		case EChaosVDTransportMode::Invalid:
+		case EChaosVDTransportMode::FileSystem:
+		UE_LOG(LogChaosVDRuntime, Error, TEXT("[%hs] Unsupported Transport mode [%s]"), __func__, *UEnum::GetValueAsString(InRecordingStartCommand.TransportMode))
+			break;
+		}
 
 		CurrentRecordingMode = EChaosVDRecordingMode::Live;
 	}
@@ -290,6 +325,8 @@ void FChaosVDRuntimeModule::StartRecording(TConstArrayView<FString> Args)
 
 	if (ensureMsgf(bIsRecording, TEXT("Failed to start CVD recording | Reason [%s]"), *FailureReason.ToString()))
 	{
+		CurrentTransportMode = InRecordingStartCommand.TransportMode;
+
 		if (ensureMsgf(!OnConnectionDelegateHandle.IsValid(), TEXT("Starting a trace while we are waiting for a pending connection")))
 		{
 			OnConnectionDelegateHandle = FTraceAuxiliary::OnConnection.AddRaw(this, &FChaosVDRuntimeModule::HandleTraceConnectionEstablished);
@@ -326,6 +363,9 @@ void FChaosVDRuntimeModule::StopRecording()
 	FTraceAuxiliary::OnTraceStopped.RemoveAll(this);
 	OnTraceStoppedDelegateHandle = FDelegateHandle();
 
+	CurrentRecordingMode = EChaosVDRecordingMode::Invalid;
+	CurrentTransportMode = EChaosVDTransportMode::Invalid;
+
 	if (!bIsRecording)
 	{
 		UE_LOG(LogChaosVDRuntime, Warning, TEXT("[%s] Attempted to stop recorded when there is no CVD recording active."), ANSI_TO_TCHAR(__FUNCTION__));
diff --git a/Engine/Source/Runtime/Experimental/ChaosVisualDebugger/Public/ChaosVDRecordingDetails.h b/Engine/Source/Runtime/Experimental/ChaosVisualDebugger/Public/ChaosVDRecordingDetails.h
--- a/Engine/Source/Runtime/Experimental/ChaosVisualDebugger/Public/ChaosVDRecordingDetails.h
+++ b/Engine/Source/Runtime/Experimental/ChaosVisualDebugger/Public/ChaosVDRecordingDetails.h
@@ -14,6 +14,39 @@ enum class EChaosVDRecordingMode : uint8
 	File
 };
 
+/** Available transport modes for CVD Trace Data*/
+UENUM()
+enum class EChaosVDTransportMode : uint8
+{
+	Invalid,
+	/** The data will be traced directly to a file in the local file system */
+	FileSystem,
+	/** Data will be traced to the selected trace store server */
+	TraceServer,
+	/** NOT IMPLEMENTED - Data will be traced directly to memory, accessible via sockets */
+	Memory,
+	/** NOT IMPLEMENTED - Data will be traced via the trace relay system using a custom transport, like normal UE Networking */
+	Relay
+};
+
+USTRUCT()
+struct FChaosVDStartRecordingCommandMessage
+{
+	GENERATED_BODY()
+
+	UPROPERTY()
+	EChaosVDRecordingMode RecordingMode = EChaosVDRecordingMode::Invalid;
+
+	UPROPERTY()
+	EChaosVDTransportMode TransportMode = EChaosVDTransportMode::Invalid;
+	
+	UPROPERTY()
+	FString Target;
+
+	UPROPERTY()
+	TArray<FString> DataChannelsEnabledOverrideList;
+};
+
 USTRUCT()
 struct FChaosVDTraceDetails
 {
@@ -33,4 +66,7 @@ struct FChaosVDTraceDetails
 
 	UPROPERTY()
 	EChaosVDRecordingMode Mode = EChaosVDRecordingMode::Invalid;
+
+	UPROPERTY()
+	EChaosVDTransportMode TransportMode = EChaosVDTransportMode::Invalid;
 };
\ No newline at end of file
diff --git a/Engine/Source/Runtime/Experimental/ChaosVisualDebugger/Public/ChaosVDRuntimeModule.h b/Engine/Source/Runtime/Experimental/ChaosVisualDebugger/Public/ChaosVDRuntimeModule.h
--- a/Engine/Source/Runtime/Experimental/ChaosVisualDebugger/Public/ChaosVDRuntimeModule.h
+++ b/Engine/Source/Runtime/Experimental/ChaosVisualDebugger/Public/ChaosVDRuntimeModule.h
@@ -41,15 +41,23 @@ public:
 	UE_API virtual void ShutdownModule() override;
 
 	/** Starts a CVD recording by starting a Trace session. It will stop any existing trace session
-	 * @param Args : Arguments array provided by the commandline. Used to determine if we want to record to file or a local trace server
+	 * @param Args : Arguments array usually provided by the commandline. Used to determine if we want to record to file or a local trace server
 	 */
 	UE_API void StartRecording(TConstArrayView<FString> Args);
+
+	/** Starts a CVD recording by starting a Trace session. It will stop any existing trace session
+	 * @param InRecordingStartCommand : Used to determine if we want to record to file or a local trace server among other settings
+	 */
+	UE_API void StartRecording(const FChaosVDStartRecordingCommandMessage& InRecordingStartCommand);
 	
 	/* Stops an active recording */
 	UE_API void StopRecording();
 
 	/** Returns true if we are currently recording a Physics simulation */
-	bool IsRecording() const { return bIsRecording; }
+	bool IsRecording() const
+	{
+		return bIsRecording;
+	}
 
 	/** Returns a unique ID used to be used to identify CVD (Chaos Visual Debugger) data */
 	UE_API int32 GenerateUniqueID();
@@ -115,9 +123,13 @@ public:
 	}
 
 	/** Returns the accumulated recording time in seconds since the recording started */
-	float GetAccumulatedRecordingTime() const { return AccumulatedRecordingTime; }
+	float GetAccumulatedRecordingTime() const
+	{
+		return AccumulatedRecordingTime;
+	}
 
-	/** Returns the full path of the active recording file*/
+	/** [Deprecated] Returns the full path of the active recording file*/
+	UE_DEPRECATED(5.7, "This method is no longer used. Use GetCurrentTraceSessionDetails")
 	UE_API FString GetLastRecordingFileNamePath() const;
 
 	UE_API FChaosVDTraceDetails GetCurrentTraceSessionDetails() const;
@@ -127,12 +139,13 @@ public:
 		return CurrentRecordingMode;
 	}
 
+	/** Finds a valid file name for a new file - Used to generate the file name for the Trace recording */
+	UE_API void GenerateRecordingFileName(FString& OutFileName);
+
 private:
 
 	/** Stops the current Trace session */
 	UE_API void StopTrace();
-	/** Finds a valid file name for a new file - Used to generate the file name for the Trace recording */
-	UE_API void GenerateRecordingFileName(FString& OutFileName);
 
 	/** Queues a full Capture of the simulation on the next frame */
 	UE_API bool RequestFullCapture(float DeltaTime);
@@ -172,11 +185,11 @@ private:
 
 	std::atomic<int32> LastGeneratedID;
 
-	FString LastRecordingFileNamePath;
-
 	TMap<FString, bool> OriginalTraceChannelsState;
 
 	EChaosVDRecordingMode CurrentRecordingMode = EChaosVDRecordingMode::Invalid;
+	EChaosVDTransportMode CurrentTransportMode = EChaosVDTransportMode::Invalid;
+
 
 	static UE_API FTransactionallySafeRWLock DelegatesRWLock;
 };
