[ChaosVD] Fix CIS issue in shipping editors.

The previous change to add support for direct trace needs access to the ChaosVDRuntime module, which currently is a dummy module in shipping editor configurations.

This fix just disables the code that access it, until we enable CVD support in shipping editors shortly.

This also fixes a PVS warning in the same block of code in nin-shipping builds

diff --git a/Engine/Plugins/ChaosVD/Source/ChaosVD/Private/Trace/ChaosVDTraceManager.cpp b/Engine/Plugins/ChaosVD/Source/ChaosVD/Private/Trace/ChaosVDTraceManager.cpp
--- a/Engine/Plugins/ChaosVD/Source/ChaosVD/Private/Trace/ChaosVDTraceManager.cpp
+++ b/Engine/Plugins/ChaosVD/Source/ChaosVD/Private/Trace/ChaosVDTraceManager.cpp
@@ -337,6 +337,8 @@ TUniquePtr<FArchiveFileWriterGeneric> FChaosVDTraceManager::OpenTraceFileForWrit
 {
 	TUniquePtr<FArchiveFileWriterGeneric> FileWriter = nullptr;
 
+#if WITH_CHAOS_VISUAL_DEBUGGER
+
 	IPlatformFile& PlatformFile = FPlatformFileManager::Get().GetPlatformFile();
 	FChaosVDRuntimeModule& CVDRuntimeModule = FChaosVDRuntimeModule::Get();
 	constexpr int32 MaxAttempts = 10;
@@ -347,8 +349,7 @@ TUniquePtr<FArchiveFileWriterGeneric> FChaosVDTraceManager::OpenTraceFileForWrit
 		CVDRuntimeModule.GenerateRecordingFileName(OutFileName);
 		OutFileName = FPaths::Combine(FPaths::ProfilingDir(), OutFileName);
 
-		IFileHandle* FileHandle = FileHandle = PlatformFile.OpenWrite(*OutFileName);
-		if (FileHandle)
+		if (IFileHandle* FileHandle = PlatformFile.OpenWrite(*OutFileName))
 		{
 			FileWriter = MakeUnique<FArchiveFileWriterGeneric>(FileHandle, *OutFileName, FileHandle->Tell());
 			break;
@@ -359,6 +360,8 @@ TUniquePtr<FArchiveFileWriterGeneric> FChaosVDTraceManager::OpenTraceFileForWrit
 		FPlatformProcess::Sleep(0.1f);
 	}
 
+#endif
+
 	return MoveTemp(FileWriter);
 }
 
