[ChaosVD] Fix for an issue where attempting to trace multiple process at the same time would fail to playback.

It seems that now we are not getting the expected session name from the trace store. I need to investigate that issue separately, but meanwhile generating a session name using the host address and session id allows the system to work again.

Also, this CL fixes a crash that can happen when a live trace is stopped after a previous failure. (If you attempt to connect to the trace store with an empty address, the call to it will crash)

diff --git a/Engine/Plugins/ChaosVD/Source/ChaosVD/Private/Trace/ChaosVDTraceManager.cpp b/Engine/Plugins/ChaosVD/Source/ChaosVD/Private/Trace/ChaosVDTraceManager.cpp
--- a/Engine/Plugins/ChaosVD/Source/ChaosVD/Private/Trace/ChaosVDTraceManager.cpp
+++ b/Engine/Plugins/ChaosVD/Source/ChaosVD/Private/Trace/ChaosVDTraceManager.cpp
@@ -160,8 +160,13 @@ const UE::Trace::FStoreClient::FSessionInfo* FChaosVDTraceManager::GetTraceSessi
 FString FChaosVDTraceManager::GetTraceFileNameFromStoreForSession(FStringView InSessionHost, uint32 SessionID)
 {
 	using namespace UE::Trace;
-	FStoreClient* StoreClient = FStoreClient::Connect(InSessionHost.GetData());
 
+	if (InSessionHost.IsEmpty())
+	{
+		return FString();
+	}
+	
+	FStoreClient* StoreClient = FStoreClient::Connect(InSessionHost.GetData());
 	if (!StoreClient)
 	{
 		return FString();
@@ -308,13 +313,7 @@ FString FChaosVDTraceManager::ConnectToLiveSession(FStringView InSessionHost, ui
 	}
 
 	using namespace UE::Trace;
-
-	FString TraceName = GetTraceFileNameFromStoreForSession(InSessionHost, SessionID);
-	if (TraceName.IsEmpty())
-	{
-		UE_LOG(LogChaosVDEditor, Error, TEXT("[%hs] Failed to get the trace file name from the trace store | Host [%s] | Session ID [%u]"), __func__, InSessionHost.GetData(), SessionID);
-		return FString();
-	}
+	FString LiveSessionName = FString::Printf(TEXT("LiveSession[%.*s - %u]"), InSessionHost.Len(), InSessionHost.GetData(), SessionID);
 
 	FStoreClient::FTraceData TraceData = GetTraceDataStreamFromStore(InSessionHost, SessionID);
 	if (!TraceData.IsValid())
@@ -323,9 +322,9 @@ FString FChaosVDTraceManager::ConnectToLiveSession(FStringView InSessionHost, ui
 		return FString();
 	}
 
-	if (ConnectToLiveSession_Internal(SessionID, ExistingRecordingPtr, TraceName, MoveTemp(TraceData)))
+	if (ConnectToLiveSession_Internal(SessionID, ExistingRecordingPtr, LiveSessionName, MoveTemp(TraceData)))
 	{
-		 return TraceName;
+		 return LiveSessionName;
 	}
 
 	UE_LOG(LogChaosVDEditor, Error, TEXT("[%hs] Failed connect to session | Host [%s] | Session ID [%u]"), __func__, InSessionHost.GetData(), SessionID);
diff --git a/Engine/Plugins/ChaosVD/Source/ChaosVD/Private/Widgets/SChaosVDRecordingControls.cpp b/Engine/Plugins/ChaosVD/Source/ChaosVD/Private/Widgets/SChaosVDRecordingControls.cpp
--- a/Engine/Plugins/ChaosVD/Source/ChaosVD/Private/Widgets/SChaosVDRecordingControls.cpp
+++ b/Engine/Plugins/ChaosVD/Source/ChaosVD/Private/Widgets/SChaosVDRecordingControls.cpp
@@ -594,9 +594,16 @@ void SChaosVDRecordingControls::HandleRecordingStop(TWeakPtr<FChaosVDSessionInfo
 
 	if (!bIsLiveSession && !EnumHasAnyFlags(SessionInfoPtr->GetSessionTypeAttributes(), EChaosVDRemoteSessionAttributes::IsMultiSessionWrapper))
 	{
-		if (FMessageDialog::Open(EAppMsgType::YesNo, LOCTEXT("OpenLastRecordingMessage", "Do you want to load the recorded file now? ")) == EAppReturnType::Yes)
+		if (!TraceTarget.IsEmpty())
 		{
-			MainTabSharedPtr->GetChaosVDEngineInstance()->LoadRecording(TraceTarget, EChaosVDLoadRecordedDataMode::SingleSource);			
+			if (FMessageDialog::Open(EAppMsgType::YesNo, LOCTEXT("OpenLastRecordingMessage", "Do you want to load the recorded file now? ")) == EAppReturnType::Yes)
+			{
+				MainTabSharedPtr->GetChaosVDEngineInstance()->LoadRecording(TraceTarget, EChaosVDLoadRecordedDataMode::SingleSource);			
+			}
+		}
+		else
+		{
+			FMessageDialog::Open(EAppMsgType::Ok, LOCTEXT("FailedToOpenLastRecordingMessage", "Failed to obtain the file path to the recording."));
 		}
 	}
 }
