[ChaosVD] Fix for CIS issue in shipping editor builds

ITraceDataRelayTransport is not available in these configurations

diff --git a/Engine/Plugins/ChaosVD/Source/ChaosVD/Private/Trace/ChaosVDTraceManager.cpp b/Engine/Plugins/ChaosVD/Source/ChaosVD/Private/Trace/ChaosVDTraceManager.cpp
--- a/Engine/Plugins/ChaosVD/Source/ChaosVD/Private/Trace/ChaosVDTraceManager.cpp
+++ b/Engine/Plugins/ChaosVD/Source/ChaosVD/Private/Trace/ChaosVDTraceManager.cpp
@@ -98,16 +98,13 @@ namespace Chaos::VD
 	FRelayDataStream::FRelayDataStream(FGuid InRemoteSessionID)
 		: CancelEvent(FPlatformProcess::GetSynchEventFromPool(false)), RemoteSessionID(InRemoteSessionID)
 	{
-
 #if WITH_CHAOS_VISUAL_DEBUGGER
 		TSharedPtr<ITraceDataRelayTransport> RelayTransportInstance = FChaosVDEngineEditorBridge::Get().GetRelayTransportInstance();
-#else
-		TSharedPtr<ITraceDataRelayTransport> RelayTransportInstance = nullptr;
-#endif
 		if (ensure(RelayTransportInstance))
 		{
 			RelayTransportInstance->RegisterRelayDataReceiverForSessionID(RemoteSessionID, ITraceDataRelayTransport::FProcessReceivedRelayDataDelegate::CreateRaw(this, &FRelayDataStream::EnqueueRelayedData));
 		}
+#endif
 	}
 
 	FRelayDataStream::~FRelayDataStream()
@@ -116,14 +113,11 @@ namespace Chaos::VD
 
 #if WITH_CHAOS_VISUAL_DEBUGGER
 		TSharedPtr<ITraceDataRelayTransport> RelayTransportInstance = FChaosVDEngineEditorBridge::Get().GetRelayTransportInstance();
-#else
-		TSharedPtr<ITraceDataRelayTransport> RelayTransportInstance = nullptr;
-#endif
-
 		if (RelayTransportInstance)
 		{
 			RelayTransportInstance->UnregisterRelayDataReceiverForSessionID(RemoteSessionID);
-		}
+		}	
+#endif
 	}
 
 	void FRelayDataStream::EnqueueRelayedData(TArray<uint8>&& InTraceDataBuffer)
