[ChaosVD] Fixed a typo in the Particle metadata details customization

diff --git a/Engine/Plugins/ChaosVD/Source/ChaosVD/Private/DetailsCustomizations/ChaosVDParticleMetadataCustomization.cpp b/Engine/Plugins/ChaosVD/Source/ChaosVD/Private/DetailsCustomizations/ChaosVDParticleMetadataCustomization.cpp
--- a/Engine/Plugins/ChaosVD/Source/ChaosVD/Private/DetailsCustomizations/ChaosVDParticleMetadataCustomization.cpp
+++ b/Engine/Plugins/ChaosVD/Source/ChaosVD/Private/DetailsCustomizations/ChaosVDParticleMetadataCustomization.cpp
@@ -91,7 +91,7 @@ TSharedRef<SWidget> FChaosVDParticleMetadataCustomization::GenerateBrowseParticl
 			[
 				SNew(STextBlock)
 				.TextStyle(FAppStyle::Get(), "SmallButtonText")
-				.Text(LOCTEXT("BrowseParticleOwnerText","Attempt To Find Find in Content Browser"))
+				.Text(LOCTEXT("BrowseParticleOwnerText","Attempt To Find in Content Browser"))
 			]
 		]
 	];
