[ChaosVD] Fixed an issue where you can load the same recording twice.

Also fixed an issue where the first frame was not automatically played properly.

There were multiple problems

1- When attempting to load the same recording twice, both will have the same session id, therefore we will start loading the trace data, but then we would unload it immediatly because there was an issue with the order of operations between load and close active sessions.
2- When a recording was unloaded, the UI was not properly updated and reset. This happened because there was cached state not properly cleared out, and then if the unload happens as part of a load call, we never get to process the UI updates to reflect the unload, before loading the new file
3- The first frame was not being played properly, because the existing implementation didn't check correctly if we had anything to play or if the data was still loading.

diff --git a/Engine/Plugins/ChaosVD/Source/ChaosVD/Private/ChaosVDEngine.cpp b/Engine/Plugins/ChaosVD/Source/ChaosVD/Private/ChaosVDEngine.cpp
--- a/Engine/Plugins/ChaosVD/Source/ChaosVD/Private/ChaosVDEngine.cpp
+++ b/Engine/Plugins/ChaosVD/Source/ChaosVD/Private/ChaosVDEngine.cpp
@@ -224,6 +224,11 @@ void FChaosVDEngine::LoadRecording_Internal(const TFunction<FString(const TShare
 
 void FChaosVDEngine::LoadRecording(const FString& FilePath, EChaosVDLoadRecordedDataMode LoadingMode)
 {
+	if (LoadingMode == EChaosVDLoadRecordedDataMode::SingleSource)
+	{
+		CloseActiveTraceSessions();
+	}
+
 	if (FilePath.EndsWith(TEXT("cvdmulti")))
 	{
 		LoadCombinedMultiRecording(FilePath);
@@ -241,6 +246,8 @@ void FChaosVDEngine::LoadRecording(const FString& FilePath, EChaosVDLoadRecorded
 
 void FChaosVDEngine::LoadCombinedMultiRecording(const FString& FilePath)
 {
+	CloseActiveTraceSessions();
+
 	if (!ensure(FilePath.EndsWith(TEXT("cvdmulti"))))
 	{
 		return;
@@ -274,6 +281,11 @@ void FChaosVDEngine::LoadCombinedMultiRecording(const FString& FilePath)
 
 bool FChaosVDEngine::ConnectToLiveSession(uint32 SessionID, const FString& InSessionAddress, EChaosVDLoadRecordedDataMode LoadingMode)
 {
+	if (LoadingMode == EChaosVDLoadRecordedDataMode::SingleSource)
+	{
+		CloseActiveTraceSessions();
+	}
+	
 	FChaosVDTraceSessionDescriptor NewSessionFromFileDescriptor;
 
 	NewSessionFromFileDescriptor.bIsLiveSession = true;
@@ -299,6 +311,11 @@ bool FChaosVDEngine::ConnectToLiveSession(uint32 SessionID, const FString& InSes
 
 bool FChaosVDEngine::ConnectToLiveSession_Direct(EChaosVDLoadRecordedDataMode LoadingMode)
 {
+	if (LoadingMode == EChaosVDLoadRecordedDataMode::SingleSource)
+	{
+		CloseActiveTraceSessions();
+	}
+
 	FChaosVDTraceSessionDescriptor NewSessionFromFileDescriptor;
 
 	NewSessionFromFileDescriptor.bIsLiveSession = true;
@@ -324,6 +341,11 @@ bool FChaosVDEngine::ConnectToLiveSession_Direct(EChaosVDLoadRecordedDataMode Lo
 
 bool FChaosVDEngine::ConnectToLiveSession_Relay(FGuid RemoteSessionID, EChaosVDLoadRecordedDataMode LoadingMode)
 {
+	if (LoadingMode == EChaosVDLoadRecordedDataMode::SingleSource)
+	{
+		CloseActiveTraceSessions();
+	}
+
 	FChaosVDTraceSessionDescriptor NewSessionFromFileDescriptor;
 
 	NewSessionFromFileDescriptor.bIsLiveSession = true;
@@ -350,11 +372,6 @@ bool FChaosVDEngine::ConnectToLiveSession_Relay(FGuid RemoteSessionID, EChaosVDL
 
 void FChaosVDEngine::OpenSession(const FChaosVDTraceSessionDescriptor& SessionDescriptor, EChaosVDLoadRecordedDataMode LoadingMode)
 {
-	if (LoadingMode == EChaosVDLoadRecordedDataMode::SingleSource)
-	{
-		CloseActiveTraceSessions();
-	}
-
 	CurrentSessionDescriptors.Emplace(SessionDescriptor);
 
 	if (CurrentSessionDescriptors.Num() <= 1)
diff --git a/Engine/Plugins/ChaosVD/Source/ChaosVD/Private/ChaosVDPlaybackController.cpp b/Engine/Plugins/ChaosVD/Source/ChaosVD/Private/ChaosVDPlaybackController.cpp
--- a/Engine/Plugins/ChaosVD/Source/ChaosVD/Private/ChaosVDPlaybackController.cpp
+++ b/Engine/Plugins/ChaosVD/Source/ChaosVD/Private/ChaosVDPlaybackController.cpp
@@ -91,6 +91,8 @@ bool FChaosVDPlaybackController::LoadChaosVDRecordingFromTraceSession(const FCha
 
 void FChaosVDPlaybackController::UnloadCurrentRecording(EChaosVDUnloadRecordingFlags UnloadOptions)
 {
+	ensure(IsInGameThread());
+
 	RecordingLastSeenTimeUpdatedAsCycle = 0;
 
 	TrackInfoUpdateGTQueue.Empty();
@@ -110,12 +112,24 @@ void FChaosVDPlaybackController::UnloadCurrentRecording(EChaosVDUnloadRecordingF
 
 	// This will make sure the cached data used by the UI is up to date.
 	// It already handles internally an unloaded recording, in which case the cached data will be properly reset
-
+	TrackInfoPerType.Reset();
 	HandleCurrentRecordingUpdated();
 
+	// Once the cached track data is cleared out, we need to mimic what the normal update loop does so the updates are sent to the UI
+	// Before we attempt to load a new recording
+	const TWeakPtr<FChaosVDPlaybackController> ThisWeakPtr = DoesSharedInstanceExist() ? AsWeak() : nullptr;
+	if (TSharedPtr<const FChaosVDTrackInfo> GameTrackInfo = GetTrackInfo(EChaosVDTrackType::Game, GameTrackID))
+	{
+		TSharedRef<const FChaosVDTrackInfo> GTTrackInfoAsSharedRef = GameTrackInfo.ToSharedRef();
+		TrySetActiveTrack(GTTrackInfoAsSharedRef);
+		EnqueueTrackInfoUpdate(GTTrackInfoAsSharedRef, PlaybackSelfInstigatorID);
+		ProcessQueuedTrackInfoUpdates(ThisWeakPtr);
+	}
+	
 	if (EnumHasAnyFlags(UnloadOptions, EChaosVDUnloadRecordingFlags::BroadcastChanges))
 	{
-		bHasPendingGTUpdateBroadcast = true;
+		BroadcastGTDataUpdate(ThisWeakPtr);
+		bHasPendingGTUpdateBroadcast = false;
 	}
 
 	bPlayedFirstFrame = false;
@@ -759,6 +773,36 @@ void FChaosVDPlaybackController::GetAvailableTrackInfosAtTrackFrame_AssumesLocke
 	}
 }
 
+void FChaosVDPlaybackController::BroadcastGTDataUpdate(const TWeakPtr<FChaosVDPlaybackController>& InControllerWeakPtr)
+{
+	ensure(IsInGameThread());
+
+	FChaosVDExtensionsManager::Get().EnumerateExtensions([InControllerWeakPtr](const TSharedRef<FChaosVDExtension>& Extension)
+		{
+			Extension->HandlePlaybackControllerDataUpdated(InControllerWeakPtr);
+			return true;
+		});
+
+	ControllerUpdatedDelegate.Broadcast(InControllerWeakPtr);
+}
+
+void FChaosVDPlaybackController::ProcessQueuedTrackInfoUpdates(const TWeakPtr<FChaosVDPlaybackController>& InControllerWeakPtr)
+{
+	ensure(IsInGameThread());
+
+	FChaosVDQueuedTrackInfoUpdate TrackInfoUpdate;
+	while (TrackInfoUpdateGTQueue.Dequeue(TrackInfoUpdate))
+	{
+		OnTrackFrameUpdated().Broadcast(InControllerWeakPtr, TrackInfoUpdate.TrackInfo, TrackInfoUpdate.InstigatorID);
+
+		FChaosVDExtensionsManager::Get().EnumerateExtensions([InControllerWeakPtr, &TrackInfoUpdate](const TSharedRef<FChaosVDExtension>& Extension)
+			{
+				Extension->HandleControllerTrackFrameUpdated(InControllerWeakPtr, TrackInfoUpdate.TrackInfo, TrackInfoUpdate.InstigatorID);
+				return true;
+			});
+	}
+}
+
 bool FChaosVDPlaybackController::Tick(float DeltaTime)
 {
 	TickPlayback(DeltaTime);
@@ -784,29 +828,13 @@ bool FChaosVDPlaybackController::Tick(float DeltaTime)
 
 	if (bHasPendingGTUpdateBroadcast)
 	{
-		FChaosVDExtensionsManager::Get().EnumerateExtensions([ThisWeakPtr](const TSharedRef<FChaosVDExtension>& Extension)
-		{
-			Extension->HandlePlaybackControllerDataUpdated(ThisWeakPtr);
-			return true;
-		});
-
-		ControllerUpdatedDelegate.Broadcast(ThisWeakPtr);
+		BroadcastGTDataUpdate(ThisWeakPtr);
 		bHasPendingGTUpdateBroadcast = false;
 	}
 
 	if (!TrackInfoUpdateGTQueue.IsEmpty())
 	{
-		FChaosVDQueuedTrackInfoUpdate TrackInfoUpdate;
-		while (TrackInfoUpdateGTQueue.Dequeue(TrackInfoUpdate))
-		{
-			OnTrackFrameUpdated().Broadcast(ThisWeakPtr, TrackInfoUpdate.TrackInfo, TrackInfoUpdate.InstigatorID);
-
-			FChaosVDExtensionsManager::Get().EnumerateExtensions([ThisWeakPtr, &TrackInfoUpdate](const TSharedRef<FChaosVDExtension>& Extension)
-			{
-				Extension->HandleControllerTrackFrameUpdated(ThisWeakPtr, TrackInfoUpdate.TrackInfo, TrackInfoUpdate.InstigatorID);
-				return true;
-			});
-		}
+		ProcessQueuedTrackInfoUpdates(ThisWeakPtr);
 	}
 
 	if (bIsRecordingLoaded)
@@ -814,21 +842,26 @@ bool FChaosVDPlaybackController::Tick(float DeltaTime)
 		// Load at least the first frame
 		if (!bPlayedFirstFrame)
 		{
-			if (IsPlayingLiveSession())
+			// Make sure we have valid data for the first frame before trying to play it
+			const TSharedPtr<const FChaosVDTrackInfo> GameTrackInfo = GetTrackInfo(EChaosVDTrackType::Game, GameTrackID);
+			const bool bHasValidGameThreadTrackData = GameTrackInfo && GameTrackInfo->MaxFrames > 0;
+			const TrackInfoByIDMap* TracksByType = bHasValidGameThreadTrackData ? TrackInfoPerType.Find(EChaosVDTrackType::Solver) : nullptr;
+
+			if (const bool bHasSolverTrackData = TracksByType && !TracksByType->IsEmpty())
 			{
-				if (const TSharedPtr<const FChaosVDTrackInfo> GameTrackInfo = GetTrackInfo(EChaosVDTrackType::Game, GameTrackID))
+				if (IsPlayingLiveSession())
 				{
 					HandleFramePlaybackControlInput(EChaosVDPlaybackButtonsID::Play, GameTrackInfo.ToSharedRef(), PlaybackSelfInstigatorID);
 					bPlayedFirstFrame = true;
 				}
-			}
-			else if (!LoadedRecording->IsEmpty())
-			{
-				constexpr int32 GameFrameToLoad = 0;
-				// Game frames do not have stages
-				constexpr int32 SolverStageToLoad = 0;
-				GoToTrackFrameAndSync(PlaybackSelfInstigatorID, EChaosVDTrackType::Game, GameTrackID, GameFrameToLoad, SolverStageToLoad);
-				bPlayedFirstFrame = true;
+				else if (!LoadedRecording->IsEmpty())
+				{
+					constexpr int32 GameFrameToLoad = 0;
+					// Game frames do not have stages
+					constexpr int32 SolverStageToLoad = 0;
+					GoToTrackFrameAndSync(PlaybackSelfInstigatorID, EChaosVDTrackType::Game, GameTrackID, GameFrameToLoad, SolverStageToLoad);
+					bPlayedFirstFrame = true;
+				}
 			}
 		}
 
diff --git a/Engine/Plugins/ChaosVD/Source/ChaosVD/Public/ChaosVDPlaybackController.h b/Engine/Plugins/ChaosVD/Source/ChaosVD/Public/ChaosVDPlaybackController.h
--- a/Engine/Plugins/ChaosVD/Source/ChaosVD/Public/ChaosVDPlaybackController.h
+++ b/Engine/Plugins/ChaosVD/Source/ChaosVD/Public/ChaosVDPlaybackController.h
@@ -451,6 +451,18 @@ public:
 	CHAOSVD_API FChaosVDPlaybackEngineSnapshot GeneratePlaybackEngineSnapshot();
 
 protected:
+	/**
+	 * Broadcast from the game thread data in this controller was updated
+	 * @param InControllerWeakPtr A pre created weak ptr instance to this controller
+	 */
+	void BroadcastGTDataUpdate(const TWeakPtr<FChaosVDPlaybackController>& InControllerWeakPtr);
+
+	/**
+	 * Process and broadcast from the game thread track info updates
+	 * @param InControllerWeakPtr A pre created weak ptr instance to this controller
+	 */
+	void ProcessQueuedTrackInfoUpdates(const TWeakPtr<FChaosVDPlaybackController>& InControllerWeakPtr);
+	
 	/** Updates (or adds) solvers data from the loaded recording to the solver tracks */
 	void UpdateSolverTracksData();
 
