[ChaosVD] Fix CIS issue in shipping editors

The trace relay transport types are not available in non-shipping configurations

diff --git a/Engine/Plugins/ChaosVD/Source/ChaosVD/Private/Widgets/SChaosVDRecordingControls.cpp b/Engine/Plugins/ChaosVD/Source/ChaosVD/Private/Widgets/SChaosVDRecordingControls.cpp
--- a/Engine/Plugins/ChaosVD/Source/ChaosVD/Private/Widgets/SChaosVDRecordingControls.cpp
+++ b/Engine/Plugins/ChaosVD/Source/ChaosVD/Private/Widgets/SChaosVDRecordingControls.cpp
@@ -845,6 +845,7 @@ void SChaosVDRecordingControls::SetSessionRecordingState(bool bIsRecording, ECha
 
 				RemoteSessionManager->SendStartRecordingCommand(SessionInfoRef->Address, RecordingParams);
 
+#if WITH_CHAOS_VISUAL_DEBUGGER
 				// Once the start recording command is issued, we can try to connect to the created session (if everything went well).
 				// If it didn't go well, this will take care of update the UI to notify the user
 				ExecuteAsyncConnectionAttemptTaskWithRetry(RemainingRetries, [WeakThis = AsWeak(), SessionInstanceId = SessionInfoRef->InstanceId]()
@@ -893,6 +894,7 @@ void SChaosVDRecordingControls::SetSessionRecordingState(bool bIsRecording, ECha
 					}
 				},
 				RecordingAttemptFailedCallback);
+#endif
 			}
 			else
 			{
