[ChaosVD] Fixed an error where in network tick sync mode, we might not clamp the track offsets to the supported limits.

Also added a verbose log entry to print each time we fallback to timestamp mode

diff --git a/Engine/Plugins/ChaosVD/Source/ChaosVD/Private/ChaosVDPlaybackController.cpp b/Engine/Plugins/ChaosVD/Source/ChaosVD/Private/ChaosVDPlaybackController.cpp
--- a/Engine/Plugins/ChaosVD/Source/ChaosVD/Private/ChaosVDPlaybackController.cpp
+++ b/Engine/Plugins/ChaosVD/Source/ChaosVD/Private/ChaosVDPlaybackController.cpp
@@ -546,7 +546,7 @@ int32 FChaosVDPlaybackController::ConvertCurrentFrameToOtherTrackFrame_AssumesLo
 
 											// Server tracks will not have a frame offset set, so we need to clamp it to 0 first
 											const int32 ToTrackNetworkTick = InToTrack->bIsServer && ToSolverFrameData->NetworkTickOffset == INDEX_NONE ? 0 : ToSolverFrameData->NetworkTickOffset;
-											const int32 FromTrackNetworkTick = InFromTrack->bIsServer && ToSolverFrameData->NetworkTickOffset == INDEX_NONE ? 0 : FromSolverFrameData->NetworkTickOffset;
+											const int32 FromTrackNetworkTick = InFromTrack->bIsServer && FromSolverFrameData->NetworkTickOffset == INDEX_NONE ? 0 : FromSolverFrameData->NetworkTickOffset;
 
 											// This works for Client --> Server conversion because in that case we want to add the frame offset.
 											// As the tick offset in server tracks is 0,  the following calculation will return a negative offset, which it will result in the intended addition in the last caluclation
@@ -570,15 +570,16 @@ int32 FChaosVDPlaybackController::ConvertCurrentFrameToOtherTrackFrame_AssumesLo
 									}
 								}
 
-								// If we failed to re-map the track based on the network tick, just fallback to the timestamp sync logic.
-								// The UI shows a warning on all the conditions this can happen, which is when any of the tracks does not have net tick offset data
-
 								if (ToFrame != INDEX_NONE)
 								{
 									return LoadedRecording->GetLowestSolverFrameNumberAtNetworkFrameNumber_AssumesLocked(InToTrack->TrackID, ToFrame);
 								}
+
+								// If we failed to re-map the track based on the network tick, just fallback to the timestamp sync logic.
+								// The UI shows a warning on all the conditions this can happen, which is when any of the tracks does not have net tick offset data
 								if (FromSolverFrameData)
 								{
+									UE_LOG(LogChaosVDEditor, Verbose, TEXT("[%hs] Failed to map frame number from recorded network tick offsets. Falling back to timestamp mode | From Track [%s] To Track [%s]"), __func__, *InFromTrack->TrackName.ToString(), *InToTrack->TrackName.ToString());
 									return LoadedRecording->GetLowestSolverFrameNumberAtCycle_AssumesLocked(InToTrack->TrackID, FromSolverFrameData->FrameCycle);
 								}
 							}
@@ -587,8 +588,8 @@ int32 FChaosVDPlaybackController::ConvertCurrentFrameToOtherTrackFrame_AssumesLo
 					}
 				}
 				
-				// From Solver to Game track, we can only convert a frame based on timestapms
-				// TODO: Technically we are ignoring the requested sync mode, but the current implementation of CVD relies in this fallback as we always want to sync from solver to game tracks using the original timestamo mode
+				// From Solver to Game track, we can only convert a frame based on timestamps
+				// TODO: Technically we are ignoring the requested sync mode, but the current implementation of CVD relies in this fallback as we always want to sync from solver to game tracks using the original timestamp mode
 				// We should investigate if it is worth coming up with a better API so this default fallback does not catch anyone using it off guard.
 				return LoadedRecording->GetLowestGameFrameAtSolverFrameNumber_AssumesLocked(InFromTrack->TrackID, InFromTrack->CurrentFrame);
 			}
