Merging fix from Joshua.Davis made in Main for an issue in CVD where QueryData was used instead of SimData, in the details customization used for Collision Goemtry details panels

diff --git a/Engine/Plugins/ChaosVD/Source/ChaosVD/Private/DetailsCustomizations/ChaosVDShapeDataCustomization.cpp b/Engine/Plugins/ChaosVD/Source/ChaosVD/Private/DetailsCustomizations/ChaosVDShapeDataCustomization.cpp
--- a/Engine/Plugins/ChaosVD/Source/ChaosVD/Private/DetailsCustomizations/ChaosVDShapeDataCustomization.cpp
+++ b/Engine/Plugins/ChaosVD/Source/ChaosVD/Private/DetailsCustomizations/ChaosVDShapeDataCustomization.cpp
@@ -340,7 +340,7 @@ void FChaosVDShapeDataCustomization::CacheCurrentCollisionChannelData()
 		FCollisionFilterData QueryFilterData = FChaosVDDetailsCustomizationUtils::ConvertToEngineFilteringData(Data->QueryData);
 		CachedQueryCollisionResponsesFlagsPerChannel = ExtractQueryCollisionResponseContainer(QueryFilterData);
 
-		FCollisionFilterData SimFilterData = FChaosVDDetailsCustomizationUtils::ConvertToEngineFilteringData(Data->QueryData);
+		FCollisionFilterData SimFilterData = FChaosVDDetailsCustomizationUtils::ConvertToEngineFilteringData(Data->SimData);
 		CachedSimCollisionResponsesFlagsPerChannel = ExtractSimCollisionResponseContainer(SimFilterData);
 	}
 }
