[ChaosVD] Fixed an issue where FChaosVDTraceManager::GetTraceFileNameFromStoreForSession returned an invalid file path.

This methods is used just to print the filepath in the status bar, therefore it didn't affect CVD functionality.

The problem was caused by requesting to the Trace store the status structure and the trace info structure, and try to use them afterwards. Under the hood, both structures share the same data buffer and the API returns a ptr to it, not a copy, therefore if you call  FStoreClient::GetStatus and then FStoreClient::GetTraceInfoById, the ptr to a FStoreClient::FStatus returned by the first call, now will have invalid data after the second one.

diff --git a/Engine/Plugins/ChaosVD/Source/ChaosVD/Private/Trace/ChaosVDTraceManager.cpp b/Engine/Plugins/ChaosVD/Source/ChaosVD/Private/Trace/ChaosVDTraceManager.cpp
--- a/Engine/Plugins/ChaosVD/Source/ChaosVD/Private/Trace/ChaosVDTraceManager.cpp
+++ b/Engine/Plugins/ChaosVD/Source/ChaosVD/Private/Trace/ChaosVDTraceManager.cpp
@@ -309,26 +309,36 @@ FString FChaosVDTraceManager::GetTraceFileNameFromStoreForSession(FStringView In
 		return FString();
 	}
 
-	const FStoreClient::FTraceInfo* TraceInfo = StoreClient->GetTraceInfoById(SessionID);
-	if (!TraceInfo)
+	// Note: The following calls to the store client are scoped because FStoreClient::FTraceInfo and FStoreClient::FStatus share the same buffer
+	// under the hood, therefor although the ptr will still be valid, the data will not after we obtain one after the other
+	FString FileName;
 	{
-		return FString();
-	}
+		const FStoreClient::FTraceInfo* TraceInfo = StoreClient->GetTraceInfoById(SessionID);
+		if (!TraceInfo)
+		{
+			return FString();
+		}
 
-	const FStoreClient::FStatus* StoreStatus = StoreClient->GetStatus();
-	if (!StoreStatus)
-	{
-		return FString();
+		const FUtf8StringView Utf8NameView = TraceInfo->GetName();
+		FileName = FString(Utf8NameView);
+		if (!FileName.EndsWith(TEXT(".utrace")))
+		{
+			FileName += TEXT(".utrace");
+		}
 	}
 
-	const FUtf8StringView Utf8NameView = TraceInfo->GetName();
-	FString FileName(Utf8NameView);
-	if (!FileName.EndsWith(TEXT(".utrace")))
+	FString TraceStorePath;
 	{
-		FileName += TEXT(".utrace");
+		const FStoreClient::FStatus* StoreStatus = StoreClient->GetStatus();
+		if (!StoreStatus)
+		{
+			return FString();
+		}
+
+		TraceStorePath = FString(StoreStatus->GetStoreDir());
 	}
 	
-	FString FullFileName = FPaths::Combine(FString(StoreStatus->GetStoreDir()), FileName);
+	FString FullFileName = FPaths::Combine(TraceStorePath, FileName);
 	FPaths::NormalizeFilename(FullFileName);
 
 	return FullFileName;
