[ChaosVD] Removing visibility of the Transport mode override from the settings menu.

This feature was moved to UE 5.8, and should not be used in 5.7.

diff --git a/Engine/Plugins/ChaosVD/Source/ChaosVD/Private/Settings/ChaosVDGeneralSettings.h b/Engine/Plugins/ChaosVD/Source/ChaosVD/Private/Settings/ChaosVDGeneralSettings.h
--- a/Engine/Plugins/ChaosVD/Source/ChaosVD/Private/Settings/ChaosVDGeneralSettings.h
+++ b/Engine/Plugins/ChaosVD/Source/ChaosVD/Private/Settings/ChaosVDGeneralSettings.h
@@ -45,7 +45,7 @@ public:
 	int32 MaxConnectionRetries = 14;
 
 	/** [DEBUG Setting] Data Transport Override - Used to change the transport mode for the trace data for testing. This is not saved by design */
-	UPROPERTY(EditAnywhere, Category = Connection)
+	UPROPERTY()
 	EChaosVDTransportMode DataTransportModeOverride = EChaosVDTransportMode::Invalid;
 
 	/** [DEBUG Setting] If True, any traces done to memory will also be saved to disk at the same time in the Profiling folder. This is not saved by design */
