[ChaosVD] Fix for a random crash caused by an array being accessed from multiple thread without a read lock in some places.

This causes a very random and low repro crash, because it only happens if the array is reallocated in another thread (to grow) while you are using a reference to an item in that array.

#okforguthub public

diff --git a/Engine/Plugins/ChaosVD/Source/ChaosVD/Private/ChaosVDPlaybackController.cpp b/Engine/Plugins/ChaosVD/Source/ChaosVD/Private/ChaosVDPlaybackController.cpp
--- a/Engine/Plugins/ChaosVD/Source/ChaosVD/Private/ChaosVDPlaybackController.cpp
+++ b/Engine/Plugins/ChaosVD/Source/ChaosVD/Private/ChaosVDPlaybackController.cpp
@@ -942,7 +942,7 @@ void FChaosVDPlaybackController::SetFrameRateOverride(float NewFrameRateOverride
 	CurrentFrameRateOverride = (NewFrameRateOverride >= MinimumFrameRateOverride) ? static_cast<int32>(NewFrameRateOverride) : InvalidFrameRateOverride;
 }
 
-float FChaosVDPlaybackController::GetFrameTimeForTrack(EChaosVDTrackType TrackType, int32 TrackID, const TSharedRef<const FChaosVDTrackInfo>& InTrackInfo) const
+float FChaosVDPlaybackController::GetFrameTimeForTrack_AssumesLocked(EChaosVDTrackType TrackType, int32 TrackID,const TSharedRef<const FChaosVDTrackInfo>& InTrackInfo) const
 {
 	const float TargetFrameTimeOverride = GetFrameTimeOverride();
 	const bool bHasValidFrameRateOverride = bUseFrameRateOverride && !FMath::IsNearlyEqual(TargetFrameTimeOverride, InvalidFrameRateOverride);
@@ -957,30 +957,41 @@ float FChaosVDPlaybackController::GetFrameTimeForTrack(EChaosVDTrackType TrackTy
 		switch(TrackType)
 		{
 			case EChaosVDTrackType::Solver:
+			{
+				if (const FChaosVDSolverFrameData* FrameData = LoadedRecording->GetSolverFrameData_AssumesLocked(TrackID, InTrackInfo->CurrentFrame))
 				{
-					if (const FChaosVDSolverFrameData* FrameData = LoadedRecording->GetSolverFrameData_AssumesLocked(TrackID, InTrackInfo->CurrentFrame))
-					{
-						CurrentTargetFrameTime = static_cast<float>(FrameData->GetFrameTime());
-					}
-					break;
+					CurrentTargetFrameTime = static_cast<float>(FrameData->GetFrameTime());
 				}
+				break;
+			}
 			case EChaosVDTrackType::Game:
+			{
+				if (const FChaosVDGameFrameData* FrameData = LoadedRecording->GetGameFrameData_AssumesLocked(InTrackInfo->CurrentFrame))
 				{
-					if (const FChaosVDGameFrameData* FrameData = LoadedRecording->GetGameFrameData_AssumesLocked(InTrackInfo->CurrentFrame))
-					{
-						CurrentTargetFrameTime = static_cast<float>(FrameData->GetFrameTime());
-					}
-
-					break;
+					CurrentTargetFrameTime = static_cast<float>(FrameData->GetFrameTime());
 				}
-			default:
+
 				break;
+			}
+		default:
+			break;
 		}
 	}
 
 	return CurrentTargetFrameTime;
 }
 
+float FChaosVDPlaybackController::GetFrameTimeForTrack(EChaosVDTrackType TrackType, int32 TrackID, const TSharedRef<const FChaosVDTrackInfo>& InTrackInfo) const
+{
+	if (!LoadedRecording)
+	{
+		return FallbackFrameTime;
+	}
+
+	FReadScopeLock ReadLock(LoadedRecording->GetRecordingDataLock());
+	return GetFrameTimeForTrack_AssumesLocked(TrackType, TrackID, InTrackInfo);
+}
+
 void FChaosVDPlaybackController::UpdateTrackVisibility(EChaosVDTrackType Type, int32 TrackID, bool bNewVisibility)
 {
 	switch (Type)
@@ -1162,7 +1173,7 @@ void FChaosVDPlaybackController::TickPlayback(float DeltaTime)
 	CurrentPlaybackTime -= TargetFrameTime * static_cast<float>(FramesToAdvance);
 
 	const int32 NextFrame = FMath::Clamp(ActiveTrack->CurrentFrame + FramesToAdvance, 0, ActiveTrack->MaxFrames);
-	const int32 NextStageNumber =  bIsGameTrack ? 0 : GetTrackLastStageAtFrame_AssumesLocked(ActiveTrack->TrackType, ActiveTrack->TrackID, NextFrame);
+	const int32 NextStageNumber =  bIsGameTrack ? 0 : GetTrackLastStageAtFrame(ActiveTrack->TrackType, ActiveTrack->TrackID, NextFrame);
 	GoToTrackFrameAndSync(PlaybackSelfInstigatorID, ActiveTrack->TrackType, ActiveTrack->TrackID, NextFrame, NextStageNumber);
 }
 
diff --git a/Engine/Plugins/ChaosVD/Source/ChaosVD/Public/ChaosVDPlaybackController.h b/Engine/Plugins/ChaosVD/Source/ChaosVD/Public/ChaosVDPlaybackController.h
--- a/Engine/Plugins/ChaosVD/Source/ChaosVD/Public/ChaosVDPlaybackController.h
+++ b/Engine/Plugins/ChaosVD/Source/ChaosVD/Public/ChaosVDPlaybackController.h
@@ -329,6 +329,7 @@ public:
 	 * @return 
 	 */
 	float GetFrameTimeForTrack(EChaosVDTrackType TrackType, int32 TrackID, const TSharedRef<const FChaosVDTrackInfo>& InTrackInfo) const;
+	float GetFrameTimeForTrack_AssumesLocked(EChaosVDTrackType TrackType, int32 TrackID, const TSharedRef<const FChaosVDTrackInfo>& InTrackInfo) const;
 
 	/**
 	 * Changes the visibility of a specific track
