[ChaosVD] Simplifying recording workflows and UI
- Now there is a single Start Recording button and a option to automatically start playback, instead of Record Live recording and Record to File. Record live sessions also always wrote the data to file, the only difference (from the user point of view) is that Record Live Session also started playback automatically (under the hood two different systems were used).
- Made the Direct Trace transport the default mode, and removed Trace Server support from the recording controls
- The trace server is still supported as a legacy option accessible from the File menu but might be fully deprecated in the future
- Made all buttons in the main toolbar proper tool menus
- They were already tool menus, but we were not instantiating them as a single menu. Now with this changes features like automatically collapsing buttons outside the screen works properly like in the rest of the editor.
- Improved failure handling (from start a trace or from disconnections) for the Direct Trace and Relay Transport modes
- Updated the code in the Trace Relay data stream to make the trace analysis thread way for an new data event instead of going to sleep for a fixed amount of time and check again for new data.
- Added an option to specify a custom path where the trace files will be stored (now that we don't depend in the Unreal Trace Server we can control it). Also updated the default path to be a "ChaosVisualDebugger" folder in the Documents folder of the user.

diff --git a/Engine/Plugins/ChaosVD/Source/ChaosVD/Private/ChaosVDCommands.cpp b/Engine/Plugins/ChaosVD/Source/ChaosVD/Private/ChaosVDCommands.cpp
--- a/Engine/Plugins/ChaosVD/Source/ChaosVD/Private/ChaosVDCommands.cpp
+++ b/Engine/Plugins/ChaosVD/Source/ChaosVD/Private/ChaosVDCommands.cpp
@@ -23,6 +23,7 @@ void FChaosVDCommands::RegisterCommands()
 	// Main Toolbar Commands
 	UI_COMMAND(OpenFile, "Open File", "Opens the file browser modal", EUserInterfaceActionType::Button, FInputChord(EModifierKey::Control ,EKeys::O));
 	UI_COMMAND(BrowseLiveSessions, "Browse Live Sessions", "Open the connect to session modal", EUserInterfaceActionType::Button, FInputChord(EModifierKey::Control | EModifierKey::Shift, EKeys::S));
+	UI_COMMAND(DisconnectFromLiveSession, "Disconnect from Live Session", "Disconnects from a legacy live session", EUserInterfaceActionType::Button, FInputChord(EModifierKey::Control | EModifierKey::Shift, EKeys::X));
 	UI_COMMAND(CombineOpenFiles, "Combine Files", "Combines all open sessions into a single file (if possible)", EUserInterfaceActionType::Button, FInputChord(EModifierKey::Control | EModifierKey::Shift , EKeys::C));
 	UI_COMMAND(OpenSceneQueryBrowser, "Scene query Browser", "Opens the Scene Query browser window", EUserInterfaceActionType::Button, FInputChord(EModifierKey::Control , EKeys::Q));
 
diff --git a/Engine/Plugins/ChaosVD/Source/ChaosVD/Private/ChaosVDCommands.h b/Engine/Plugins/ChaosVD/Source/ChaosVD/Private/ChaosVDCommands.h
--- a/Engine/Plugins/ChaosVD/Source/ChaosVD/Private/ChaosVDCommands.h
+++ b/Engine/Plugins/ChaosVD/Source/ChaosVD/Private/ChaosVDCommands.h
@@ -25,6 +25,7 @@ public:
 	// Main Toolbar Commands
 	TSharedPtr<FUICommandInfo> OpenFile;
 	TSharedPtr<FUICommandInfo> BrowseLiveSessions;
+	TSharedPtr<FUICommandInfo> DisconnectFromLiveSession;
 	TSharedPtr<FUICommandInfo> CombineOpenFiles;
 	TSharedPtr<FUICommandInfo> OpenSceneQueryBrowser;
 
diff --git a/Engine/Plugins/ChaosVD/Source/ChaosVD/Private/ChaosVDEngine.cpp b/Engine/Plugins/ChaosVD/Source/ChaosVD/Private/ChaosVDEngine.cpp
--- a/Engine/Plugins/ChaosVD/Source/ChaosVD/Private/ChaosVDEngine.cpp
+++ b/Engine/Plugins/ChaosVD/Source/ChaosVD/Private/ChaosVDEngine.cpp
@@ -26,6 +26,7 @@ void FChaosVDEngine::Initialize()
 	CurrentScene->Initialize();
 
 	PlaybackController = MakeShared<FChaosVDPlaybackController>(CurrentScene);
+	PlaybackController->Initialize();
 
 	if (TSharedPtr<FChaosVDRemoteSessionsManager> RemoteSessionManager = FChaosVDEngineEditorBridge::Get().GetRemoteSessionsManager())
 	{
@@ -72,6 +73,29 @@ void FChaosVDEngine::CloseActiveTraceSessions()
 	}
 }
 
+void FChaosVDEngine::CloseSessionByRemoteSessionID(FGuid RemoteSessionID)
+{
+	FChaosVDTraceSessionDescriptor DescriptorFoundCopy = GetSessionDescriptor(RemoteSessionID);
+
+	if (!DescriptorFoundCopy.IsValid())
+	{
+		return;
+	}
+
+	if (const TSharedPtr<FChaosVDTraceManager> CVDTraceManager = FChaosVDModule::Get().GetTraceManager())
+	{
+		CVDTraceManager->CloseSession(DescriptorFoundCopy.SessionName);
+	}
+
+	OnSessionClosed().Broadcast(DescriptorFoundCopy);
+	CurrentSessionDescriptors.Remove(DescriptorFoundCopy);
+
+	if (PlaybackController && CurrentSessionDescriptors.IsEmpty())
+	{
+		PlaybackController->UnloadCurrentRecording(EChaosVDUnloadRecordingFlags::BroadcastChanges);
+	}
+}
+
 void FChaosVDEngine::StopActiveTraceSessions()
 {
 	if (const TSharedPtr<FChaosVDTraceManager> CVDTraceManager = FChaosVDModule::Get().GetTraceManager())
@@ -92,6 +116,7 @@ void FChaosVDEngine::DeInitialize()
 
 	CurrentScene->DeInitialize();
 	CurrentScene.Reset();
+	PlaybackController->DeInitialize();
 	PlaybackController.Reset();
 
 	CloseActiveTraceSessions();
@@ -196,6 +221,17 @@ bool FChaosVDEngine::HasAnyLiveSessionActive() const
 	return PlaybackController && PlaybackController->IsPlayingLiveSession();
 }
 
+const FChaosVDTraceSessionDescriptor& FChaosVDEngine::GetSessionDescriptor(FGuid RemoteSessionID) const
+{
+	static FChaosVDTraceSessionDescriptor InvalidSessionDescriptor;
+	const FChaosVDTraceSessionDescriptor* DescriptorFound = CurrentSessionDescriptors.FindByPredicate([RemoteSessionID](const FChaosVDTraceSessionDescriptor& InDescriptor)
+	{
+		return RemoteSessionID == InDescriptor.RemoteSessionID;
+	});
+
+	return DescriptorFound ? *DescriptorFound : InvalidSessionDescriptor;
+}
+
 void FChaosVDEngine::LoadRecording_Internal(const TFunction<FString(const TSharedPtr<FChaosVDRecording>&)>& LoadCallback, EChaosVDLoadRecordedDataMode LoadingMode)
 {
 	FChaosVDTraceSessionDescriptor NewSessionFromFileDescriptor;
@@ -309,7 +345,14 @@ bool FChaosVDEngine::ConnectToLiveSession(uint32 SessionID, const FString& InSes
 	return bSuccess;
 }
 
+
 bool FChaosVDEngine::ConnectToLiveSession_Direct(EChaosVDLoadRecordedDataMode LoadingMode)
+{
+	UE_LOG(LogChaosVDEditor, Error, TEXT("[%hs] This menthod is no longer supported. Please use the overload that takes a Remote Session ID."), __func__);
+	return false;
+}
+
+bool FChaosVDEngine::ConnectToLiveSession_Direct(FGuid RemoteSessionID, EChaosVDLoadRecordedDataMode LoadingMode)
 {
 	if (LoadingMode == EChaosVDLoadRecordedDataMode::SingleSource)
 	{
@@ -318,6 +361,7 @@ bool FChaosVDEngine::ConnectToLiveSession_Direct(EChaosVDLoadRecordedDataMode Lo
 
 	FChaosVDTraceSessionDescriptor NewSessionFromFileDescriptor;
 
+	NewSessionFromFileDescriptor.RemoteSessionID = RemoteSessionID;
 	NewSessionFromFileDescriptor.bIsLiveSession = true;
 	
 	TSharedPtr<FChaosVDRecording> ExistingRecordingInstance = nullptr;
diff --git a/Engine/Plugins/ChaosVD/Source/ChaosVD/Private/ChaosVDPlaybackController.cpp b/Engine/Plugins/ChaosVD/Source/ChaosVD/Private/ChaosVDPlaybackController.cpp
--- a/Engine/Plugins/ChaosVD/Source/ChaosVD/Private/ChaosVDPlaybackController.cpp
+++ b/Engine/Plugins/ChaosVD/Source/ChaosVD/Private/ChaosVDPlaybackController.cpp
@@ -8,8 +8,10 @@
 #include "ChaosVDRuntimeModule.h"
 #include "ChaosVDScene.h"
 #include "ChaosVDSceneCompositionReport.h"
+#include "ChaosVDSettingsManager.h"
 #include "ExtensionsSystem/ChaosVDExtensionsManager.h"
 #include "Misc/MessageDialog.h"
+#include "Settings/ChaosVDGeneralSettings.h"
 #include "Trace/ChaosVDTraceManager.h"
 #include "Trace/ChaosVDTraceProvider.h"
 #include "TraceServices/Model/AnalysisSession.h"
@@ -19,6 +21,8 @@
 
 #define LOCTEXT_NAMESPACE "ChaosVisualDebugger"
 
+class UChaosVDGeneralSettings;
+
 namespace Chaos::VisualDebugger::Cvars
 {
 	static bool bPlayAllPreviousFrameSteps = true;
@@ -46,6 +50,40 @@ FChaosVDPlaybackController::FChaosVDPlaybackController(const TWeakPtr<FChaosVDSc
 FChaosVDPlaybackController::~FChaosVDPlaybackController()
 {
 	UnloadCurrentRecording(EChaosVDUnloadRecordingFlags::Silent);
+	DeInitialize();
+}
+
+void FChaosVDPlaybackController::Initialize()
+{
+	if (bInitialized)
+	{
+		UE_LOG(LogChaosVDEditor, Warning, TEXT("[%hs] Playback controller already initialized"), __func__);
+		return;
+	}
+
+	if (UChaosVDGeneralSettings* Settings = FChaosVDSettingsManager::Get().GetSettingsObject<UChaosVDGeneralSettings>())
+	{
+		HandleGeneralSettingsUpdated(Settings);
+		Settings->OnSettingsChanged().AddSP(this, &FChaosVDPlaybackController::HandleGeneralSettingsUpdated);
+	}
+
+	bInitialized = true;
+}
+
+void FChaosVDPlaybackController::DeInitialize()
+{
+	if (!bInitialized)
+	{
+		UE_LOG(LogChaosVDEditor, Verbose, TEXT("[%hs] Playback controller not initialized"), __func__);
+		return;
+	}
+
+	if (UChaosVDGeneralSettings* Settings = FChaosVDSettingsManager::Get().GetSettingsObject<UChaosVDGeneralSettings>())
+	{
+		Settings->OnSettingsChanged().RemoveAll(this);
+	}
+	
+	bInitialized = false;
 }
 
 bool FChaosVDPlaybackController::LoadChaosVDRecordingFromTraceSession(const FChaosVDTraceSessionDescriptor& InSessionDescriptor)
@@ -210,6 +248,14 @@ TSharedRef<FChaosVDTrackInfo> FChaosVDPlaybackController::CreateTrackInfo(int32
 	return NewTrack;
 }
 
+void FChaosVDPlaybackController::HandleGeneralSettingsUpdated(UObject* SettingsObject)
+{
+	if (UChaosVDGeneralSettings* GeneralSettings = Cast<UChaosVDGeneralSettings>(SettingsObject))
+	{
+		bStartLiveSessionPlaybackAutomatically = GeneralSettings->bAutoStartLiveSessionsPlayback;
+	}
+}
+
 void FChaosVDPlaybackController::GoToRecordedSolverStage_AssumesLocked(const int32 InTrackID, const int32 FrameNumber, const int32 StageNumber, FGuid InstigatorID)
 {
 	if (const TSharedPtr<FChaosVDScene> SceneToControlSharedPtr = SceneToControl.Pin())
@@ -546,7 +592,7 @@ int32 FChaosVDPlaybackController::ConvertCurrentFrameToOtherTrackFrame_AssumesLo
 
 											// Server tracks will not have a frame offset set, so we need to clamp it to 0 first
 											const int32 ToTrackNetworkTick = InToTrack->bIsServer && ToSolverFrameData->NetworkTickOffset == INDEX_NONE ? 0 : ToSolverFrameData->NetworkTickOffset;
-											const int32 FromTrackNetworkTick = InFromTrack->bIsServer && ToSolverFrameData->NetworkTickOffset == INDEX_NONE ? 0 : FromSolverFrameData->NetworkTickOffset;
+											const int32 FromTrackNetworkTick = InFromTrack->bIsServer && FromSolverFrameData->NetworkTickOffset == INDEX_NONE ? 0 : FromSolverFrameData->NetworkTickOffset;
 
 											// This works for Client --> Server conversion because in that case we want to add the frame offset.
 											// As the tick offset in server tracks is 0,  the following calculation will return a negative offset, which it will result in the intended addition in the last caluclation
@@ -849,7 +895,8 @@ bool FChaosVDPlaybackController::Tick(float DeltaTime)
 
 			if (const bool bHasSolverTrackData = TracksByType && !TracksByType->IsEmpty())
 			{
-				if (IsPlayingLiveSession())
+				// If this a live session and automatic playback is enabled, start playing
+				if (IsPlayingLiveSession() && bStartLiveSessionPlaybackAutomatically)
 				{
 					HandleFramePlaybackControlInput(EChaosVDPlaybackButtonsID::Play, GameTrackInfo.ToSharedRef(), PlaybackSelfInstigatorID);
 					bPlayedFirstFrame = true;
@@ -865,8 +912,8 @@ bool FChaosVDPlaybackController::Tick(float DeltaTime)
 			}
 		}
 
-		// If we are live, make sure we don't lag too much behind
-		if (!bPauseRequested && IsPlayingLiveSession())
+		// If we are live and playing, make sure we don't lag too much behind
+		if (!bPauseRequested && IsPlayingLiveSession() && ActiveTrack->bIsPlaying)
 		{
 			if (const TSharedPtr<const FChaosVDTrackInfo> GameTrackInfo = GetTrackInfo(EChaosVDTrackType::Game, GameTrackID))
 			{
diff --git a/Engine/Plugins/ChaosVD/Source/ChaosVD/Private/Settings/ChaosVDGeneralSettings.h b/Engine/Plugins/ChaosVD/Source/ChaosVD/Private/Settings/ChaosVDGeneralSettings.h
--- a/Engine/Plugins/ChaosVD/Source/ChaosVD/Private/Settings/ChaosVDGeneralSettings.h
+++ b/Engine/Plugins/ChaosVD/Source/ChaosVD/Private/Settings/ChaosVDGeneralSettings.h
@@ -46,9 +46,17 @@ public:
 
 	/** [DEBUG Setting] Data Transport Override - Used to change the transport mode for the trace data for testing. This is not saved by design */
 	UPROPERTY(EditAnywhere, Category = Connection)
-	EChaosVDTransportMode DataTransportModeOverride = EChaosVDTransportMode::Invalid;
+	EChaosVDTransportMode DataTransportModeOverride = EChaosVDTransportMode::Direct;
 
-	/** [DEBUG Setting] If True, any traces done to memory will also be saved to disk at the same time in the Profiling folder. This is not saved by design */
+	/** If True, any traces done to memory will also be saved to disk at the same time in the user's document folder (or a user specified folder). This is not saved by design */
 	UPROPERTY(EditAnywhere, Category = Connection)
-	bool bSaveMemoryTracesToDisk = true;
+	bool bSaveRecordingsToDisk = true;
+
+	/** Custom path where recordings will be saved from. By default, recordings are saved in a Chaos Visual Debugger folder in the documents folder */
+	UPROPERTY(Config, EditAnywhere, Category = Connection)
+	FFilePath DefaultSavePathOverride;
+
+	/** If True, playback will start automatically after a recording started, and we successfully connect to it */
+	UPROPERTY(Config, EditAnywhere, Category = Connection)
+	bool bAutoStartLiveSessionsPlayback = true;
 };
diff --git a/Engine/Plugins/ChaosVD/Source/ChaosVD/Private/Trace/ChaosVDTraceAnalyzer.cpp b/Engine/Plugins/ChaosVD/Source/ChaosVD/Private/Trace/ChaosVDTraceAnalyzer.cpp
--- a/Engine/Plugins/ChaosVD/Source/ChaosVD/Private/Trace/ChaosVDTraceAnalyzer.cpp
+++ b/Engine/Plugins/ChaosVD/Source/ChaosVD/Private/Trace/ChaosVDTraceAnalyzer.cpp
@@ -328,24 +328,9 @@ bool FChaosVDTraceAnalyzer::OnEvent(uint16 RouteId, EStyle Style, const FOnEvent
 		}	
 	case RouteId_ChaosVDUsesAutoRTFM:
 		{
-			const bool bUsingAutoRTFM = EventData.GetValue<bool>("bUsingAutoRTFM");
-
-			if (bUsingAutoRTFM)
+			if (EventData.GetValue<bool>("bUsingAutoRTFM"))
 			{
-				FText AutoRTFMWarning = NSLOCTEXT("ChaosVisualDebugger","AutoRTFMWarningMessage", "This recording was made with AutoRTFM enabled. \n\nAutoRTFM is not fully supported and framing/timing of recorded data during transactions might be off.");
-
-				if (EnumHasAnyFlags(ChaosVDTraceProvider->GetRecordingForSession()->GetAttributes(), EChaosVDRecordingAttributes::Merged))
-				{
-					UE_LOG(LogChaosVDEditor, Warning, TEXT("[%s]"), *AutoRTFMWarning.ToString())
-				}
-				else
-				{
-					FTSTicker::GetCoreTicker().AddTicker(FTickerDelegate::CreateLambda([AutoRTFMWarning](float DeltaTime)
-					{
-						FMessageDialog::Open(EAppMsgCategory::Warning, EAppMsgType::Ok, AutoRTFMWarning, NSLOCTEXT("ChaosVisualDebugger", "AutoRTFMWarningMessageTitle", "Partially unsupported CVD Recording"));
-						return false;
-					}));
-				}
+				UE_LOG(LogChaosVDEditor, Log, TEXT("Processing recording was made with AutoRTFM enabled. "));
 			}
 
 			break;
diff --git a/Engine/Plugins/ChaosVD/Source/ChaosVD/Private/Trace/ChaosVDTraceManager.cpp b/Engine/Plugins/ChaosVD/Source/ChaosVD/Private/Trace/ChaosVDTraceManager.cpp
--- a/Engine/Plugins/ChaosVD/Source/ChaosVD/Private/Trace/ChaosVDTraceManager.cpp
+++ b/Engine/Plugins/ChaosVD/Source/ChaosVD/Private/Trace/ChaosVDTraceManager.cpp
@@ -23,15 +23,6 @@
 
 namespace Chaos::VD
 {
-	namespace CVars
-	{
-		static float SecondsToSleepWaitingForRelayData = 0.1f;
-		static FAutoConsoleVariableRef CVarUseCVDDynamicMeshGenerator(
-			TEXT("p.Chaos.VD.Tool.SecondsToSleepWaitingForRelayData"),
-			SecondsToSleepWaitingForRelayData,
-			TEXT("Time in seconds the trace analysis thread should wait to receive new data from a relay connection"));
-	}
-
 	FAsyncProgressNotification::~FAsyncProgressNotification()
 	{
 		if (NotificationHandle.IsValid())
@@ -96,7 +87,25 @@ namespace Chaos::VD
 	}
 
 	FRelayDataStream::FRelayDataStream(FGuid InRemoteSessionID)
-		: CancelEvent(FPlatformProcess::GetSynchEventFromPool(false)), RemoteSessionID(InRemoteSessionID)
+		: CancelEvent(FPlatformProcess::GetSynchEventFromPool(false)),
+		WaitForDataEvent(FPlatformProcess::GetSynchEventFromPool(false)),
+		RemoteSessionID(InRemoteSessionID)
+	{
+#if WITH_CHAOS_VISUAL_DEBUGGER
+		TSharedPtr<ITraceDataRelayTransport> RelayTransportInstance = FChaosVDEngineEditorBridge::Get().GetRelayTransportInstance();
+		if (ensure(RelayTransportInstance))
+		{
+			RelayTransportInstance->RegisterRelayDataReceiverForSessionID(RemoteSessionID, ITraceDataRelayTransport::FProcessReceivedRelayDataDelegate::CreateRaw(this, &FRelayDataStream::EnqueueRelayedData));
+		}
+#endif
+	}
+
+
+	FRelayDataStream::FRelayDataStream(FGuid InRemoteSessionID, TUniquePtr<FArchiveFileWriterGeneric>&& InFileWriter)
+		: CancelEvent(FPlatformProcess::GetSynchEventFromPool(false)),
+		WaitForDataEvent(FPlatformProcess::GetSynchEventFromPool(false)),
+		RemoteSessionID(InRemoteSessionID),
+		FileWriter(MoveTemp(InFileWriter))
 	{
 #if WITH_CHAOS_VISUAL_DEBUGGER
 		TSharedPtr<ITraceDataRelayTransport> RelayTransportInstance = FChaosVDEngineEditorBridge::Get().GetRelayTransportInstance();
@@ -110,6 +119,7 @@ namespace Chaos::VD
 	FRelayDataStream::~FRelayDataStream()
 	{
 		FPlatformProcess::ReturnSynchEventToPool(CancelEvent);
+		FPlatformProcess::ReturnSynchEventToPool(WaitForDataEvent);
 
 #if WITH_CHAOS_VISUAL_DEBUGGER
 		TSharedPtr<ITraceDataRelayTransport> RelayTransportInstance = FChaosVDEngineEditorBridge::Get().GetRelayTransportInstance();
@@ -122,12 +132,30 @@ namespace Chaos::VD
 
 	void FRelayDataStream::EnqueueRelayedData(const TConstArrayView<uint8> InTraceDataBuffer)
 	{
+		// An empty data array signals an end of file event from the relay instance
+		if (InTraceDataBuffer.IsEmpty())
+		{
+			Close();
+			return;
+		}
+
+		WaitForDataEvent->Trigger();
+	
 		FSerializedDataBuffer DataBuffer(InTraceDataBuffer);
 		DataQueue.Enqueue(MoveTemp(DataBuffer));
 	}
 
 	void FRelayDataStream::EnqueueRelayedData(TArray<uint8>&& InTraceDataBuffer)
 	{
+		// An empty data array signals an end of file event from the relay instance
+		if (InTraceDataBuffer.IsEmpty())
+		{
+			Close();
+			return;
+		}
+
+		WaitForDataEvent->Trigger();
+	
 		FSerializedDataBuffer DataBuffer(MoveTemp(InTraceDataBuffer));
 		DataQueue.Enqueue(MoveTemp(DataBuffer));
 	}
@@ -193,13 +221,20 @@ namespace Chaos::VD
 			}
 
 			// If at least some bytes have been written return those. Otherwise, 
-			// we'll sleep and wait for some data to arrive.
+			// we'll wait for some data to arrive.
 			if (BytesCopied > 0)
 			{
+				if (FileWriter.IsValid() && !FileWriter->IsError())
+				{
+					FileWriter->Serialize(Dest, BytesCopied);
+				}	
 				break;
 			}
-
-			FPlatformProcess::Sleep(CVars::SecondsToSleepWaitingForRelayData);
+			else
+			{
+				WaitForDataEvent->Reset();
+				WaitForDataEvent->Wait();
+			}
 		}
 
 		return BytesCopied;
@@ -212,6 +247,17 @@ namespace Chaos::VD
 			CancelEvent->Trigger();
 		}
 
+		if (WaitForDataEvent)
+		{
+			WaitForDataEvent->Trigger();
+		}
+
+		if (FileWriter.IsValid())
+		{
+			FileWriter->Close();
+			FileWriter.Reset();
+		}
+
 		bClosed = true;
 	}
 
@@ -221,6 +267,8 @@ namespace Chaos::VD
 	}
 }
 
+uint16 FChaosVDTraceManager::CurrentDirectModeStartingPort = DefaultStartingPort;
+
 FChaosVDTraceManager::FChaosVDTraceManager() 
 {
 	ChaosVDTraceModule = MakeShared<FChaosVDTraceModule>();
@@ -347,6 +395,10 @@ UE::Trace::FStoreClient::FTraceData FChaosVDTraceManager::GetTraceDataStreamFrom
 	return StoreClient->ReadTrace(SessionID);
 }
 
+FString FChaosVDTraceManager::GetDefaultSaveDirectoryPath()
+{
+	return FPaths::Combine(FPlatformProcess::UserDir(), TEXT("ChaosVisualDebugger"));
+}
 
 TUniquePtr<UE::Trace::IInDataStream> FChaosVDTraceManager::CreateFileDataStream(TUniquePtr<IFileHandle>&& InFileHandle)
 {
@@ -416,7 +468,22 @@ TUniquePtr<UE::Trace::IInDataStream> FChaosVDTraceManager::CreateFileDataStream(
 
 TUniquePtr<Chaos::VD::FRelayDataStream> FChaosVDTraceManager::CreateRelayDataStream(FGuid RemoteSessionID)
 {
-	TUniquePtr<Chaos::VD::FRelayDataStream> NewDataStream = MakeUnique<Chaos::VD::FRelayDataStream>(RemoteSessionID);
+	TUniquePtr<FArchiveFileWriterGeneric> TraceFileWriter = nullptr;
+
+	if (UChaosVDGeneralSettings* Settings = FChaosVDSettingsManager::Get().GetSettingsObject<UChaosVDGeneralSettings>())
+	{
+		if (Settings->bSaveRecordingsToDisk)
+		{
+			TraceFileWriter = OpenTraceFileForWrite(Settings->DefaultSavePathOverride.FilePath.IsEmpty() ? GetDefaultSaveDirectoryPath() : Settings->DefaultSavePathOverride.FilePath);
+			if (!TraceFileWriter)
+			{
+				UE_LOG(LogChaosVDEditor, Error, TEXT("[%hs] Failed to open file for relay trace stream"), __func__);
+				return nullptr;
+			}
+		}
+	}
+
+	TUniquePtr<Chaos::VD::FRelayDataStream> NewDataStream = MakeUnique<Chaos::VD::FRelayDataStream>(RemoteSessionID, MoveTemp(TraceFileWriter));
 	
 	return MoveTemp(NewDataStream);
 }
@@ -430,6 +497,11 @@ void FChaosVDTraceManager::SetPendingExternalRecordingToProcess(const TSharedPtr
 
 bool FChaosVDTraceManager::ConnectToLiveSession_Internal(uint32 SessionID, const TSharedPtr<FChaosVDRecording>& ExistingRecordingPtr, const FString& InRequestedSessionName, UE::Trace::FStoreClient::FTraceData&& InTraceDataStream)
 {
+	if (!InTraceDataStream.IsValid())
+	{
+		return false;
+	}
+
 	ITraceServicesModule& TraceServicesModule = FModuleManager::LoadModuleChecked<ITraceServicesModule>("TraceServices");
 	if (const TSharedPtr<TraceServices::IAnalysisService> TraceAnalysisService = TraceServicesModule.GetAnalysisService())
 	{
@@ -476,21 +548,36 @@ FString FChaosVDTraceManager::ConnectToLiveSession(FStringView InSessionHost, ui
 	return FString();
 }
 
-TUniquePtr<FArchiveFileWriterGeneric> FChaosVDTraceManager::OpenTraceFileForWrite()
+TUniquePtr<FArchiveFileWriterGeneric> FChaosVDTraceManager::OpenTraceFileForWrite(const FString& InDirectoryToSavePath)
 {
+	IPlatformFile& PlatformFile = FPlatformFileManager::Get().GetPlatformFile();
+
+	if (InDirectoryToSavePath.IsEmpty())
+	{
+		UE_LOG(LogChaosVDEditor, Error, TEXT("[%hs] Failed open file for write | No save directory was provided"), __func__);
+		return nullptr;
+	}
+
+	PlatformFile.CreateDirectory(*InDirectoryToSavePath);
+	if (!PlatformFile.DirectoryExists(*InDirectoryToSavePath))
+	{
+		UE_LOG(LogChaosVDEditor, Error, TEXT("[%hs] Failed open file for write | Failed to create or access provided save directory [%s]"), __func__, *InDirectoryToSavePath);
+		return nullptr;
+	}
+
 	TUniquePtr<FArchiveFileWriterGeneric> FileWriter = nullptr;
 
 #if WITH_CHAOS_VISUAL_DEBUGGER
 
-	IPlatformFile& PlatformFile = FPlatformFileManager::Get().GetPlatformFile();
 	FChaosVDRuntimeModule& CVDRuntimeModule = FChaosVDRuntimeModule::Get();
+	
 	constexpr int32 MaxAttempts = 10;
 	int32 CurrentAttempts = 0;
 	while (CurrentAttempts < MaxAttempts)
 	{
 		FString OutFileName;
 		CVDRuntimeModule.GenerateRecordingFileName(OutFileName);
-		OutFileName = FPaths::Combine(FPaths::ProfilingDir(), OutFileName);
+		OutFileName = FPaths::Combine(InDirectoryToSavePath, OutFileName);
 
 		if (IFileHandle* FileHandle = PlatformFile.OpenWrite(*OutFileName))
 		{
@@ -514,9 +601,9 @@ FString FChaosVDTraceManager::ConnectToLiveSession_Direct(uint16& OutSessionPort
 
 	if (UChaosVDGeneralSettings* Settings = FChaosVDSettingsManager::Get().GetSettingsObject<UChaosVDGeneralSettings>())
 	{
-		if (Settings->bSaveMemoryTracesToDisk)
+		if (Settings->bSaveRecordingsToDisk)
 		{
-			TraceFileWriter = OpenTraceFileForWrite();
+			TraceFileWriter = OpenTraceFileForWrite(Settings->DefaultSavePathOverride.FilePath.IsEmpty() ? GetDefaultSaveDirectoryPath() : Settings->DefaultSavePathOverride.FilePath);
 			if (!TraceFileWriter)
 			{
 				UE_LOG(LogChaosVDEditor, Error, TEXT("[%hs] Failed to open open file for direct trace stream"), __func__);
@@ -526,7 +613,7 @@ FString FChaosVDTraceManager::ConnectToLiveSession_Direct(uint16& OutSessionPort
 	}
 
 	TUniquePtr<UE::Trace::FDirectSocketStream> DirectSocketStream = MakeUnique<UE::Trace::FDirectSocketStream>(MoveTemp(TraceFileWriter));
-	OutSessionPort = DirectSocketStream->StartListening();
+	OutSessionPort = DirectSocketStream->StartListening(CurrentDirectModeStartingPort);
 
 	if (OutSessionPort == 0)
 	{
@@ -534,6 +621,11 @@ FString FChaosVDTraceManager::ConnectToLiveSession_Direct(uint16& OutSessionPort
 		return FString();
 	}
 
+	// If the port was in use, FDirectSocketStream will try multiple ports in sequence, therefore the next port we should try while we have recording active (common case in Multi-session mode)
+	// is the next one from the last success
+	CurrentDirectModeStartingPort = OutSessionPort + 1;
+
+
 	FString DirectSessionName = FString::Printf(TEXT("DirectSession[127.0.0.1:%u]"), OutSessionPort);
 	
 	if (ConnectToLiveSession_Internal(~0u, ExistingRecordingPtr, DirectSessionName, MoveTemp(DirectSocketStream)))
@@ -604,6 +696,13 @@ void FChaosVDTraceManager::CloseSession(const FString& InSessionName)
 		}
 
 		AnalysisSessionByName.Remove(InSessionName);
+
+		if (AnalysisSessionByName.IsEmpty())
+		{
+			// If we don't have any more session left, we can reset back to the original starting port
+			// As all sockets we had in use are closed at this point
+			CurrentDirectModeStartingPort = DefaultStartingPort;
+		}
 	}
 }
 
diff --git a/Engine/Plugins/ChaosVD/Source/ChaosVD/Private/Widgets/ChaosVDRecordingControls.cpp b/Engine/Plugins/ChaosVD/Source/ChaosVD/Private/Widgets/ChaosVDRecordingControls.cpp
new file mode 100644
--- /dev/null
+++ b/Engine/Plugins/ChaosVD/Source/ChaosVD/Private/Widgets/ChaosVDRecordingControls.cpp
@@ -0,0 +1,1338 @@
+// Copyright Epic Games, Inc. All Rights Reserved.
+
+#include "Widgets/ChaosVDRecordingControls.h"
+
+#include "AsyncCompilationHelpers.h"
+#include "ChaosVDEngine.h"
+#include "ChaosVDModule.h"
+#include "ChaosVDRecordingDetails.h"
+#include "ChaosVDScene.h"
+#include "ChaosVDSettingsManager.h"
+#include "ChaosVDStyle.h"
+#include "ChaosVisualDebugger/ChaosVisualDebuggerTrace.h"
+#include "Editor.h"
+#include "SChaosVDNameListPicker.h"
+#include "SEnumCombo.h"
+#include "SocketSubsystem.h"
+#include "Input/Reply.h"
+#include "Misc/MessageDialog.h"
+#include "StatusBarSubsystem.h"
+#include "ToolMenu.h"
+#include "ToolMenus.h"
+#include "Chaos/ChaosVDEngineEditorBridge.h"
+#include "Chaos/ChaosVDTraceRelayTransport.h"
+#include "Framework/MultiBox/MultiBoxBuilder.h"
+#include "Misc/ScopedSlowTask.h"
+#include "Settings/ChaosVDGeneralSettings.h"
+#include "Widgets/Images/SImage.h"
+#include "Widgets/Input/SButton.h"
+#include "Widgets/SChaosVDMainTab.h"
+#include "Widgets/Input/SComboButton.h"
+#include "Widgets/Notifications/SNotificationList.h"
+#include "Widgets/Text/STextBlock.h"
+
+#include UE_INLINE_GENERATED_CPP_BY_NAME(ChaosVDRecordingControls)
+
+#define LOCTEXT_NAMESPACE "ChaosVisualDebugger"
+
+class UChaosVDGeneralSettings;
+
+FChaosVDRecordingControls::FChaosVDRecordingControls(const TSharedRef<SChaosVDMainTab>& InMainTabSharedRef)
+{
+	MainTabWeakPtr = InMainTabSharedRef;
+	StatusBarID = InMainTabSharedRef->GetStatusBarName();
+	
+	RecordingAnimation = FCurveSequence();
+	RecordingAnimation.AddCurve(0.f, 1.5f, ECurveEaseFunction::Linear);
+
+	CurrentSelectedSessionId = FChaosVDRemoteSessionsManager::LocalEditorSessionID;
+}
+
+FChaosVDRecordingControls::~FChaosVDRecordingControls()
+{
+	Deinitialize();
+}
+
+void FChaosVDRecordingControls::Initialize()
+{
+	if (bInitialized)
+	{
+		return;
+	}
+
+	if (TSharedPtr<FChaosVDRemoteSessionsManager> RemoteSessionManager = FChaosVDEngineEditorBridge::Get().GetRemoteSessionsManager())
+	{
+		RemoteSessionManager->OnSessionRecordingStarted().AddSP(this, &FChaosVDRecordingControls::HandleRecordingStart);
+		RemoteSessionManager->OnSessionRecordingStopped().AddSP(this, &FChaosVDRecordingControls::HandleRecordingStop);
+		RemoteSessionManager->OnSessionExpired().AddSP(this, &FChaosVDRecordingControls::CancelInFlightConnectionAttempt);
+		RemoteSessionManager->OnSessionsUpdated().AddSP(this, &FChaosVDRecordingControls::HandleRemoteSessionsUpdated);
+	}
+
+	bInitialized = true;
+}
+
+void FChaosVDRecordingControls::Deinitialize()
+{
+	if (!bInitialized)
+	{
+		return;
+	}
+
+	if (TSharedPtr<FChaosVDRemoteSessionsManager> RemoteSessionManager = FChaosVDEngineEditorBridge::Get().GetRemoteSessionsManager())
+	{
+		RemoteSessionManager->OnSessionRecordingStarted().RemoveAll(this);
+		RemoteSessionManager->OnSessionRecordingStopped().RemoveAll(this);
+		RemoteSessionManager->OnSessionExpired().RemoveAll(this);
+		RemoteSessionManager->OnSessionsUpdated().RemoveAll(this);
+	}
+
+	bInitialized = false;
+}
+
+TSharedRef<SWidget> FChaosVDRecordingControls::GenerateToggleRecordingStateButton(const FText& StartRecordingTooltip)
+{
+	return SNew(SButton)
+			.OnClicked(FOnClicked::CreateRaw(this, &FChaosVDRecordingControls::ToggleRecordingState))
+			.ForegroundColor(FSlateColor::UseForeground())
+			.IsFocusable(false)
+			.IsEnabled_Raw(this, &FChaosVDRecordingControls::IsRecordingToggleButtonEnabled)
+			.VAlign(VAlign_Center)
+			.HAlign(HAlign_Center)
+			.OnHovered_Lambda([this](){ bRecordingButtonHovered = true;})
+			.OnUnhovered_Lambda([this](){ bRecordingButtonHovered = false;})
+			.ToolTipText_Lambda([this, StartRecordingTooltip]()
+			{
+				return IsRecording() ? LOCTEXT("StopRecordButtonDesc", "Stop the current recording ") : StartRecordingTooltip;
+			})
+			[
+				SNew(SHorizontalBox)
+				+SHorizontalBox::Slot()
+				.Padding(FMargin(0, 0, 0, 0))
+				.VAlign(VAlign_Center)
+				.AutoWidth()
+				[
+					SNew(SImage)
+					.Image_Raw(this, &FChaosVDRecordingControls::GetRecordOrStopButton)
+					.ColorAndOpacity_Lambda([this]()
+					{
+						if (IsRecording())
+						{
+							TSharedPtr<SChaosVDMainTab> MainTabPtr = MainTabWeakPtr.Pin();
+							if (MainTabPtr && !RecordingAnimation.IsPlaying())
+							{
+								RecordingAnimation.Play(MainTabPtr->AsShared(), true);
+							}
+
+							const FLinearColor Color = bRecordingButtonHovered ? FLinearColor::Red : FLinearColor::White;
+							return FSlateColor(bRecordingButtonHovered ? Color : Color.CopyWithNewOpacity(0.2f + 0.8f * RecordingAnimation.GetLerp()));
+						}
+
+						RecordingAnimation.Pause();
+						return FSlateColor::UseSubduedForeground();
+					})
+				]
+				+SHorizontalBox::Slot()
+				.Padding(FMargin(4, 0, 0, 0))
+				.VAlign(VAlign_Center)
+				.AutoWidth()
+				[
+					SNew(STextBlock)
+					.Visibility_Lambda([this](){return IsRecording() ? EVisibility::Collapsed : EVisibility::Visible;})
+					.TextStyle(FAppStyle::Get(), "SmallButtonText")
+					.Text(LOCTEXT("StartRecordingButtonLabel", "Start Recording"))
+				]
+			];
+}
+
+FText FChaosVDRecordingControls::GetCurrentSelectedSessionName() const
+{
+	if (TSharedPtr<FChaosVDSessionInfo> CurrentSessionPtr = GetCurrentSessionInfo())
+	{
+		return FText::AsCultureInvariant(CurrentSessionPtr->SessionName);
+	}
+
+	static FText InvalidSessionName = LOCTEXT("InvalidSessionLabel", "Select a Session");
+
+	return InvalidSessionName;
+}
+
+TSharedRef<SWidget> FChaosVDRecordingControls::GenerateTargetSessionSelector()
+{
+	return SNew(SHorizontalBox)
+		+SHorizontalBox::Slot()
+		.AutoWidth()
+		[
+			SNew(SComboButton)
+			.MenuPlacement(MenuPlacement_AboveAnchor).ComboButtonStyle(&FAppStyle::Get().GetWidgetStyle<FComboButtonStyle>("ComboButton"))
+			.OnGetMenuContent(this, &FChaosVDRecordingControls::GenerateTargetSessionDropdown)
+			.HasDownArrow(true)
+			.ButtonContent()
+			[
+				SNew(STextBlock)
+				.Text_Raw(this, &FChaosVDRecordingControls::GetCurrentSelectedSessionName)
+				.TextStyle(&FAppStyle::Get().GetWidgetStyle<FTextBlockStyle>("NormalText"))
+		]];
+}
+
+TSharedRef<SWidget> FChaosVDRecordingControls::GenerateTargetSessionDropdown()
+{
+	using namespace Chaos::VisualDebugger;
+
+	FMenuBuilder MenuBuilder(true, nullptr);
+
+	MenuBuilder.BeginSection("CVDRecordingWidgetTargets", LOCTEXT("CVDRecordingTargetsMenu", "Available Targets"));
+	{
+		if (TSharedPtr<FChaosVDRemoteSessionsManager> RemoteSessionManager = FChaosVDEngineEditorBridge::Get().GetRemoteSessionsManager())
+		{
+			RemoteSessionManager->EnumerateActiveSessions([this, &MenuBuilder](const TSharedRef<FChaosVDSessionInfo>& InSessionInfoRef)
+			{
+				if (EnumHasAnyFlags(InSessionInfoRef->GetSessionTypeAttributes(), EChaosVDRemoteSessionAttributes::IsMultiSessionWrapper))
+				{
+					return true;
+				}
+	
+				FText SessionNameAsText = FText::AsCultureInvariant(InSessionInfoRef->SessionName);
+				MenuBuilder.AddMenuEntry(
+				SessionNameAsText,
+				FText::Format(LOCTEXT("SingleTargetItemTooltip", "Select {0} session as current target"), SessionNameAsText),
+				GetIconForSession(InSessionInfoRef->InstanceId),
+				FUIAction(FExecuteAction::CreateSP(this, &FChaosVDRecordingControls::SelectTargetSession, InSessionInfoRef->InstanceId), EUIActionRepeatMode::RepeatDisabled));
+				return true;
+			});
+		}
+	}
+	
+	MenuBuilder.EndSection();
+
+	MenuBuilder.BeginSection("CVDRecordingWidgetTargetsMulti", LOCTEXT("CVDRecordingMultiTargetsMenu", "Multi Target"));
+
+	FText AllRemoteTargetsLabel = LOCTEXT("AllRemoteOption", "All Remote");
+	MenuBuilder.AddMenuEntry(
+			AllRemoteTargetsLabel,
+			LOCTEXT("MultiRemoteTargetTooltip", "Select this to act on all remote targets"),
+			GetIconForSession(FChaosVDRemoteSessionsManager::AllRemoteSessionsWrapperGUID),
+			FUIAction(FExecuteAction::CreateSP(this, &FChaosVDRecordingControls::SelectTargetSession, FChaosVDRemoteSessionsManager::AllRemoteSessionsWrapperGUID),
+			FCanExecuteAction::CreateSP(this, &FChaosVDRecordingControls::CanSelectMultiSessionTarget, FChaosVDRemoteSessionsManager::AllRemoteSessionsWrapperGUID),
+			EUIActionRepeatMode::RepeatDisabled));
+
+	FText AllRemoteServersTargetsLabel = LOCTEXT("AllRemoteServersOption", "All Remote Servers");
+	MenuBuilder.AddMenuEntry(
+			AllRemoteServersTargetsLabel,
+			LOCTEXT("MultiRemoteServerTargetTooltip", "Select this to act on all remote server targets"),
+			GetIconForSession(FChaosVDRemoteSessionsManager::AllRemoteServersWrapperGUID),
+			FUIAction(FExecuteAction::CreateSP(this, &FChaosVDRecordingControls::SelectTargetSession, FChaosVDRemoteSessionsManager::AllRemoteServersWrapperGUID),
+			FCanExecuteAction::CreateSP(this, &FChaosVDRecordingControls::CanSelectMultiSessionTarget, FChaosVDRemoteSessionsManager::AllRemoteServersWrapperGUID),
+			EUIActionRepeatMode::RepeatDisabled));
+
+	FText AllRemoteClientsTargetsLabel = LOCTEXT("AllRemoteClientsOption", "All Remote Clients");
+	MenuBuilder.AddMenuEntry(
+			AllRemoteClientsTargetsLabel,
+			LOCTEXT("MultiRemoteClientTargetTooltip", "Select this to act on all remote client targets"),
+			GetIconForSession(FChaosVDRemoteSessionsManager::AllRemoteClientsWrapperGUID),
+			FUIAction(FExecuteAction::CreateSP(this, &FChaosVDRecordingControls::SelectTargetSession, FChaosVDRemoteSessionsManager::AllRemoteClientsWrapperGUID),
+			FCanExecuteAction::CreateSP(this, &FChaosVDRecordingControls::CanSelectMultiSessionTarget, FChaosVDRemoteSessionsManager::AllRemoteClientsWrapperGUID),
+			EUIActionRepeatMode::RepeatDisabled));
+
+	FText AllTargets = LOCTEXT("AllTargetsOption", "All");
+	MenuBuilder.AddMenuEntry(
+		AllTargets,
+		LOCTEXT("MultiAllTargetTooltip", "Select this to act on all targets, both Local and Remote"),
+		GetIconForSession(FChaosVDRemoteSessionsManager::AllSessionsWrapperGUID),
+		FUIAction(FExecuteAction::CreateSP(this, &FChaosVDRecordingControls::SelectTargetSession, FChaosVDRemoteSessionsManager::AllSessionsWrapperGUID), EUIActionRepeatMode::RepeatDisabled));
+
+	MenuBuilder.AddMenuSeparator();
+
+	FText CustomTargets = LOCTEXT("CustomTargetsOption", "Custom Selection");
+	MenuBuilder.AddSubMenu(CustomTargets,
+		LOCTEXT("MultiCustomTargetTooltip", "Select this to act on the specific targets you selected"),
+		FNewMenuDelegate::CreateSP(this, &FChaosVDRecordingControls::GenerateCustomTargetsMenu)
+		);
+
+	MenuBuilder.EndSection();
+
+	return MenuBuilder.MakeWidget();
+}
+
+TSharedRef<SWidget> FChaosVDRecordingControls::GenerateRecordingTimeTextBlock()
+{
+	return SNew(SBox)
+			.VAlign(VAlign_Center)
+			.Visibility_Raw(this, &FChaosVDRecordingControls::GetRecordingTimeTextBlockVisibility)
+			.Padding(12.0f, 0.0f, 0.0f, 0.0f)
+			[
+				SNew(STextBlock)
+				.TextStyle(FAppStyle::Get(), "SmallButtonText")
+				.Text_Raw(this, &FChaosVDRecordingControls::GetRecordingTimeText)
+				.ColorAndOpacity(FColor::White)
+			];
+}
+
+
+TSharedRef<SWidget> FChaosVDRecordingControls::GenerateAdvancedRecordOptionsButton()
+{
+	return SNew(SComboButton)
+		.ContentPadding(FMargin(6.0f, 0.0f))
+		.IsEnabled_Raw(this, &FChaosVDRecordingControls::IsRecordingToggleButtonEnabled) // This button needs to follow the enabled condition of the recording button
+		.MenuPlacement(MenuPlacement_AboveAnchor).ComboButtonStyle(&FAppStyle::Get()
+		.GetWidgetStyle<FComboButtonStyle>("ComboButton"))
+		.OnGetMenuContent(this, &FChaosVDRecordingControls::GenerateAdvancedRecordOptionsMenu)
+		.HasDownArrow(true);
+}
+
+TSharedRef<SWidget> FChaosVDRecordingControls::GenerateAdvancedRecordOptionsMenu()
+{
+	using namespace Chaos::VisualDebugger;
+
+	FMenuBuilder MenuBuilder(true, nullptr);
+
+	constexpr bool bSearchable = true;
+	constexpr bool bNoIndent = true;
+
+	MenuBuilder.BeginSection("CVDRecordingWidget", LOCTEXT("CVDAdvancedRecordingMenu", "Advanced"));
+
+	{
+		FUIAction AutoStartPlaybackUIAction;
+		AutoStartPlaybackUIAction.ExecuteAction = FExecuteAction::CreateSP(this, &FChaosVDRecordingControls::ToggleAutoStartPlayback);
+		AutoStartPlaybackUIAction.GetActionCheckState = FGetActionCheckState::CreateStatic(&FUIAction::IsActionCheckedPassthrough, FIsActionChecked::CreateSP(this, &FChaosVDRecordingControls::IsAutoStartPlaybackEnabled));
+	
+		MenuBuilder.AddMenuEntry(
+			LOCTEXT("AutoStartPlaybackLabel", "Auto Start Playback"),
+			LOCTEXT("AutoStartPlaybackToolTip", "If set to true, playback will start automatically as soon the recording is started and the connection is successfully established"),
+			FSlateIcon(),
+			AutoStartPlaybackUIAction, NAME_None, EUserInterfaceActionType::ToggleButton);
+	}
+
+	{
+		FUIAction SaveToDiskUIAction;
+		SaveToDiskUIAction.ExecuteAction = FExecuteAction::CreateSP(this, &FChaosVDRecordingControls::ToggleSaveTracesToDisk);
+		SaveToDiskUIAction.GetActionCheckState = FGetActionCheckState::CreateStatic(&FUIAction::IsActionCheckedPassthrough, FIsActionChecked::CreateSP(this, &FChaosVDRecordingControls::IsSaveTracesToDiskEnabled));
+
+		MenuBuilder.AddMenuEntry(
+		LOCTEXT("AutoSaveToDiskLabel", "Save Recorded Data to Disk"),
+		LOCTEXT("AutoStartPlaybackToolTip", "If set to true, any recorded data for the session will be saved automatically as soon it is recorded."),
+		FSlateIcon(),
+		SaveToDiskUIAction, NAME_None, EUserInterfaceActionType::ToggleButton);
+	}
+
+	MenuBuilder.AddSeparator();
+	
+	MenuBuilder.AddWidget(
+			GenerateLoadingModeSelector(),
+			LOCTEXT("LoadingModeLabel", "Loading Mode"),
+			bSearchable,
+			bNoIndent,
+			LOCTEXT("LoadingModeLabelTooltip", "Alters if the newly recorded data should unload previous loaded data (Single Source), or just add to it (Multi Source)"));
+
+	MenuBuilder.EndSection();
+
+	return MenuBuilder.MakeWidget();
+}
+
+EVisibility FChaosVDRecordingControls::GetRecordingTimeTextBlockVisibility() const
+{
+	TSharedPtr<FChaosVDSessionInfo> SessionInfo = GetCurrentSessionInfo();
+	bool bIsVisible = SessionInfo && !EnumHasAnyFlags(SessionInfo->GetSessionTypeAttributes(), EChaosVDRemoteSessionAttributes::IsMultiSessionWrapper) && IsRecording();
+
+	return  bIsVisible ? EVisibility::Visible : EVisibility::Collapsed;
+}
+
+TSharedRef<SWidget> FChaosVDRecordingControls::GenerateDataChannelsButton()
+{
+	return SNew(SComboButton)
+			.ContentPadding(FMargin(6.0f, 0.0f))
+			.IsEnabled_Raw(this, &FChaosVDRecordingControls::HasDataChannelsSupport)
+			.MenuPlacement(MenuPlacement_AboveAnchor).ComboButtonStyle(&FAppStyle::Get()
+			.GetWidgetStyle<FComboButtonStyle>("SimpleComboButton"))
+			.OnGetMenuContent(this, &FChaosVDRecordingControls::GenerateDataChannelsMenu)
+			.HasDownArrow(true)
+			.ButtonContent()
+			[
+				SNew(SHorizontalBox)
+				+SHorizontalBox::Slot()
+				.VAlign(VAlign_Center)
+				.AutoWidth()
+				[
+					SNew(STextBlock)
+					.Text(LOCTEXT("DataChannelsButton", "Data Channels"))
+					.TextStyle(&FAppStyle::Get().GetWidgetStyle<FTextBlockStyle>("DialogButtonText"))
+				]
+			];
+}
+
+TSharedRef<SWidget> FChaosVDRecordingControls::GenerateLoadingModeSelector()
+{
+	TAttribute<int32> GetCurrentValueAttribute;
+	GetCurrentValueAttribute.BindSPLambda(this, [this]()
+	{
+		return static_cast<int32>(CurrentLoadingMode);
+	});
+
+	SEnumComboBox::FOnEnumSelectionChanged EnumValueChangedDelegate = SEnumComboBox::FOnEnumSelectionChanged::CreateSPLambda(this, [this](int32 NewValue, ESelectInfo::Type SelectionType)
+	{
+		CurrentLoadingMode = static_cast<EChaosVDLoadRecordedDataMode>(NewValue);
+	});
+	
+	return SNew(SEnumComboBox, StaticEnum<EChaosVDLoadRecordedDataMode>())
+		.IsEnabled_Raw(this, &FChaosVDRecordingControls::CanChangeLoadingMode)
+		.CurrentValue(GetCurrentValueAttribute)
+		.OnEnumSelectionChanged(EnumValueChangedDelegate);
+}
+
+TSharedRef<SWidget> FChaosVDRecordingControls::GenerateDataChannelsMenu()
+{
+	using namespace Chaos::VisualDebugger;
+
+	FMenuBuilder MenuBuilder(true, nullptr);
+
+	MenuBuilder.BeginSection("CVDRecordingWidget", LOCTEXT("CVDRecordingMenuChannels", "Data Channels"));
+	{
+		if (TSharedPtr<FChaosVDSessionInfo> SessionInfo = GetCurrentSessionInfo())
+		{
+			for (const TPair<FString, FChaosVDDataChannelState>& DataChannelStateWithName : SessionInfo->DataChannelsStatesByName)
+			{
+				FText ChannelNamesAsText = FText::AsCultureInvariant(DataChannelStateWithName.Key);
+				MenuBuilder.AddMenuEntry(
+				ChannelNamesAsText,
+				FText::Format(LOCTEXT("ChannelDesc", "Enable/disable the {0} channel"), ChannelNamesAsText),
+				FSlateIcon(),
+				FUIAction(FExecuteAction::CreateSP(this, &FChaosVDRecordingControls::ToggleChannelEnabledState, DataChannelStateWithName.Key),
+					FCanExecuteAction::CreateSP(this, &FChaosVDRecordingControls::CanChangeChannelEnabledState, DataChannelStateWithName.Key), FIsActionChecked::CreateSP(this, &FChaosVDRecordingControls::IsChannelEnabled, DataChannelStateWithName.Key)), NAME_None, EUserInterfaceActionType::ToggleButton);
+			}
+		}
+	}
+
+	MenuBuilder.EndSection();
+
+	return MenuBuilder.MakeWidget();
+}
+
+void FChaosVDRecordingControls::ToggleChannelEnabledState(FString ChannelName)
+{
+	TSharedPtr<FChaosVDRemoteSessionsManager> RemoteSessionManager = FChaosVDEngineEditorBridge::Get().GetRemoteSessionsManager();
+	TSharedPtr<FChaosVDSessionInfo> SessionInfo = GetCurrentSessionInfo();
+	if (RemoteSessionManager && SessionInfo)
+	{
+		if (FChaosVDDataChannelState* ChannelState = SessionInfo->DataChannelsStatesByName.Find(ChannelName))
+		{
+			ChannelState->bWaitingUpdatedState = true;
+			RemoteSessionManager->SendDataChannelStateChangeCommand(SessionInfo->Address, {ChannelState->ChannelName, !ChannelState->bIsEnabled });
+		}
+	}
+}
+
+bool FChaosVDRecordingControls::IsChannelEnabled(FString ChannelName)
+{
+	if (TSharedPtr<FChaosVDSessionInfo> SessionInfo = GetCurrentSessionInfo())
+	{
+		if (FChaosVDDataChannelState* ChannelState = SessionInfo->DataChannelsStatesByName.Find(ChannelName))
+		{
+			return ChannelState->bIsEnabled;
+		}
+	}
+
+	return false;
+}
+
+bool FChaosVDRecordingControls::CanChangeChannelEnabledState(FString ChannelName)
+{
+	if (TSharedPtr<FChaosVDSessionInfo> SessionInfo = GetCurrentSessionInfo())
+	{
+		if (FChaosVDDataChannelState* ChannelState = SessionInfo->DataChannelsStatesByName.Find(ChannelName))
+		{
+			return ChannelState->bCanChangeChannelState && !ChannelState->bWaitingUpdatedState;
+		}
+	}
+	
+	return false;
+}
+
+void FChaosVDRecordingControls::SelectTargetSession(FGuid SessionId)
+{
+	CurrentSelectedSessionId = SessionId;
+}
+
+void FChaosVDRecordingControls::GenerateCustomTargetsMenu(FMenuBuilder& MenuBuilder)
+{
+	if (TSharedPtr<FChaosVDRemoteSessionsManager> RemoteSessionManager = FChaosVDEngineEditorBridge::Get().GetRemoteSessionsManager())
+	{
+		RemoteSessionManager->EnumerateActiveSessions([this, &MenuBuilder](const TSharedRef<FChaosVDSessionInfo>& InSessionInfoRef)
+		{
+			if (EnumHasAnyFlags(InSessionInfoRef->GetSessionTypeAttributes(), EChaosVDRemoteSessionAttributes::IsMultiSessionWrapper))
+			{
+				return true;
+			}
+	
+			FText SessionNameAsText = FText::AsCultureInvariant(InSessionInfoRef->SessionName);
+			MenuBuilder.AddMenuEntry(
+			SessionNameAsText,
+			FText::Format(LOCTEXT("MultiTargetItemTooltip", "Select {0} session as one of the current targets"), SessionNameAsText),
+			GetIconForSession(InSessionInfoRef->InstanceId),
+			FUIAction(
+			FExecuteAction::CreateSP(this, &FChaosVDRecordingControls::ToggleSessionSelectionInCustomTarget, InSessionInfoRef->InstanceId),
+			FCanExecuteAction::CreateSP(this, &FChaosVDRecordingControls::CanSelectInCustomTarget, InSessionInfoRef->InstanceId),
+			FIsActionChecked::CreateSP(this, &FChaosVDRecordingControls::IsSessionPartOfCustomTargetSelection, InSessionInfoRef->InstanceId)),
+			NAME_None, EUserInterfaceActionType::ToggleButton);
+
+			return true;
+		});
+	}
+}
+
+bool FChaosVDRecordingControls::IsSessionPartOfCustomTargetSelection(FGuid SessionGuid)
+{
+	if (TSharedPtr<FChaosVDRemoteSessionsManager> RemoteSessionManager = FChaosVDEngineEditorBridge::Get().GetRemoteSessionsManager())
+	{
+		if (TSharedPtr<FChaosVDMultiSessionInfo> CustomSessionTarget = StaticCastSharedPtr<FChaosVDMultiSessionInfo>(RemoteSessionManager->GetSessionInfo(FChaosVDRemoteSessionsManager::CustomSessionsWrapperGUID).Pin()))
+		{
+			return CustomSessionTarget->InnerSessionsByInstanceID.Find(SessionGuid) != nullptr;
+		}
+	}
+
+	return false;
+}
+
+void FChaosVDRecordingControls::ToggleSessionSelectionInCustomTarget(FGuid SessionGuid)
+{
+	if (TSharedPtr<FChaosVDRemoteSessionsManager> RemoteSessionManager = FChaosVDEngineEditorBridge::Get().GetRemoteSessionsManager())
+	{
+		if (TSharedPtr<FChaosVDMultiSessionInfo> CustomSessionTarget = StaticCastSharedPtr<FChaosVDMultiSessionInfo>(RemoteSessionManager->GetSessionInfo(FChaosVDRemoteSessionsManager::CustomSessionsWrapperGUID).Pin()))
+		{
+			if (CustomSessionTarget->InnerSessionsByInstanceID.Find(SessionGuid) != nullptr)
+			{
+				CustomSessionTarget->InnerSessionsByInstanceID.Remove(SessionGuid);
+			}
+			else
+			{
+				CustomSessionTarget->InnerSessionsByInstanceID.Add(SessionGuid, RemoteSessionManager->GetSessionInfo(SessionGuid));
+			}
+
+			SelectTargetSession(FChaosVDRemoteSessionsManager::CustomSessionsWrapperGUID);
+		}
+	}
+}
+
+bool FChaosVDRecordingControls::CanSelectInCustomTarget(FGuid SessionGuid) const
+{
+	if (TSharedPtr<FChaosVDRemoteSessionsManager> RemoteSessionManager = FChaosVDEngineEditorBridge::Get().GetRemoteSessionsManager())
+	{
+		if (TSharedPtr<FChaosVDSessionInfo> CustomSessionTarget = RemoteSessionManager->GetSessionInfo(SessionGuid).Pin())
+		{
+			return CustomSessionTarget->ReadyState == EChaosVDRemoteSessionReadyState::Ready;
+		}
+	}
+
+	return false;
+}
+
+bool FChaosVDRecordingControls::CanSelectMultiSessionTarget(FGuid SessionGuid) const
+{
+	if (TSharedPtr<FChaosVDRemoteSessionsManager> RemoteSessionManager = FChaosVDEngineEditorBridge::Get().GetRemoteSessionsManager())
+	{
+		if (TSharedPtr<FChaosVDSessionInfo> CustomSessionTarget = RemoteSessionManager->GetSessionInfo(SessionGuid).Pin())
+		{
+			return CanSelectMultiSessionTarget(CustomSessionTarget.ToSharedRef());
+		}
+	}
+
+	return false;
+}
+
+bool FChaosVDRecordingControls::CanSelectMultiSessionTarget(const TSharedRef<FChaosVDSessionInfo>& SessionInfoRef) const
+{
+	if (EnumHasAnyFlags(SessionInfoRef->GetSessionTypeAttributes(), EChaosVDRemoteSessionAttributes::IsMultiSessionWrapper))
+	{
+		TSharedRef<FChaosVDMultiSessionInfo> AsMultiSessionInfo = StaticCastSharedRef<FChaosVDMultiSessionInfo>(SessionInfoRef);
+		return !AsMultiSessionInfo->InnerSessionsByInstanceID.IsEmpty();
+	}
+
+	return false;
+}
+
+void FChaosVDRecordingControls::ToggleAutoStartPlayback()
+{
+	if (UChaosVDGeneralSettings* Settings = FChaosVDSettingsManager::Get().GetSettingsObject<UChaosVDGeneralSettings>())
+	{
+		FProperty* AutoStartProperty = FindFieldChecked<FProperty>(UChaosVDGeneralSettings::StaticClass(), "bAutoStartLiveSessionsPlayback");
+		Settings->PreEditChange(AutoStartProperty);
+		Settings->bAutoStartLiveSessionsPlayback = !Settings->bAutoStartLiveSessionsPlayback;
+		Settings->PostEditChange();
+	}
+}
+
+bool FChaosVDRecordingControls::IsAutoStartPlaybackEnabled() const
+{
+	UChaosVDGeneralSettings* Settings = FChaosVDSettingsManager::Get().GetSettingsObject<UChaosVDGeneralSettings>();
+	return Settings ? Settings->bAutoStartLiveSessionsPlayback : false;
+}
+
+void FChaosVDRecordingControls::ToggleSaveTracesToDisk()
+{
+	if (UChaosVDGeneralSettings* Settings = FChaosVDSettingsManager::Get().GetSettingsObject<UChaosVDGeneralSettings>())
+	{
+		FProperty* SaveToDiskProperty = FindFieldChecked<FProperty>(UChaosVDGeneralSettings::StaticClass(), "bSaveRecordingsToDisk");
+		Settings->PreEditChange(SaveToDiskProperty);
+		Settings->bSaveRecordingsToDisk = !Settings->bSaveRecordingsToDisk;
+		Settings->PostEditChange();
+	}
+}
+
+bool FChaosVDRecordingControls::IsSaveTracesToDiskEnabled()
+{
+	UChaosVDGeneralSettings* Settings = FChaosVDSettingsManager::Get().GetSettingsObject<UChaosVDGeneralSettings>();
+	return Settings ? Settings->bSaveRecordingsToDisk : false;
+}
+
+FSlateIcon FChaosVDRecordingControls::GetIconForSession(FGuid SessionId)
+{
+	if (TSharedPtr<FChaosVDSessionInfo> SessionInfo = GetSessionInfo(SessionId))
+	{
+		return SessionInfo->IsRecording() ? FSlateIcon(FChaosVDStyle::GetStyleSetName(),FName("RecordIcon"), FName("RecordIcon")) : FSlateIcon();
+	}
+
+	return FSlateIcon();
+}
+
+TSharedPtr<FChaosVDSessionInfo> FChaosVDRecordingControls::GetCurrentSessionInfo() const
+{
+	return GetSessionInfo(CurrentSelectedSessionId);
+}
+
+TSharedPtr<FChaosVDSessionInfo> FChaosVDRecordingControls::GetSessionInfo(FGuid Id) const
+{
+	TSharedPtr<FChaosVDRemoteSessionsManager> RemoteSessionManager = FChaosVDEngineEditorBridge::Get().GetRemoteSessionsManager();
+	TSharedPtr<FChaosVDSessionInfo> SessionInfo = RemoteSessionManager ? RemoteSessionManager->GetSessionInfo(Id).Pin() : nullptr;
+
+	return SessionInfo;
+}
+
+bool FChaosVDRecordingControls::HasDataChannelsSupport() const
+{
+	if (TSharedPtr<FChaosVDSessionInfo> SessionInfo = GetCurrentSessionInfo())
+	{
+		return !SessionInfo->DataChannelsStatesByName.IsEmpty();
+	}
+
+	return false;
+}
+
+bool FChaosVDRecordingControls::CanChangeLoadingMode() const
+{
+	if (TSharedPtr<FChaosVDSessionInfo> CurrentSession = GetCurrentSessionInfo())
+	{
+		// In multi session mode targets, the loading mode is controlled automatically
+		if (EnumHasAnyFlags(CurrentSession->GetSessionTypeAttributes(), EChaosVDRemoteSessionAttributes::IsMultiSessionWrapper))
+		{
+			return false;
+		}
+		else if (const TSharedPtr<SChaosVDMainTab> MainTabSharedPtr = MainTabWeakPtr.Pin())
+		{
+			// If nothing is loaded yet, it does not make sense change the loading mode
+			return !MainTabSharedPtr->GetChaosVDEngineInstance()->GetCurrentSessionDescriptors().IsEmpty();
+		}
+	}
+
+	return false;
+}
+
+const FSlateBrush* FChaosVDRecordingControls::GetRecordOrStopButton() const
+{
+	const FSlateBrush* RecordIconBrush = FChaosVDStyle::Get().GetBrush("RecordIcon");
+	return bRecordingButtonHovered && IsRecording() ? FChaosVDStyle::Get().GetBrush("StopIcon") : RecordIconBrush;
+}
+
+void FChaosVDRecordingControls::HandleRecordingStop(TWeakPtr<FChaosVDSessionInfo> SessionInfo)
+{
+	if (UStatusBarSubsystem* StatusBarSubsystem = GEditor ? GEditor->GetEditorSubsystem<UStatusBarSubsystem>() : nullptr)
+	{
+		StatusBarSubsystem->PopStatusBarMessage(StatusBarID, RecordingMessageHandle);
+
+		const FText LiveSessionEnded = LOCTEXT("RecordingSessionEndedMessage"," Recording session has ended");
+		LiveSessionEndedMessageHandle = StatusBarSubsystem->PushStatusBarMessage(StatusBarID, LiveSessionEnded);
+	}
+}
+
+void FChaosVDRecordingControls::HandleRecordingStart(TWeakPtr<FChaosVDSessionInfo> SessionInfo)
+{
+	UStatusBarSubsystem* StatusBarSubsystem = GEditor ? GEditor->GetEditorSubsystem<UStatusBarSubsystem>() : nullptr;
+	if (!StatusBarSubsystem)
+	{
+		return;
+	}
+	
+	if (RecordingPathMessageHandle.IsValid())
+	{
+		StatusBarSubsystem->PopStatusBarMessage(StatusBarID, RecordingPathMessageHandle);
+		RecordingPathMessageHandle = FStatusBarMessageHandle();
+	}
+	
+	if (LiveSessionEndedMessageHandle.IsValid())
+	{
+		StatusBarSubsystem->PopStatusBarMessage(StatusBarID, LiveSessionEndedMessageHandle);
+		LiveSessionEndedMessageHandle = FStatusBarMessageHandle();
+	}
+
+	RecordingMessageHandle = StatusBarSubsystem->PushStatusBarMessage(StatusBarID, LOCTEXT("RecordingMessage", "Recording..."));
+}
+
+void FChaosVDRecordingControls::ExecuteAsyncConnectionAttemptTaskWithRetry(FGuid SessionID, int32 RemainingRetries, const TFunction<bool()>& InRecordingStartAttemptCallback, const TFunction<void()>& InRecordingFailedCallback)
+{
+	TSharedPtr<FAsyncConnectionAttemptTask> ConnectionAttemptRequest = MakeShared<FAsyncConnectionAttemptTask>(SessionID, InRecordingStartAttemptCallback, InRecordingFailedCallback, StaticCastWeakPtr<FChaosVDRecordingControls>(AsWeak()));
+	ConnectionAttemptRequest->Start(RemainingRetries, IntervalBetweenAutoplayConnectionAttemptsSeconds);
+}
+
+void FChaosVDRecordingControls::FAsyncConnectionAttemptTask::HandleSuccess()
+{
+	State = EState::Success;
+
+	if (TSharedPtr<FChaosVDRecordingControls> OwnerPtr = Owner.Pin())
+	{
+		OwnerPtr->HandleConnectionAttemptResult(SessionID, EChaosVDLiveConnectionAttemptResult::Success, ProgressNotification);
+	}
+}
+
+void FChaosVDRecordingControls::FAsyncConnectionAttemptTask::HandleFailure()
+{
+	State = EState::Failed;
+
+	if (RecordingFailedCallback)
+	{
+		RecordingFailedCallback();
+	}
+
+	if (TSharedPtr<FChaosVDRecordingControls> OwnerPtr = Owner.Pin())
+	{
+		OwnerPtr->HandleConnectionAttemptResult(SessionID, EChaosVDLiveConnectionAttemptResult::Failed, ProgressNotification);
+	}
+}
+
+bool FChaosVDRecordingControls::FAsyncConnectionAttemptTask::CanExecute() const
+{
+	// The only two states where Execute() can be run are Not Started and InProgress
+	// In Progress is allowed because we might need to execute again as part of the retry process
+	return State == EState::NotStarted || State == EState::InProgress;
+}
+
+void FChaosVDRecordingControls::CancelInFlightConnectionAttempt(FGuid SessionID)
+{
+	if (TSharedPtr<FAsyncConnectionAttemptTask>* ExistingRequestInFlightPtrPtr = InFlightAsyncConnectionRequests.Find(SessionID))
+	{
+		if (TSharedPtr<FAsyncConnectionAttemptTask> ExistingRequestInFlightPtr = *ExistingRequestInFlightPtrPtr)
+		{
+			ExistingRequestInFlightPtr->Cancel();
+		}
+		else
+		{
+			UE_LOG(LogChaosVDEditor, Verbose, TEXT("[%hs] Attempted to cancel a connection of a Session that is no linger valid | Session ID [%s] ."), __func__, *SessionID.ToString());
+		}
+
+		return;
+	}
+
+	UE_LOG(LogChaosVDEditor, Verbose, TEXT("[%hs] Session ID [%s] not found."), __func__, *SessionID.ToString());
+}
+
+void FChaosVDRecordingControls::HandleRemoteSessionsUpdated()
+{
+	TSharedPtr<FChaosVDSessionInfo> CurrentSessionPtr = GetCurrentSessionInfo();
+	if (!CurrentSessionPtr)
+	{
+		// if the session that is currently selected is not valid, try to find a valid session and select that as default target session
+		if (TSharedPtr<FChaosVDRemoteSessionsManager> RemoteSessionManager = FChaosVDEngineEditorBridge::Get().GetRemoteSessionsManager())
+		{
+			if (TSharedPtr<FChaosVDSessionInfo> FirstValidSession = RemoteSessionManager->GetFirstAvailableSingleSessionInfo().Pin())
+			{
+				SelectTargetSession(FirstValidSession->InstanceId);
+				return;
+			}
+		}
+
+		SelectTargetSession(FChaosVDRemoteSessionsManager::InvalidSessionGUID);
+	}
+}
+
+FChaosVDRecordingControls::FAsyncConnectionAttemptTask::~FAsyncConnectionAttemptTask()
+{
+	if (RetryDelegateHandle.IsValid())
+	{
+		FTSTicker::RemoveTicker(RetryDelegateHandle);
+		RetryDelegateHandle = FTSTicker::FDelegateHandle();
+	}
+}
+
+void FChaosVDRecordingControls::FAsyncConnectionAttemptTask::Cancel()
+{
+	State = EState::Canceled;
+
+	if (RecordingFailedCallback)
+	{
+		RecordingFailedCallback();
+	}
+
+	if (TSharedPtr<FChaosVDRecordingControls> OwnerPtr = Owner.Pin())
+	{
+		OwnerPtr->HandleConnectionAttemptResult(SessionID, EChaosVDLiveConnectionAttemptResult::Canceled, ProgressNotification);
+	}
+}
+
+void FChaosVDRecordingControls::FAsyncConnectionAttemptTask::ScheduleRetry()
+{
+	if (!CanExecute())
+	{
+		UE_LOG(LogChaosVDEditor, Error, TEXT("[%hs] This connection attempt is in an invalid state | Current State Value [%u]"), __func__, State);
+		return;
+	}
+
+	if (RetryDelegateHandle.IsValid())
+	{
+		UE_LOG(LogChaosVDEditor, Warning, TEXT("[%hs] Attempted to schedule a retry when there is one already scheduled"), __func__);
+		return;
+	}
+
+	if (RemainingRetries <= 0)
+	{
+		UE_LOG(LogChaosVDEditor, Warning, TEXT("[%hs] Failed to connect to live session | attempts exhausted..."), __func__);
+		HandleFailure();
+		return;
+	}
+	
+	RemainingRetries--;
+
+	RetryDelegateHandle = FTSTicker::GetCoreTicker().AddTicker(FTickerDelegate::CreateSPLambda(AsShared(), [this](float DeltaTime)
+	{
+		Execute();
+		return false;
+	}), IntervalBetweenAttemptsSeconds);
+}
+
+void FChaosVDRecordingControls::FAsyncConnectionAttemptTask::Start(int32 InMaxRetriesAttempts, float InIntervalBetweenAttemptsSeconds)
+{
+	if (State != EState::NotStarted)
+	{
+		UE_LOG(LogChaosVDEditor, Error, TEXT("[%hs] This connection attempt is already started or in an invalid state | Current State Value [%u]"), __func__, State);
+		return;
+	}
+
+	TSharedPtr<FChaosVDRecordingControls> OwnerPtr = Owner.Pin();
+	if (!OwnerPtr)
+	{
+		UE_LOG(LogChaosVDEditor, Error, TEXT("[%hs] This connection attempt is does not have a valid owner | Cannot start attempt... "), __func__);
+		return;
+	}
+
+	if (TSharedPtr<FAsyncConnectionAttemptTask>* ExistingRequestInFlight = OwnerPtr->InFlightAsyncConnectionRequests.Find(SessionID))
+	{
+		UE_LOG(LogChaosVDEditor, Warning, TEXT("[%hs] There is a connection attempt in fight for this session. attempting to cancel it | Session ID [%s]"), __func__, *SessionID.ToString());
+		(*ExistingRequestInFlight)->Cancel();
+	}
+
+	OwnerPtr->InFlightAsyncConnectionRequests.Add(SessionID, AsShared());
+
+	State = EState::InProgress;
+
+	RemainingRetries = InMaxRetriesAttempts;
+	IntervalBetweenAttemptsSeconds = InIntervalBetweenAttemptsSeconds;
+	
+	if (!ProgressNotification)
+	{
+		FNotificationInfo Info(LOCTEXT("ConnectingToSessionMessage", "Connecting Session ..."));
+		Info.bFireAndForget = false;
+		Info.FadeOutDuration = 3.0f;
+		Info.ExpireDuration = 0.0f;
+		Info.ButtonDetails.Add(FNotificationButtonInfo(LOCTEXT("CancelConnectionAttemptButton", "Cancel"), LOCTEXT("CancelConnectionAttemptButtonTip", "Cancels the active connection attempt."),
+		FSimpleDelegate::CreateSP(AsShared(), &FAsyncConnectionAttemptTask::Cancel)));
+
+		ProgressNotification = OwnerPtr->PushConnectionAttemptNotification(Info);
+	}
+	
+	// We need to wait at least one tick before attempting to connect to give it time to the trace to be initialized, write to disk, and for the
+	// session manager to hear back from a remote instance
+	ScheduleRetry();
+}
+
+void FChaosVDRecordingControls::FAsyncConnectionAttemptTask::Execute()
+{
+	if (RetryDelegateHandle.IsValid())
+	{
+		FTSTicker::RemoveTicker(RetryDelegateHandle);
+		RetryDelegateHandle.Reset();
+	}
+
+	if (!CanExecute())
+	{
+		UE_LOG(LogChaosVDEditor, Error, TEXT("[%hs] This connection attempt is in an invalid state | Current State Value [%u]"), __func__, State);
+		return;
+	}
+
+	if (!RecordingStartAttemptCallback)
+	{
+		UE_LOG(LogChaosVDEditor, Error, TEXT("[%hs] Failed to connect to session | Invalid callback provided..."), __func__);
+		HandleFailure();
+		return;
+	}
+
+	if (TSharedPtr<FChaosVDRecordingControls> OwnerPtr = Owner.Pin())
+	{
+		 OwnerPtr->UpdateConnectionAttemptNotification(ProgressNotification, RemainingRetries);
+	}
+
+	// CVD needs the trace session name to be able to load a live session. Although the session exist, the session name might not be written right away
+	// Trace files don't really have metadata, it is all part of the same stream, so we need to wait until it is written which might take a few ticks.
+	// Therefore, if it is not ready, try again a few times.
+	if (!RecordingStartAttemptCallback())
+	{
+		UE_LOG(LogChaosVDEditor, Verbose, TEXT("[%hs] Failed to connect to live session | Attempting again in [%f]..."), __func__, IntervalBetweenAttemptsSeconds);
+		ScheduleRetry();
+	}
+	else
+	{
+		HandleSuccess();
+	}
+}
+
+FChaosVDRecordingControls::FTraceStoreConnectionSettings FChaosVDRecordingControls::GetTraceStoreConnectionSettings(const TSharedRef<FChaosVDSessionInfo>& InSessionInfoRef)
+{
+	FTraceStoreConnectionSettings ConnectionSettings;
+
+	if (!ensure(!EnumHasAnyFlags(InSessionInfoRef->GetSessionTypeAttributes(), EChaosVDRemoteSessionAttributes::IsMultiSessionWrapper)))
+	{
+		UE_LOG(LogChaosVDEditor, Error, TEXT("[%s] Cannot not be called with a multi session wrapper."), ANSI_TO_TCHAR(__FUNCTION__));
+
+		return ConnectionSettings;
+	}
+
+	const FChaosVDTraceDetails& RecordingSessionDetails = InSessionInfoRef->GetConnectionDetails();
+
+	if (const UE::Trace::FStoreClient::FSessionInfo* TraceSessionInfo = FChaosVDTraceManager::GetTraceSessionInfo(RecordingSessionDetails.TraceTarget, RecordingSessionDetails.TraceGuid))
+	{
+		ConnectionSettings.SessionAddress = RecordingSessionDetails.TraceTarget;
+		ConnectionSettings.TraceID = TraceSessionInfo->GetTraceId();
+		return ConnectionSettings;
+	}
+
+	return ConnectionSettings;
+}
+
+void FChaosVDRecordingControls::ToggleMultiSessionSessionRecordingState(const TSharedRef<FChaosVDMultiSessionInfo>& InSessionInfoRef)
+{
+	const TSharedPtr<SChaosVDMainTab> MainTabSharedPtr = MainTabWeakPtr.Pin();
+	if (!ensure(MainTabSharedPtr))
+	{
+		return;
+	}
+
+	bool bNewRecordingState = !IsRecording();
+
+	if (bNewRecordingState)
+	{
+		CurrentLoadingMode = EChaosVDLoadRecordedDataMode::MultiSource;
+		MainTabSharedPtr->GetChaosVDEngineInstance()->CloseActiveTraceSessions();
+	}
+
+	InSessionInfoRef->EnumerateInnerSessions([this, bNewRecordingState](const TSharedRef<FChaosVDSessionInfo>& InInnerSessionRef)
+	{
+		SetSessionRecordingState(bNewRecordingState, InInnerSessionRef);
+		return true;
+	});
+}
+
+void FChaosVDRecordingControls::ToggleSingleSessionRecordingState(const TSharedRef<FChaosVDSessionInfo>& SessionInfoRef)
+{
+	SetSessionRecordingState(!SessionInfoRef->IsRecording(), SessionInfoRef);
+}
+
+void FChaosVDRecordingControls::SetSessionRecordingState(bool bIsRecording, const TSharedRef<FChaosVDSessionInfo>& SessionInfoRef)
+{
+	TSharedPtr<FChaosVDRemoteSessionsManager> RemoteSessionManager = FChaosVDEngineEditorBridge::Get().GetRemoteSessionsManager();
+	if (!RemoteSessionManager)
+	{
+		UE_LOG(LogChaosVDEditor, Error, TEXT("[%hs] Session Manager is not available"), __func__);
+		return;
+	}
+
+	if (bIsRecording)
+	{
+		const TSharedPtr<SChaosVDMainTab> MainTabSharedPtr = MainTabWeakPtr.Pin();
+		if (!ensure(MainTabSharedPtr))
+		{
+			// This should not be possible, but if we don't have a valid tab instance it will not be possible to properly handle any start recording attempt
+			UE_LOG(LogChaosVDEditor, Error, TEXT("[%hs] Main Tab instance is not available"), __func__);
+			return;
+		}
+
+		TArray<FString, TInlineAllocator<1>> RecordingArgs;
+
+		int32 RemainingRetries = 4;
+		EChaosVDTransportMode DataTransportOverrideMode = EChaosVDTransportMode::Invalid;
+
+		if (UChaosVDGeneralSettings* Settings = FChaosVDSettingsManager::Get().GetSettingsObject<UChaosVDGeneralSettings>())
+		{
+			RemainingRetries = Settings->MaxConnectionRetries;
+			DataTransportOverrideMode = Settings->DataTransportModeOverride;
+		}
+		else
+		{
+			UE_LOG(LogChaosVDEditor, Warning, TEXT("Failed to obtain setting object. Setting the retries attempts to connect to a session to 4 as a fallback."))
+		}
+
+		// Set the session in a busy state so the auto-expire policy for it is relaxed. Starting a recording might cause the remote instance to stall for a few seconds
+		// if the map loaded is too complex
+		SessionInfoRef->ReadyState = EChaosVDRemoteSessionReadyState::Busy;
+
+		auto RecordingAttemptFailedCallback = [WeakMainTabPtr = StaticCastWeakPtr<SChaosVDMainTab>(MainTabSharedPtr->AsWeak()), SessionGUID = SessionInfoRef->InstanceId]()
+		{
+			TSharedPtr<FChaosVDRemoteSessionsManager> RemoteSessionManagerPtr = FChaosVDEngineEditorBridge::Get().GetRemoteSessionsManager();
+			TSharedPtr<FChaosVDSessionInfo> SessionInfoPtr = RemoteSessionManagerPtr ? RemoteSessionManagerPtr->GetSessionInfo(SessionGUID).Pin() : nullptr;
+
+			if (!SessionInfoPtr)
+			{
+				return;
+			}
+
+			SessionInfoPtr->ReadyState = EChaosVDRemoteSessionReadyState::Ready;
+
+			if (TSharedPtr<SChaosVDMainTab> MainTab = WeakMainTabPtr.Pin())
+			{
+				MainTab->GetChaosVDEngineInstance()->CloseSessionByRemoteSessionID(SessionGUID);
+			}
+		};
+
+		FChaosVDStartRecordingCommandMessage RecordingParams;
+		RecordingParams.RecordingMode = EChaosVDRecordingMode::Live;
+		RecordingParams.TransportMode = DataTransportOverrideMode != EChaosVDTransportMode::Invalid ? DataTransportOverrideMode : EChaosVDTransportMode::TraceServer;
+
+		check(GLog);
+		bool bOutCanBindAll = false;
+		//TODO: Add a way to specify a local address in case we have multiple adapters?
+		TSharedRef<FInternetAddr> LocalIP = ISocketSubsystem::Get(PLATFORM_SOCKETSUBSYSTEM)->GetLocalHostAddr(*GLog, bOutCanBindAll);
+
+		constexpr bool bAppendPort = false;
+		RecordingParams.Target = LocalIP->ToString(bAppendPort);
+
+#if WITH_CHAOS_VISUAL_DEBUGGER
+		Chaos::VisualDebugger::FChaosVDDataChannelsManager::Get().EnumerateChannels([&RecordingParams](const TSharedRef<Chaos::VisualDebugger::FChaosVDOptionalDataChannel>& Channel)
+		{
+			if (Channel->IsChannelEnabled())
+			{
+				RecordingParams.DataChannelsEnabledOverrideList.Add(Channel->GetId().ToString());
+			}
+			return true;
+		});
+#endif
+
+		if (RecordingParams.TransportMode == EChaosVDTransportMode::Direct)
+		{
+			// Direct Tracing mode requires opening a trace session in the editor side first before starting the recording
+			// so targets have something to connect to (Usually targets connect to a running trace server)
+			if (!MainTabSharedPtr->ConnectToLiveSession_Direct(SessionInfoRef->InstanceId, CurrentLoadingMode))
+			{
+				UE_LOG(LogChaosVDEditor, Error, TEXT("[%hs] Failed to connect to live session using direct trace mode"), __func__);
+				RecordingAttemptFailedCallback();
+				return;
+			}
+
+			// If we get here it means we have at least one valid session
+			const FChaosVDTraceSessionDescriptor& SessionDesc = MainTabSharedPtr->GetChaosVDEngineInstance()->GetCurrentSessionDescriptors().Last();
+
+			RecordingParams.Target.Appendf(TEXT(":%u"), SessionDesc.SessionPort);
+
+			RemoteSessionManager->SendStartRecordingCommand(SessionInfoRef->Address, RecordingParams);
+
+			ExecuteAsyncConnectionAttemptTaskWithRetry(SessionInfoRef->InstanceId, RemainingRetries, [WeakThis = AsWeak(), SessionInstanceId = SessionInfoRef->InstanceId]()
+			{
+				TSharedPtr<FChaosVDRemoteSessionsManager> RemoteSessionManagerPtr = FChaosVDEngineEditorBridge::Get().GetRemoteSessionsManager();
+				TSharedPtr<FChaosVDSessionInfo> SessionInfoPtr = RemoteSessionManagerPtr ? RemoteSessionManagerPtr->GetSessionInfo(SessionInstanceId).Pin() : nullptr;
+				if (!SessionInfoPtr)
+				{
+					return false;
+				}
+
+				// In Direct trace mode we if we get connection details from the remote instance, we can assume the connection succeeded.
+				return SessionInfoPtr->GetConnectionDetails().IsValid();
+	
+			}, RecordingAttemptFailedCallback);
+		}
+		else if (RecordingParams.TransportMode == EChaosVDTransportMode::Relay)
+		{
+			// Relay Tracing mode requires opening a trace session in the editor side first before starting the recording
+			// so targets have something to connect to (Usually targets connect to a running trace server)
+			if (!MainTabSharedPtr->ConnectToLiveSession_Relay(SessionInfoRef->InstanceId, CurrentLoadingMode))
+			{
+				UE_LOG(LogChaosVDEditor, Error, TEXT("[%hs] Failed to connect to live session using relay trace mode"), __func__);
+				RecordingAttemptFailedCallback();
+				return;
+			}
+
+			RemoteSessionManager->SendStartRecordingCommand(SessionInfoRef->Address, RecordingParams);
+
+#if WITH_CHAOS_VISUAL_DEBUGGER
+			// Once the start recording command is issued, we can try to connect to the created session (if everything went well).
+			// If it didn't go well, this will take care of update the UI to notify the user
+			ExecuteAsyncConnectionAttemptTaskWithRetry(SessionInfoRef->InstanceId, RemainingRetries, [WeakThis = AsWeak(), SessionInstanceId = SessionInfoRef->InstanceId]()
+			{
+				TSharedPtr<FChaosVDRemoteSessionsManager> RemoteSessionManagerPtr = FChaosVDEngineEditorBridge::Get().GetRemoteSessionsManager();
+				TSharedPtr<FChaosVDSessionInfo> SessionInfoPtr = RemoteSessionManagerPtr ? RemoteSessionManagerPtr->GetSessionInfo(SessionInstanceId).Pin() : nullptr;
+				if (!SessionInfoPtr)
+				{
+					return false;
+				}
+
+				const FChaosVDTraceDetails& RecordingSessionDetails = SessionInfoPtr->GetConnectionDetails();
+
+				// We didn't receive the connection details yet. Trying again later...
+				if (!RecordingSessionDetails.IsValid())
+				{
+					return false;
+				}
+
+				// Check if there is a connection attempt in progress and its results
+				const Chaos::VD::EConnectionAttemptResult ConnectionAttemptStatus = FChaosVDEngineEditorBridge::Get().GetRelayTransportInstance()->GetConnectionAttemptResult(SessionInstanceId);
+				switch (ConnectionAttemptStatus)
+				{
+				case Chaos::VD::InProgress:
+					{
+						// A connection attempt started, but we don't have the results yet. Trying again later...
+						return false;
+					}
+				case Chaos::VD::NotStarted:
+					{
+						Chaos::VD::FRelayConnectionInfo ConnectionInfo;
+						ConnectionInfo.Port = RecordingSessionDetails.Port;
+						ConnectionInfo.Address = RecordingSessionDetails.TraceTarget;
+						ConnectionInfo.CertificateAuthority = RecordingSessionDetails.CertAuth;
+						Chaos::VD::EConnectionAttemptResult Result = FChaosVDEngineEditorBridge::Get().GetRelayTransportInstance()->ConnectToRelay(SessionInstanceId, ConnectionInfo);
+
+						return Result == Chaos::VD::EConnectionAttemptResult::Success;
+					}
+				case Chaos::VD::Success:
+					// The current connection attempt succeeded, we can stop the async retry loop
+					return true;
+				case Chaos::VD::Failed:
+				default:
+					// There was a connection attempt, but it failed. Abandoning the retry attempt as they will not succeed
+					return false;
+				}
+			},
+			RecordingAttemptFailedCallback);
+#endif
+		}
+	}
+	else
+	{
+		CancelInFlightConnectionAttempt(SessionInfoRef->InstanceId);
+		RemoteSessionManager->SendStopRecordingCommand(SessionInfoRef->Address);
+	}
+}
+
+FReply FChaosVDRecordingControls::ToggleRecordingState()
+{
+	TSharedPtr<FChaosVDSessionInfo> SessionInfoPtr = GetCurrentSessionInfo();
+	if (!SessionInfoPtr)
+	{
+		return FReply::Handled();
+	}
+
+	if (EnumHasAnyFlags(SessionInfoPtr->GetSessionTypeAttributes(), EChaosVDRemoteSessionAttributes::IsMultiSessionWrapper))
+	{
+		ToggleMultiSessionSessionRecordingState(StaticCastSharedPtr<FChaosVDMultiSessionInfo>(SessionInfoPtr).ToSharedRef());
+	}
+	else
+	{
+		ToggleSingleSessionRecordingState(SessionInfoPtr.ToSharedRef());
+	}
+
+	return FReply::Handled();
+}
+
+bool FChaosVDRecordingControls::IsRecordingToggleButtonEnabled() const
+{
+	if (CurrentSelectedSessionId == FChaosVDRemoteSessionsManager::InvalidSessionGUID)
+	{
+		return false;
+	}
+
+	TSharedPtr<FChaosVDSessionInfo> SessionInfo = GetCurrentSessionInfo();
+	if (!SessionInfo)
+	{
+		return false;
+	}
+
+	if (!IsRecording())
+	{
+		return true;
+	}
+
+	// If we are recording, don't show the stop button for the mode that is disabled
+	if (EnumHasAnyFlags(SessionInfo->GetSessionTypeAttributes(), EChaosVDRemoteSessionAttributes::IsMultiSessionWrapper))
+	{
+		TSharedPtr<FChaosVDMultiSessionInfo> AsMultiSessionInfo = StaticCastSharedPtr<FChaosVDMultiSessionInfo>(SessionInfo);
+
+		return CanSelectMultiSessionTarget(SessionInfo.ToSharedRef());
+	}
+
+	return true;
+}
+
+EVisibility FChaosVDRecordingControls::IsRecordingToggleButtonVisible() const
+{
+	return EVisibility::Visible;
+}
+
+void FChaosVDRecordingControls::AddToMenu(FName ExistingMenu)
+{
+	UToolMenu* ToolBar = UToolMenus::Get()->ExtendMenu(ExistingMenu);
+	if (!ensure(ToolBar))
+	{
+		return;
+	}
+
+	FToolMenuSection& Section = ToolBar->AddSection("RecordingControls");
+
+	Section.AddDynamicEntry("Controls", FNewToolMenuSectionDelegate::CreateLambda([](FToolMenuSection& InSection)
+	{
+		const UChaosVDMainToolbarMenuContext* Context = InSection.FindContext<UChaosVDMainToolbarMenuContext>();
+		TSharedPtr<SChaosVDMainTab> MainTabPtr = Context->MainTab.Pin();
+		TSharedPtr<FChaosVDRecordingControls> RecordingControls = MainTabPtr ? MainTabPtr->GetRecordingControls() : nullptr;
+		if (!ensure(RecordingControls))
+		{
+			return;
+		}
+
+		TSharedRef<SWidget> StartRecordingButton = SNew(SBox).Padding(4.0f,0.0f)[ RecordingControls->GenerateToggleRecordingStateButton(LOCTEXT("RecordButtonDesc", "Starts a recording for the current session, saving it directly to file")) ];
+		TSharedRef<SWidget> SessionsDropdown = SNew(SBox).Padding(4.0f,0.0f)[ RecordingControls->GenerateTargetSessionSelector() ];
+		TSharedRef<SWidget> RecordingTime = RecordingControls->GenerateRecordingTimeTextBlock();
+		TSharedRef<SWidget> DataChannelsButton = SNew(SBox).Padding(4.0f,0.0f)[ RecordingControls->GenerateDataChannelsButton() ];
+		TSharedRef<SWidget> AdvancedRecordingOptionsButton = SNew(SBox).Padding(1.0f,0.0f)[ RecordingControls->GenerateAdvancedRecordOptionsButton() ];
+
+		InSection.AddSeparator("RecordingControlsDivider");
+
+		InSection.AddEntry(
+			FToolMenuEntry::InitWidget(
+				"AvailableSessions",
+				SessionsDropdown,
+				FText::GetEmpty(),
+				false,
+				false
+			));
+
+		InSection.AddEntry(
+			FToolMenuEntry::InitWidget(
+				"StartRecording",
+				StartRecordingButton,
+				FText::GetEmpty(),
+				true,
+				false
+			));
+
+		InSection.AddEntry(
+		FToolMenuEntry::InitWidget(
+			"AdvancedRecordingOptions",
+			AdvancedRecordingOptionsButton,
+			FText::GetEmpty(),
+			true,
+			false
+		));
+
+		InSection.AddEntry(
+			FToolMenuEntry::InitWidget(
+				"RecordingTime",
+				RecordingTime,
+				FText::GetEmpty(),
+				false,
+				false
+			));
+
+		InSection.AddEntry(
+			FToolMenuEntry::InitWidget(
+				"DataChannelsButton",
+				DataChannelsButton,
+				FText::GetEmpty(),
+				false,
+				false
+			));
+	}));
+}
+
+bool FChaosVDRecordingControls::IsRecording() const
+{
+	TSharedPtr<FChaosVDSessionInfo> SessionInfo = GetCurrentSessionInfo();
+	return SessionInfo ? SessionInfo->IsRecording() : false;
+}
+
+FText FChaosVDRecordingControls::GetRecordingTimeText() const
+{
+	if (TSharedPtr<FChaosVDSessionInfo> SessionInfo = GetCurrentSessionInfo())
+	{
+		FNumberFormattingOptions FormatOptions;
+		FormatOptions.MinimumFractionalDigits = 2;
+		FormatOptions.MaximumFractionalDigits = 2;
+
+		FText SecondsText = FText::AsNumber(SessionInfo->LastKnownRecordingState.ElapsedTime, &FormatOptions);
+		return FText::Format(LOCTEXT("RecordingTimer","{0} s"), SecondsText);
+	}
+
+	return LOCTEXT("RecordingTimerError","Failed to get time information");
+}
+
+TSharedPtr<SNotificationItem> FChaosVDRecordingControls::PushConnectionAttemptNotification(const FNotificationInfo& InNotificationInfo)
+{
+	TSharedPtr<SNotificationItem> ConnectionAttemptNotification = FSlateNotificationManager::Get().AddNotification(InNotificationInfo);
+	
+	if (ConnectionAttemptNotification.IsValid())
+	{
+		ConnectionAttemptNotification->SetCompletionState(SNotificationItem::CS_Pending);
+		return ConnectionAttemptNotification;
+	}
+
+	return nullptr;
+}
+
+void FChaosVDRecordingControls::UpdateConnectionAttemptNotification(const TSharedPtr<SNotificationItem>& InNotification, int32 AttemptsRemaining)
+{
+	if (InNotification)
+	{
+		InNotification->SetSubText(FText::FormatOrdered(LOCTEXT("SessionConnectionAttemptSubText", "Attempts Remaining {0}"), AttemptsRemaining));
+	}
+}
+
+void FChaosVDRecordingControls::HandleConnectionAttemptResult(FGuid SessionGUID, EChaosVDLiveConnectionAttemptResult Result, const TSharedPtr<SNotificationItem>& InNotification)
+{
+	InFlightAsyncConnectionRequests.Remove(SessionGUID);
+
+	if (InNotification)
+	{
+		switch (Result)
+		{
+		case EChaosVDLiveConnectionAttemptResult::Success:
+			{
+				InNotification->SetText(LOCTEXT("SessionConnectionSuccess", "Connected!"));
+				InNotification->SetSubText(FText::GetEmpty());
+				InNotification->SetCompletionState(SNotificationItem::ECompletionState::CS_Success);
+				break;
+			}
+		case EChaosVDLiveConnectionAttemptResult::Canceled:
+			{
+				InNotification->SetText(LOCTEXT("SessionConnectionCanceledText", "Connection Attempt Canceled"));
+				InNotification->SetCompletionState(SNotificationItem::ECompletionState::CS_None);
+				break;
+			}
+		case EChaosVDLiveConnectionAttemptResult::Failed:
+			{
+				InNotification->SetText(LOCTEXT("SessionConnectionFailedText", "Failed to connect"));
+				InNotification->SetSubText(LOCTEXT("SessionConnectionFailedSubText", "See the logs for more details..."));
+				InNotification->SetCompletionState(SNotificationItem::ECompletionState::CS_Fail);
+				break;
+			}
+		default:
+			{
+				InNotification->SetText(LOCTEXT("SessionConnectionUnknownText", "Something went wrong..."));
+				InNotification->SetSubText(LOCTEXT("SessionConnectionUnknownTextSubText", "See the logs for more details..."));
+				InNotification->SetCompletionState(SNotificationItem::ECompletionState::CS_Fail);
+				break;
+			}
+		}
+
+		InNotification->ExpireAndFadeout();
+	}
+}
+
+#undef LOCTEXT_NAMESPACE 
diff --git a/Engine/Plugins/ChaosVD/Source/ChaosVD/Private/Widgets/ChaosVDRecordingControls.h b/Engine/Plugins/ChaosVD/Source/ChaosVD/Private/Widgets/ChaosVDRecordingControls.h
new file mode 100644
--- /dev/null
+++ b/Engine/Plugins/ChaosVD/Source/ChaosVD/Private/Widgets/ChaosVDRecordingControls.h
@@ -0,0 +1,205 @@
+// Copyright Epic Games, Inc. All Rights Reserved.
+
+#pragma once
+
+#include "StatusBarSubsystem.h"
+#include "Chaos/ChaosVDRemoteSessionsManager.h"
+#include "ChaosVDEngine.h"
+#include "Templates/SharedPointer.h"
+#include "Widgets/SCompoundWidget.h"
+
+#include "ChaosVDRecordingControls.generated.h"
+
+class FMenuBuilder;
+
+namespace Chaos::VisualDebugger
+{
+	struct FChaosVDOptionalDataChannel;
+}
+
+class FAsyncCompilationNotification;
+class FReply;
+class SButton;
+class SChaosVDMainTab;
+struct FSlateBrush;
+struct FSlateIcon;
+
+UENUM()
+enum class EChaosVDLiveConnectionAttemptResult
+{
+	Success,
+	Failed,
+	Canceled
+};
+
+typedef Chaos::VisualDebugger::FChaosVDOptionalDataChannel FCVDDataChannel;
+
+class FChaosVDRecordingControls : public TSharedFromThis<FChaosVDRecordingControls>
+{
+public:
+
+	FChaosVDRecordingControls(const TSharedRef<SChaosVDMainTab>& InMainTabSharedRef);
+	~FChaosVDRecordingControls();
+
+	void Initialize();
+	void Deinitialize();
+	void AddToMenu(FName ExistingMenu);
+
+protected:
+
+	TSharedRef<SWidget> GenerateToggleRecordingStateButton(const FText& StartRecordingTooltip);
+	TSharedRef<SWidget> GenerateTargetSessionSelector();
+	TSharedRef<SWidget> GenerateTargetSessionDropdown();
+	TSharedRef<SWidget> GenerateDataChannelsMenu();
+	TSharedRef<SWidget> GenerateDataChannelsButton();
+	TSharedRef<SWidget> GenerateLoadingModeSelector();
+	TSharedRef<SWidget> GenerateRecordingTimeTextBlock();
+	TSharedRef<SWidget> GenerateAdvancedRecordOptionsButton();
+	TSharedRef<SWidget> GenerateAdvancedRecordOptionsMenu();
+	
+	EVisibility GetRecordingTimeTextBlockVisibility() const;
+
+	void ToggleChannelEnabledState(FString ChannelName);
+	bool IsChannelEnabled(FString ChannelName);
+	bool CanChangeChannelEnabledState(FString ChannelName);
+
+	void SelectTargetSession(FGuid SessionId);
+
+	void GenerateCustomTargetsMenu(FMenuBuilder& MenuBuilder);
+
+	bool IsSessionPartOfCustomTargetSelection(FGuid SessionGuid);
+	void ToggleSessionSelectionInCustomTarget(FGuid SessionGuid);
+	bool CanSelectInCustomTarget(FGuid SessionGuid) const;
+
+	bool CanSelectMultiSessionTarget(FGuid SessionGuid) const;
+	bool CanSelectMultiSessionTarget(const TSharedRef<FChaosVDSessionInfo>& SessionInfoRef) const;
+
+	void ToggleAutoStartPlayback();
+	bool IsAutoStartPlaybackEnabled() const;
+	void ToggleSaveTracesToDisk();
+	bool IsSaveTracesToDiskEnabled();	
+
+	FSlateIcon GetIconForSession(FGuid SessionId);
+
+	TSharedPtr<FChaosVDSessionInfo> GetCurrentSessionInfo() const;
+	TSharedPtr<FChaosVDSessionInfo> GetSessionInfo(FGuid Id) const;
+	void HandleRemoteSessionsUpdated();
+
+	bool HasDataChannelsSupport() const;
+
+	bool CanChangeLoadingMode() const;
+	
+	const FSlateBrush* GetRecordOrStopButton() const;
+	
+	void HandleRecordingStop(TWeakPtr<FChaosVDSessionInfo> SessionInfo);
+	void HandleRecordingStart(TWeakPtr<FChaosVDSessionInfo> SessionInfo);
+
+	void ExecuteAsyncConnectionAttemptTaskWithRetry(FGuid SessionID, int32 RemainingRetries, const TFunction<bool()>& InRecordingStartAttemptCallback, const TFunction<void()>& InRecordingFailedCallback);
+
+	FReply ToggleRecordingState();
+	void ToggleMultiSessionSessionRecordingState(const TSharedRef<FChaosVDMultiSessionInfo>& InSessionInfoRef);
+	void ToggleSingleSessionRecordingState(const TSharedRef<FChaosVDSessionInfo>& SessionInfoRef);
+	void SetSessionRecordingState(bool bIsRecording, const TSharedRef<FChaosVDSessionInfo>& SessionInfoRef);
+
+	bool IsRecordingToggleButtonEnabled() const;
+	EVisibility IsRecordingToggleButtonVisible() const;
+
+	bool IsRecording() const;
+
+	FText GetRecordingTimeText() const;
+
+	TSharedPtr<SNotificationItem> PushConnectionAttemptNotification(const FNotificationInfo& InNotificationInfo);
+	void UpdateConnectionAttemptNotification(const TSharedPtr<SNotificationItem>& InNotification, int32 AttemptsRemaining);
+	void HandleConnectionAttemptResult(FGuid SessionGUID, EChaosVDLiveConnectionAttemptResult Result, const TSharedPtr<SNotificationItem>& InNotification);
+
+	struct FTraceStoreConnectionSettings
+	{
+		int32 TraceID = INDEX_NONE;
+		FString SessionAddress;
+	};
+
+	FTraceStoreConnectionSettings GetTraceStoreConnectionSettings(const TSharedRef<FChaosVDSessionInfo>& InSessionInfoRef);
+
+	FText GetCurrentSelectedSessionName() const;
+
+	FGuid CurrentSelectedSessionId = FGuid();
+
+	FName StatusBarID;
+	
+	FStatusBarMessageHandle RecordingMessageHandle;
+	FStatusBarMessageHandle RecordingPathMessageHandle;
+	FStatusBarMessageHandle LiveSessionEndedMessageHandle;
+	FDelegateHandle RecordingStartedHandle;
+	FDelegateHandle RecordingStoppedHandle;
+
+	TWeakPtr<SChaosVDMainTab> MainTabWeakPtr;
+
+	float IntervalBetweenAutoplayConnectionAttemptsSeconds = 1.0f;
+
+	bool bRecordingButtonHovered = false;
+
+	FCurveSequence RecordingAnimation;
+
+	EChaosVDLoadRecordedDataMode CurrentLoadingMode = EChaosVDLoadRecordedDataMode::SingleSource;
+
+	void CancelInFlightConnectionAttempt(FGuid SessionID);
+
+	enum class ERecordingStartRequestState
+	{
+		Invalid,
+		InProgress,
+		Success,
+		Failed,
+		Canceled,
+	};
+
+	struct FAsyncConnectionAttemptTask : public TSharedFromThis<FAsyncConnectionAttemptTask>
+	{
+		explicit FAsyncConnectionAttemptTask(FGuid InSessionID, const TFunction<bool()>& InRecordingStartAttemptCallback, const TFunction<void()>& InRecordingFailedCallback, const TWeakPtr<FChaosVDRecordingControls>& RequestOwner)
+			: SessionID(InSessionID),
+			RecordingStartAttemptCallback(InRecordingStartAttemptCallback),
+			RecordingFailedCallback(InRecordingFailedCallback),
+			Owner(RequestOwner)
+		{
+		}
+
+		~FAsyncConnectionAttemptTask();
+
+		enum class EState : uint8
+		{
+			NotStarted,
+			InProgress,
+			Success,
+			Canceled,
+			Failed
+		};
+
+		void Cancel();
+
+		void Start(int32 InMaxRetriesAttempts, float InIntervalBetweenAttemptsSeconds);
+		void ScheduleRetry();
+		void HandleSuccess();
+		void HandleFailure();
+
+	private:
+
+		bool CanExecute() const;
+
+		void Execute();
+
+		FGuid SessionID;
+		int32 RemainingRetries = 10;
+
+		float IntervalBetweenAttemptsSeconds = 1.0f;
+	
+		TFunction<bool()> RecordingStartAttemptCallback;
+		TSharedPtr<SNotificationItem> ProgressNotification;
+		const TFunction<void()> RecordingFailedCallback;
+		EState State = EState::NotStarted;
+		FTSTicker::FDelegateHandle RetryDelegateHandle;
+		TWeakPtr<FChaosVDRecordingControls> Owner;
+	};
+
+	TMap<FGuid, TSharedPtr<FAsyncConnectionAttemptTask>> InFlightAsyncConnectionRequests;
+	bool bInitialized = false;
+};
diff --git a/Engine/Plugins/ChaosVD/Source/ChaosVD/Private/Widgets/SChaosVDMainTab.cpp b/Engine/Plugins/ChaosVD/Source/ChaosVD/Private/Widgets/SChaosVDMainTab.cpp
--- a/Engine/Plugins/ChaosVD/Source/ChaosVD/Private/Widgets/SChaosVDMainTab.cpp
+++ b/Engine/Plugins/ChaosVD/Source/ChaosVD/Private/Widgets/SChaosVDMainTab.cpp
@@ -66,7 +66,7 @@
 #include "Visualizers/ChaosVDSolverCollisionDataComponentVisualizer.h"
 #include "Widgets/SChaosBrowseTraceFileSourceModal.h"
 #include "Widgets/SChaosVDBrowseSessionsModal.h"
-#include "Widgets/SChaosVDRecordingControls.h"
+#include "Widgets/ChaosVDRecordingControls.h"
 #include "Widgets/Docking/SDockTab.h"
 #include "Widgets/Input/SComboButton.h"
 #include "Widgets/Layout/SSeparator.h"
@@ -179,22 +179,9 @@ void SChaosVDMainTab::Construct(const FArguments& InArgs, TSharedPtr<FChaosVDEng
 					SNew(SHorizontalBox)
 					+SHorizontalBox::Slot()
 					.HAlign(HAlign_Left)
-					.AutoWidth()
 					[
 						GenerateMainToolbarWidget()
 					]
-					+SHorizontalBox::Slot()
-					.HAlign(HAlign_Left)
-					.AutoWidth()
-					[
-						SNew(SChaosVDRecordingControls, StaticCastSharedRef<SChaosVDMainTab>(AsShared()))
-					]
-					+SHorizontalBox::Slot()
-					.HAlign(HAlign_Right)
-					.AutoWidth()
-					[
-						GenerateSettingsMenuWidget()
-					]
 				]
 			]
 		]
@@ -226,6 +213,17 @@ void SChaosVDMainTab::Construct(const FArguments& InArgs, TSharedPtr<FChaosVDEng
 	}
 }
 
+void SChaosVDMainTab::ExtendMainToolbar()
+{
+	if (!RecordingControls)
+	{
+		RecordingControls = MakeShared<FChaosVDRecordingControls>(StaticCastSharedRef<SChaosVDMainTab>(AsShared()));
+		RecordingControls->Initialize();
+	}
+	
+	RecordingControls->AddToMenu(MainToolBarName);
+	AddSettingsMenu();
+}
 SChaosVDMainTab::~SChaosVDMainTab()
 {
 	if (TSharedPtr<FChaosVDRemoteSessionsManager> RemoteSessionManager = FChaosVDEngineEditorBridge::Get().GetRemoteSessionsManager())
@@ -256,7 +254,11 @@ void SChaosVDMainTab::BindUICommands(const TSharedRef<FUICommandList>& InGlobalU
 	FUIAction BrowseLiveSessionsAction;
 	BrowseLiveSessionsAction.ExecuteAction.BindSP(SharedThis(this), &SChaosVDMainTab::BrowseLiveSessionsFromTraceStore);
 	InGlobalUICommandsRef->MapAction(Commands.BrowseLiveSessions, BrowseLiveSessionsAction);
-	
+
+	FUIAction DisconnectFromLiveSessionAction;
+	DisconnectFromLiveSessionAction.ExecuteAction.BindSP(SharedThis(this), &SChaosVDMainTab::HandleDisconnectSessionClicked);
+	InGlobalUICommandsRef->MapAction(Commands.DisconnectFromLiveSession, DisconnectFromLiveSessionAction);
+
 	FUIAction OpenSceneQueryBrowserAction;
 	OpenSceneQueryBrowserAction.ExecuteAction.BindSPLambda(SharedThis(this), [this]()
 	{
@@ -586,13 +588,19 @@ void SChaosVDMainTab::CombineOpenSessions()
 
 void SChaosVDMainTab::GenerateMainWindowMenu()
 {
-	FMenuBarBuilder MenuBarBuilder = FMenuBarBuilder(TSharedPtr<FUICommandList >());
+	FMenuBarBuilder MenuBarBuilder = FMenuBarBuilder(GlobalCommandList);
 	MenuBarBuilder.AddPullDownMenu(
 		LOCTEXT("FileMenuLabel", "File"),
 		FText::GetEmpty(),
 		FNewMenuDelegate::CreateSPLambda(this, [this](FMenuBuilder& MenuBuilder)
 		{
 			MenuBuilder.AddSubMenu(LOCTEXT("RecentFilesMenuLabel", "Recent Files"), LOCTEXT("RecentFilesMenuLabelToolTip", "Shows a list of recently used CVD Files"), FNewMenuDelegate::CreateSP(this, &SChaosVDMainTab::GenerateRecentFilesMenu));
+
+			const FChaosVDCommands& Commands = FChaosVDCommands::Get();
+			MenuBuilder.BeginSection(NAME_None, LOCTEXT("LegacySessionSectionLabel", "Legacy Sessions Management"));
+				MenuBuilder.AddMenuEntry(Commands.BrowseLiveSessions, NAME_None);
+				MenuBuilder.AddMenuEntry(Commands.DisconnectFromLiveSession, NAME_None);
+			MenuBuilder.EndSection();
 		}),
 		"File"
 	);
@@ -661,7 +669,7 @@ FReply SChaosVDMainTab::BrowseAndOpenChaosVDRecording()
 		case EChaosVDBrowseFileModalResponse::Cancel:
 			break;
 		default:
-				ensureMsgf(false, TEXT("Invalid responce received"));
+				ensureMsgf(false, TEXT("Invalid response received"));
 			break;
 	}
 
@@ -704,31 +712,16 @@ TSharedRef<SButton> SChaosVDMainTab::CreateSimpleButton(TFunction<FText()>&& Get
 
 TSharedRef<SWidget> SChaosVDMainTab::GenerateMainToolbarWidget()
 {
-	RegisterMainTabMenu();
+	CreateMainToolbarMenu();
 
 	FToolMenuContext MenuContext;
-
 	UChaosVDMainToolbarMenuContext* CommonContextObject = NewObject<UChaosVDMainToolbarMenuContext>();
 	CommonContextObject->MainTab = SharedThis(this);
-
 	MenuContext.AddObject(CommonContextObject);
 
-	return UToolMenus::Get()->GenerateWidget(MainToolBarName, MenuContext);
-}
-
-TSharedRef<SWidget> SChaosVDMainTab::GenerateSettingsMenuWidget()
-{
-	RegisterSettingsMenu();
-
-	FToolMenuContext MenuContext;
-
-	UChaosVDMainToolbarMenuContext* CommonContextObject = NewObject<UChaosVDMainToolbarMenuContext>();
-	CommonContextObject->MainTab = SharedThis(this);
-
-	MenuContext.AddObject(CommonContextObject);
+	ExtendMainToolbar();
 
-	const FName SettingsMenuName = FName("ChaosVDMainTabSettingsMenu");
-	return UToolMenus::Get()->GenerateWidget(SettingsMenuName, MenuContext);
+	return UToolMenus::Get()->GenerateWidget(MainToolBarName, MenuContext);
 }
 
 void SChaosVDMainTab::LoadCVDFile(const FString& InFilename, EChaosVDLoadRecordedDataMode LoadingMode)
@@ -911,7 +904,13 @@ bool SChaosVDMainTab::ConnectToLiveSession(int32 SessionID, const FString& InSes
 
 bool SChaosVDMainTab::ConnectToLiveSession_Direct(EChaosVDLoadRecordedDataMode LoadingMode) const
 {
-	return GetChaosVDEngineInstance()->ConnectToLiveSession_Direct(LoadingMode);
+	UE_LOG(LogChaosVDEditor, Error, TEXT("[%hs] This menthod is no longer supported. Please use the overload that takes a Remote Session ID."), __func__);
+	return false;
+}
+
+bool SChaosVDMainTab::ConnectToLiveSession_Direct(FGuid RemoteSessionID, EChaosVDLoadRecordedDataMode LoadingMode) const
+{
+	return GetChaosVDEngineInstance()->ConnectToLiveSession_Direct(RemoteSessionID, LoadingMode);
 }
 
 bool SChaosVDMainTab::ConnectToLiveSession_Relay(FGuid RemoteSessionID, EChaosVDLoadRecordedDataMode LoadingMode) const
@@ -944,7 +943,7 @@ void SChaosVDMainTab::CleanUpDisableCPUThrottlingDelegate() const
 	}
 }
 
-void SChaosVDMainTab::RegisterMainTabMenu()
+void SChaosVDMainTab::CreateMainToolbarMenu()
 {
 	const UToolMenus* ToolMenus = UToolMenus::Get();
 	if (ToolMenus->IsMenuRegistered(MainToolBarName))
@@ -988,105 +987,6 @@ void SChaosVDMainTab::RegisterMainTabMenu()
 			));
 	}));
 
-	Section.AddDynamicEntry("ConnectToSession", FNewToolMenuSectionDelegate::CreateLambda([](FToolMenuSection& InSection)
-	{
-		const UChaosVDMainToolbarMenuContext* Context = InSection.FindContext<UChaosVDMainToolbarMenuContext>();
-		TSharedPtr<SChaosVDMainTab> MainTabPtr = Context->MainTab.Pin();
-		if (!MainTabPtr)
-		{
-			return;
-		}
-
-		TFunction<FText()> GetTextDelegate = [WeakTab = MainTabPtr.ToWeakPtr()]()
-		{
-			return LOCTEXT("ConnectToSession", "Connect to Session");
-		};
-
-		TFunction<FText()> GetTooltipTextDelegate = [WeakTab = MainTabPtr.ToWeakPtr()]()
-		{
-			return LOCTEXT("ConnectToSessionTooltip", "Opens a panel where you can browse active live sessions and connect to one");
-		};
-		
-		FOnClicked OnClickedDelegate = FOnClicked::CreateLambda([WeakTab = StaticCastWeakPtr<SChaosVDMainTab>(MainTabPtr->AsWeak())]()
-		{
-			if (TSharedPtr<SChaosVDMainTab> TabPtr = WeakTab.Pin())
-			{
-				return TabPtr->HandleSessionConnectionClicked();
-			}
-
-			return FReply::Handled();
-		});
-
-		TSharedRef<SButton> ConnectToSessionButton = MainTabPtr->CreateSimpleButton(MoveTemp(GetTextDelegate), MoveTemp(GetTooltipTextDelegate), FChaosVDStyle::Get().GetBrush("OpenSessionIcon"), Context, MoveTemp(OnClickedDelegate));
-
-		InSection.AddEntry(
-			FToolMenuEntry::InitWidget(
-				"ConnectToSession",
-				ConnectToSessionButton,
-				FText::GetEmpty(),
-				true,
-				false
-			));
-	}));
-
-	Section.AddDynamicEntry("DisconnectFromSession", FNewToolMenuSectionDelegate::CreateLambda([](FToolMenuSection& InSection)
-	{
-		const UChaosVDMainToolbarMenuContext* Context = InSection.FindContext<UChaosVDMainToolbarMenuContext>();
-		TSharedPtr<SChaosVDMainTab> MainTabPtr = Context->MainTab.Pin();
-		if (!MainTabPtr)
-		{
-			return;
-		}
-
-		TFunction<FText()> GetTextDelegate = [WeakTab = MainTabPtr.ToWeakPtr()]()
-		{
-			if (TSharedPtr<SChaosVDMainTab> TabPtr = WeakTab.Pin())
-			{
-				return TabPtr->GetDisconnectButtonText();
-			}
-
-			return FText();
-		};
-
-		TFunction<FText()> GetTooltipTextDelegate = [WeakTab = MainTabPtr.ToWeakPtr()]()
-		{
-			return LOCTEXT("DisconnectFromSessionTooltip", "Disconnects from the current live session but it does not stop it");
-		};
-			
-		FOnClicked OnClickedDelegate = FOnClicked::CreateLambda([WeakTab = StaticCastWeakPtr<SChaosVDMainTab>(MainTabPtr->AsWeak())]()
-		{
-			if (TSharedPtr<SChaosVDMainTab> TabPtr = WeakTab.Pin())
-			{
-				return TabPtr->HandleDisconnectSessionClicked();
-			}
-
-			return FReply::Handled();
-		});
-
-		TAttribute<EVisibility> ButtonVisibilityAttribute;
-		ButtonVisibilityAttribute.BindLambda([WeakTab = MainTabPtr.ToWeakPtr()]
-		{
-			if (TSharedPtr<SChaosVDMainTab> TabPtr = WeakTab.Pin())
-			{
-				return TabPtr->GetChaosVDEngineInstance()->HasAnyLiveSessionActive() ? EVisibility::Visible : EVisibility::Collapsed;
-			}
-			return EVisibility::Collapsed;
-		});
-
-		TSharedRef<SButton> DisconnectFromSessionButton = MainTabPtr->CreateSimpleButton(MoveTemp(GetTextDelegate), MoveTemp(GetTooltipTextDelegate), FChaosVDStyle::Get().GetBrush("OpenSessionIcon"), Context, MoveTemp(OnClickedDelegate));
-
-		DisconnectFromSessionButton->SetVisibility(ButtonVisibilityAttribute);
-
-		InSection.AddEntry(
-			FToolMenuEntry::InitWidget(
-				"DisconnectFromSessionButton",
-				DisconnectFromSessionButton,
-				FText::GetEmpty(),
-				true,
-				false
-			));
-	}));
-
 	Section.AddDynamicEntry("CombineFiles", FNewToolMenuSectionDelegate::CreateLambda([](FToolMenuSection& InSection)
 	{
 		const UChaosVDMainToolbarMenuContext* Context = InSection.FindContext<UChaosVDMainToolbarMenuContext>();
@@ -1176,16 +1076,9 @@ void SChaosVDMainTab::RegisterMainTabMenu()
 	}));
 }
 
-void SChaosVDMainTab::RegisterSettingsMenu()
+void SChaosVDMainTab::AddSettingsMenu()
 {
-	const UToolMenus* ToolMenus = UToolMenus::Get();
-	const FName SettingsMenuName = FName("ChaosVDMainTabSettingsMenu");
-	if (ToolMenus->IsMenuRegistered(SettingsMenuName))
-	{
-		return;
-	}
-
-	UToolMenu* ToolBar = UToolMenus::Get()->RegisterMenu(SettingsMenuName, NAME_None, EMultiBoxType::SlimHorizontalToolBar);
+	UToolMenu* ToolBar = UToolMenus::Get()->ExtendMenu(MainToolBarName);
 	if (!ToolBar)
 	{
 		return;
@@ -1263,7 +1156,7 @@ FReply SChaosVDMainTab::HandleSessionConnectionClicked()
 	return FReply::Handled();
 }
 
-FReply SChaosVDMainTab::HandleDisconnectSessionClicked()
+void SChaosVDMainTab::HandleDisconnectSessionClicked()
 {
 	if (TSharedPtr<FChaosVDPlaybackController> PlaybackControllerPtr =  GetChaosVDEngineInstance()->GetPlaybackController())
 	{
@@ -1281,8 +1174,6 @@ FReply SChaosVDMainTab::HandleDisconnectSessionClicked()
 			PlaybackControllerPtr->HandleDisconnectedFromSession();
 		}
 	}
-
-	return FReply::Handled();
 }
 
 FText SChaosVDMainTab::GetDisconnectButtonText() const
diff --git a/Engine/Plugins/ChaosVD/Source/ChaosVD/Private/Widgets/SChaosVDRecordingControls.cpp b/Engine/Plugins/ChaosVD/Source/ChaosVD/Private/Widgets/SChaosVDRecordingControls.cpp
deleted file mode 100644
diff --git a/Engine/Plugins/ChaosVD/Source/ChaosVD/Private/Widgets/SChaosVDRecordingControls.h b/Engine/Plugins/ChaosVD/Source/ChaosVD/Private/Widgets/SChaosVDRecordingControls.h
deleted file mode 100644
diff --git a/Engine/Plugins/ChaosVD/Source/ChaosVD/Public/ChaosVDEngine.h b/Engine/Plugins/ChaosVD/Source/ChaosVD/Public/ChaosVDEngine.h
--- a/Engine/Plugins/ChaosVD/Source/ChaosVD/Public/ChaosVDEngine.h
+++ b/Engine/Plugins/ChaosVD/Source/ChaosVD/Public/ChaosVDEngine.h
@@ -42,6 +42,7 @@ public:
 
 	CHAOSVD_API void Initialize();
 	CHAOSVD_API void CloseActiveTraceSessions();
+	CHAOSVD_API void CloseSessionByRemoteSessionID(FGuid RemoteSessionID);
 
 	void StopActiveTraceSessions();
 
@@ -75,7 +76,11 @@ public:
 	void LoadCombinedMultiRecording(const FString& FilePath);
 
 	bool ConnectToLiveSession(uint32 SessionID, const FString& InSessionAddress, EChaosVDLoadRecordedDataMode LoadingMode = EChaosVDLoadRecordedDataMode::MultiSource);
+
+	UE_DEPRECATED(5.8, "Please use the version that takes a remote Session ID")
 	bool ConnectToLiveSession_Direct(EChaosVDLoadRecordedDataMode LoadingMode = EChaosVDLoadRecordedDataMode::MultiSource);
+	
+	bool ConnectToLiveSession_Direct(FGuid RemoteSessionID, EChaosVDLoadRecordedDataMode LoadingMode = EChaosVDLoadRecordedDataMode::MultiSource);
 	bool ConnectToLiveSession_Relay(FGuid RemoteSessionID, EChaosVDLoadRecordedDataMode LoadingMode = EChaosVDLoadRecordedDataMode::MultiSource);
 
 	void OpenSession(const FChaosVDTraceSessionDescriptor& SessionDescriptor, EChaosVDLoadRecordedDataMode LoadingMode = EChaosVDLoadRecordedDataMode::MultiSource);
@@ -94,6 +99,8 @@ public:
 		return OnSessionClosedDelegate;
 	}
 
+	const FChaosVDTraceSessionDescriptor& GetSessionDescriptor(FGuid RemoteSessionID) const;
+
 private:
 
 	void LoadRecording_Internal(const TFunction<FString(const TSharedPtr<FChaosVDRecording>&)>&, EChaosVDLoadRecordedDataMode LoadingMode = EChaosVDLoadRecordedDataMode::MultiSource);
diff --git a/Engine/Plugins/ChaosVD/Source/ChaosVD/Public/ChaosVDPlaybackController.h b/Engine/Plugins/ChaosVD/Source/ChaosVD/Public/ChaosVDPlaybackController.h
--- a/Engine/Plugins/ChaosVD/Source/ChaosVD/Public/ChaosVDPlaybackController.h
+++ b/Engine/Plugins/ChaosVD/Source/ChaosVD/Public/ChaosVDPlaybackController.h
@@ -119,6 +119,9 @@ public:
 	FChaosVDPlaybackController(const TWeakPtr<FChaosVDScene>& InSceneToControl);
 	virtual ~FChaosVDPlaybackController() override;
 
+	void Initialize();
+	void DeInitialize();
+
 	/** Loads a recording using a CVD Trace Session Descriptor */
 	bool LoadChaosVDRecordingFromTraceSession(const FChaosVDTraceSessionDescriptor& InSessionDescriptor);
 
@@ -488,6 +491,8 @@ protected:
 
 	TSharedRef<FChaosVDTrackInfo> CreateTrackInfo(int32 SlotIndex = INDEX_NONE);
 
+	void HandleGeneralSettingsUpdated(UObject* SettingsObject);
+
 	bool bIsScrubbingTimeline = false;
 
 	/** Map containing all track info, by track type*/
@@ -535,6 +540,10 @@ protected:
 	int32 LastAssignedTrackSlot = 0;
 
 	EChaosVDSyncTimelinesMode CurrentSyncMode = EChaosVDSyncTimelinesMode::RecordedTimestamp;
+
+	bool bStartLiveSessionPlaybackAutomatically = true;
+
+	bool bInitialized = false;
 };
 
 template <typename TVisitorCallback>
diff --git a/Engine/Plugins/ChaosVD/Source/ChaosVD/Public/Trace/ChaosVDTraceManager.h b/Engine/Plugins/ChaosVD/Source/ChaosVD/Public/Trace/ChaosVDTraceManager.h
--- a/Engine/Plugins/ChaosVD/Source/ChaosVD/Public/Trace/ChaosVDTraceManager.h
+++ b/Engine/Plugins/ChaosVD/Source/ChaosVD/Public/Trace/ChaosVDTraceManager.h
@@ -51,6 +51,7 @@ namespace Chaos::VD
 	{
 	public:
 		explicit FRelayDataStream(FGuid InRemoteSessionID);
+		explicit FRelayDataStream(FGuid InRemoteSessionID, TUniquePtr<FArchiveFileWriterGeneric>&& InFileWriter);
 
 		virtual ~FRelayDataStream() override;
 
@@ -65,11 +66,14 @@ namespace Chaos::VD
 
 		TSpscQueue<FSerializedDataBuffer> DataQueue;
 		FEvent*	CancelEvent = nullptr;
+		FEvent*	WaitForDataEvent = nullptr;
 		TOptional<FSerializedDataBuffer> CurrentPacket;
 		uint32 CurrentPacketOffset = 0;
 		bool bClosed = false;
 
 		FGuid RemoteSessionID = FGuid();
+
+		TUniquePtr<FArchiveFileWriterGeneric> FileWriter;
 	};
 }
 
@@ -78,12 +82,33 @@ namespace Chaos::VD
  */
 struct FChaosVDTraceSessionDescriptor
 {
+
+	/** The trace session name. Used to get the trace analysis instance for this CVD session*/
 	FString SessionName;
+	/** Port number used to establish a connection to this session. The value is 0 if no specific port was used*/
 	uint16 SessionPort = 0;
+	/** Session ID of this CVD session in the Remote Session management system*/
 	FGuid RemoteSessionID;
+
+	/** True if this is a session recorded live (not loaded from disk)*/
 	bool bIsLiveSession = false;
 
-	bool IsValid() const { return !SessionName.IsEmpty(); }
+	/** Returns true if valid. For a session to be valued needs at minimum have a Trace Session name */
+	bool IsValid() const
+	{
+		return !SessionName.IsEmpty();
+	}
+
+	friend bool operator==(const FChaosVDTraceSessionDescriptor& Lhs, const FChaosVDTraceSessionDescriptor& RHS)
+	{
+		return Lhs.SessionName == RHS.SessionName
+			&& Lhs.RemoteSessionID == RHS.RemoteSessionID;
+	}
+
+	friend bool operator!=(const FChaosVDTraceSessionDescriptor& Lhs, const FChaosVDTraceSessionDescriptor& RHS)
+	{
+		return !(Lhs == RHS);
+	}
 };
 
 /** Objects that allows us to use TLS to temporarily store and access a ptr to an existing instance.
@@ -169,9 +194,11 @@ public:
 	 */
 	UE::Trace::FStoreClient::FTraceData GetTraceDataStreamFromStore(FStringView InSessionHost, uint32 SessionID);
 
+	static FString GetDefaultSaveDirectoryPath();
+
 private:
 
-	TUniquePtr<FArchiveFileWriterGeneric>  OpenTraceFileForWrite();
+	TUniquePtr<FArchiveFileWriterGeneric> OpenTraceFileForWrite(const FString& InDirectoryToSavePath);
 
 	bool ConnectToLiveSession_Internal(uint32 SessionID, const TSharedPtr<FChaosVDRecording>& ExistingRecordingPtr, const FString& InRequestedSessionName, UE::Trace::FStoreClient::FTraceData&& InTraceDataStream);
 
@@ -186,6 +213,9 @@ private:
 	TMap<FString, TSharedPtr<const TraceServices::IAnalysisSession>> AnalysisSessionByName;
 
 	TSharedPtr<FChaosVDTraceModule> ChaosVDTraceModule;
+
+	static uint16 CurrentDirectModeStartingPort;
+	static constexpr uint16 DefaultStartingPort = 1993;
 };
 
 template <typename TCallback>
diff --git a/Engine/Plugins/ChaosVD/Source/ChaosVD/Public/Widgets/SChaosVDMainTab.h b/Engine/Plugins/ChaosVD/Source/ChaosVD/Public/Widgets/SChaosVDMainTab.h
--- a/Engine/Plugins/ChaosVD/Source/ChaosVD/Public/Widgets/SChaosVDMainTab.h
+++ b/Engine/Plugins/ChaosVD/Source/ChaosVD/Public/Widgets/SChaosVDMainTab.h
@@ -28,6 +28,7 @@ class FChaosVDWorldOutlinerTab;
 class FStructOnScope;
 class FUICommandList;
 class SButton;
+class FChaosVDRecordingControls;
 class SDockTab;
 class FChaosVDScene;
 
@@ -105,9 +106,15 @@ public:
 	{
 		return ComponentVisualizers;
 	}
-	
+
+	/** Connects to a live session using via the Unreal Trace Server and its Trace Store. This is a legacy option */
 	bool ConnectToLiveSession(int32 SessionID, const FString& InSessionAddress, EChaosVDLoadRecordedDataMode LoadingMode) const;
+
+	UE_DEPRECATED(5.8, "Use the version that takes a Remote session ID")
 	bool ConnectToLiveSession_Direct(EChaosVDLoadRecordedDataMode LoadingMode) const;
+
+	bool ConnectToLiveSession_Direct(FGuid RemoteSessionID, EChaosVDLoadRecordedDataMode LoadingMode) const;
+
 	bool ConnectToLiveSession_Relay(FGuid RemoteSessionID, EChaosVDLoadRecordedDataMode LoadingMode) const;
 
 	/**
@@ -160,8 +167,15 @@ public:
 
 	virtual void Tick(const FGeometry& AllottedGeometry, const double InCurrentTime, const float InDeltaTime) override;
 
+	TSharedPtr<FChaosVDRecordingControls> GetRecordingControls()
+	{
+		return RecordingControls;
+	}
+
 private:
 
+	void ExtendMainToolbar();
+
 	TSharedPtr<SNotificationItem> CreateProcessingTraceDataNotification();
 
 	void UpdateTraceProcessingNotificationUpdate();
@@ -177,8 +191,8 @@ private:
 	void SetUpDisableCPUThrottlingDelegate();
 	void CleanUpDisableCPUThrottlingDelegate() const;
 
-	void RegisterMainTabMenu();
-	void RegisterSettingsMenu();
+	void CreateMainToolbarMenu();
+	void AddSettingsMenu();
 
 	CHAOSVD_API void HandleTabSpawned(TSharedRef<SDockTab> Tab, FName TabID);
 	CHAOSVD_API void HandleTabDestroyed(TSharedRef<SDockTab> Tab, FName TabID);
@@ -197,7 +211,6 @@ private:
 	TSharedRef<SButton> CreateSimpleButton(TFunction<FText()>&& GetTextDelegate, TFunction<FText()>&& ToolTipTextDelegate, const FSlateBrush* ButtonIcon, const UChaosVDMainToolbarMenuContext* MenuContext, const FOnClicked& InButtonClickedCallback);
 
 	TSharedRef<SWidget> GenerateMainToolbarWidget();
-	TSharedRef<SWidget> GenerateSettingsMenuWidget();
 
 	void BrowseChaosVDRecordingFromFolder(FStringView FolderPath, EChaosVDLoadRecordedDataMode LoadingMode);
 
@@ -231,7 +244,7 @@ private:
 	FDelegateHandle DisableCPUThrottleHandle;
 
 	FReply HandleSessionConnectionClicked();
-	FReply HandleDisconnectSessionClicked();
+	void HandleDisconnectSessionClicked();
 	FText GetDisconnectButtonText() const;
 
 	bool bShowTrackSelectorKeyShortcut = false;
@@ -241,6 +254,8 @@ private:
 	/** Commandlist used for any UI actions that need to be processed globally regardless on which specific widget we are in */
 	TSharedPtr<FUICommandList> GlobalCommandList;
 
+	TSharedPtr<FChaosVDRecordingControls> RecordingControls;
+
 	static const FName MainToolBarName;
 };
 
diff --git a/Engine/Source/Runtime/Experimental/ChaosSolverEngine/Private/Chaos/ChaosVDRemoteSessionsManager.cpp b/Engine/Source/Runtime/Experimental/ChaosSolverEngine/Private/Chaos/ChaosVDRemoteSessionsManager.cpp
--- a/Engine/Source/Runtime/Experimental/ChaosSolverEngine/Private/Chaos/ChaosVDRemoteSessionsManager.cpp
+++ b/Engine/Source/Runtime/Experimental/ChaosSolverEngine/Private/Chaos/ChaosVDRemoteSessionsManager.cpp
@@ -44,17 +44,23 @@ bool FChaosVDSessionInfo::IsRecording() const
 
 EChaosVDRecordingMode FChaosVDSessionInfo::GetRecordingMode() const
 {
-	return LastKnownConnectionDetails.IsValid() ? LastKnownConnectionDetails.Mode : LastRequestedRecordingMode;
+	PRAGMA_DISABLE_DEPRECATION_WARNINGS
+	return LastKnownConnectionDetails.Mode;
+	PRAGMA_ENABLE_DEPRECATION_WARNINGS
 }
 
 EChaosVDRecordingMode FChaosVDSessionInfo::GetLastRequestedRecordingMode() const
 {
-	return LastRequestedRecordingMode;
+	PRAGMA_DISABLE_DEPRECATION_WARNINGS
+	return EChaosVDRecordingMode::Invalid;
+	PRAGMA_ENABLE_DEPRECATION_WARNINGS
 }
 
 void FChaosVDSessionInfo::SetLastRequestedRecordingMode(EChaosVDRecordingMode NewRecordingMode)
 {
+	PRAGMA_DISABLE_DEPRECATION_WARNINGS
 	LastRequestedRecordingMode = NewRecordingMode;
+	PRAGMA_ENABLE_DEPRECATION_WARNINGS
 }
 
 bool FChaosVDSessionInfo::IsConnected() const
@@ -80,6 +86,7 @@ bool FChaosVDMultiSessionInfo::IsRecording() const
 
 EChaosVDRecordingMode FChaosVDMultiSessionInfo::GetRecordingMode() const
 {
+	PRAGMA_DISABLE_DEPRECATION_WARNINGS
 	EChaosVDRecordingMode FirstValidInstanceRecordingMode = EChaosVDRecordingMode::Invalid;
 	EnumerateInnerSessions([&FirstValidInstanceRecordingMode](const TSharedRef<FChaosVDSessionInfo>& InSessionRef)
 	{
@@ -96,6 +103,7 @@ EChaosVDRecordingMode FChaosVDMultiSessionInfo::GetRecordingMode() const
 	});
 
 	return FirstValidInstanceRecordingMode;
+	PRAGMA_ENABLE_DEPRECATION_WARNINGS
 }
 
 void FChaosVDRemoteSessionsManager::RegisterBuiltInMessageTypes()
@@ -143,6 +151,19 @@ TWeakPtr<FChaosVDSessionInfo> FChaosVDRemoteSessionsManager::GetSessionInfo(FGui
 	return nullptr;
 }
 
+TWeakPtr<FChaosVDSessionInfo> FChaosVDRemoteSessionsManager::GetFirstAvailableSingleSessionInfo()
+{
+	for (const TPair<FGuid, TSharedPtr<FChaosVDSessionInfo>>& SessionWithID : ActiveSessionsByInstanceId)
+	{
+		if (SessionWithID.Value && !EnumHasAnyFlags(SessionWithID.Value->GetSessionTypeAttributes(), EChaosVDRemoteSessionAttributes::IsMultiSessionWrapper))
+		{
+			return SessionWithID.Value;
+		}
+	}
+
+	return nullptr;
+}
+
 TSharedPtr<FChaosVDSessionInfo> FChaosVDRemoteSessionsManager::CreatedWrapperSessionInfo(FGuid InstanceId, const FString& SessionName)
 {
 	TSharedPtr<FChaosVDMultiSessionInfo> NewSessionInfo = MakeShared<FChaosVDMultiSessionInfo>();
@@ -190,7 +211,7 @@ void FChaosVDRemoteSessionsManager::RegisterExternalSupportedMessageType(const U
 
 void FChaosVDRemoteSessionsManager::InitializeMessagingSystem(const TSharedPtr<IMessageBus>& InMessageBus)
 {
-	if (!ensure(InMessageBus))
+	if (!InMessageBus)
 	{
 		return;
 	}
@@ -198,7 +219,7 @@ void FChaosVDRemoteSessionsManager::InitializeMessagingSystem(const TSharedPtr<I
 	MessageBusPtr = InMessageBus;
 	MessageEndpoint = CreateEndPoint(InMessageBus.ToSharedRef());
 	
-	if (!ensure(MessageEndpoint))
+	if (!MessageEndpoint)
 	{
 		return;
 	}
diff --git a/Engine/Source/Runtime/Experimental/ChaosSolverEngine/Public/Chaos/ChaosVDRemoteSessionsManager.h b/Engine/Source/Runtime/Experimental/ChaosSolverEngine/Public/Chaos/ChaosVDRemoteSessionsManager.h
--- a/Engine/Source/Runtime/Experimental/ChaosSolverEngine/Public/Chaos/ChaosVDRemoteSessionsManager.h
+++ b/Engine/Source/Runtime/Experimental/ChaosSolverEngine/Public/Chaos/ChaosVDRemoteSessionsManager.h
@@ -249,9 +249,12 @@ public:
 	CHAOSSOLVERENGINE_API const FChaosVDTraceDetails& GetConnectionDetails();
 
 	CHAOSSOLVERENGINE_API virtual bool IsRecording() const;
-	
+
+	UE_DEPRECATED(5.8, "This method will no longer be used and there is not a replacement planned.")
 	CHAOSSOLVERENGINE_API virtual EChaosVDRecordingMode GetRecordingMode() const;
+	UE_DEPRECATED(5.8, "This method will no longer be used and there is not a replacement planned.")
 	CHAOSSOLVERENGINE_API virtual EChaosVDRecordingMode GetLastRequestedRecordingMode() const;
+	UE_DEPRECATED(5.8, "This method will no longer be used and there is not a replacement planned.")
 	CHAOSSOLVERENGINE_API virtual void SetLastRequestedRecordingMode(EChaosVDRecordingMode NewRecordingMode);
 
 
@@ -264,6 +267,7 @@ public:
 	}
 
 protected:
+	UE_DEPRECATED(5.8, "This method will no longer be used and there is not a replacement planned.")
 	EChaosVDRecordingMode LastRequestedRecordingMode = EChaosVDRecordingMode::Invalid;
 	FChaosVDTraceDetails LastKnownConnectionDetails;
 	const EChaosVDRemoteSessionAttributes SessionTypeAttributes;
@@ -284,6 +288,8 @@ struct FChaosVDMultiSessionInfo : public FChaosVDSessionInfo
 	virtual ~FChaosVDMultiSessionInfo() override = default;
 
 	CHAOSSOLVERENGINE_API virtual bool IsRecording() const override;
+
+	UE_DEPRECATED(5.8, "This method will no longer be used and there is not a replacement planned.")
 	CHAOSSOLVERENGINE_API virtual EChaosVDRecordingMode GetRecordingMode() const override;
 
 	template<typename TCallback>
@@ -381,6 +387,11 @@ public:
 	 */
 	CHAOSSOLVERENGINE_API TWeakPtr<FChaosVDSessionInfo> GetSessionInfo(FGuid Id);
 
+	/**
+	 * Returns the first available session info 
+	 */
+	CHAOSSOLVERENGINE_API TWeakPtr<FChaosVDSessionInfo> GetFirstAvailableSingleSessionInfo();
+
 	/**
 	 * Iterates through all active and valid cvd sessions, and executes the provided callback to it.
 	 * if the callbacks returns false, the iteration will stop
diff --git a/Engine/Source/Runtime/Experimental/ChaosVisualDebugger/Private/ChaosVDRuntimeModule.cpp b/Engine/Source/Runtime/Experimental/ChaosVisualDebugger/Private/ChaosVDRuntimeModule.cpp
--- a/Engine/Source/Runtime/Experimental/ChaosVisualDebugger/Private/ChaosVDRuntimeModule.cpp
+++ b/Engine/Source/Runtime/Experimental/ChaosVisualDebugger/Private/ChaosVDRuntimeModule.cpp
@@ -169,7 +169,6 @@ FChaosVDTraceDetails FChaosVDRuntimeModule::GetCurrentTraceSessionDetails() cons
 	if (FTraceAuxiliary::IsConnected(Details.SessionGuid, Details.TraceGuid))
 	{
 		Details.TraceTarget = FTraceAuxiliary::GetTraceDestinationString();
-		Details.Mode = GetCurrentRecordingMode();
 		Details.TransportMode = CurrentTransportMode;
 		Details.MarkAsValid();
 		return Details;
@@ -186,7 +185,7 @@ void FChaosVDRuntimeModule::StopTrace()
 
 void FChaosVDRuntimeModule::GenerateRecordingFileName(FString& OutFileName)
 {
-	FStringFormatOrderedArguments NameArgs { FString(FApp::GetProjectName()), LexToString(FApp::GetBuildTargetType()), FDateTime::Now().ToString(TEXT("%Y%m%d_%H%M%S")) };
+	FStringFormatOrderedArguments NameArgs { FString(FApp::GetProjectName()), LexToString(FApp::GetBuildTargetType()), FDateTime::Now().ToString(TEXT("%Y%m%d_%H%M%S%s")) };
 	
 	OutFileName = FString::Format(TEXT("ChaosVD-{0}-{1}-{2}.utrace"), NameArgs);
 }
@@ -340,27 +339,19 @@ void FChaosVDRuntimeModule::StartRecording(const FChaosVDStartRecordingCommandMe
 
 	if (InRecordingStartCommand.RecordingMode == EChaosVDRecordingMode::File)
 	{
-		FString FileTarget;
-		GenerateRecordingFileName(FileTarget);
 
-		UE_LOG(LogChaosVDRuntime, Log, TEXT("[%hs] Generated trace file name [%s]"), __func__, *FileTarget);
-
-		switch (InRecordingStartCommand.TransportMode)
+		if (InRecordingStartCommand.TransportMode == EChaosVDTransportMode::FileSystem)
 		{
-		case EChaosVDTransportMode::FileSystem:
+			FString FileTarget;
+			GenerateRecordingFileName(FileTarget);
+			UE_LOG(LogChaosVDRuntime, Log, TEXT("[%hs] Generated trace file name [%s]"), __func__, *FileTarget);
+
 			bIsRecording = FTraceAuxiliary::Start(FTraceAuxiliary::EConnectionType::File, *FileTarget, nullptr, &TracingOptions);
-			break;
-		case EChaosVDTransportMode::TraceServer:
-			bIsRecording = FTraceAuxiliary::Start(FTraceAuxiliary::EConnectionType::Network, *InRecordingStartCommand.Target, nullptr, &TracingOptions);
-			break;
-		case EChaosVDTransportMode::Direct:
-		case EChaosVDTransportMode::Relay:
-		case EChaosVDTransportMode::Invalid:
+		}
+		else
+		{
 			UE_LOG(LogChaosVDRuntime, Error, TEXT("[%hs] Unsupported Transport mode [%s]"), __func__, *UEnum::GetValueAsString(InRecordingStartCommand.TransportMode))
-			break;
 		}
-		
-		CurrentRecordingMode = EChaosVDRecordingMode::File;
 	}
 	else if(InRecordingStartCommand.RecordingMode == EChaosVDRecordingMode::Live)
 	{
@@ -401,8 +392,6 @@ void FChaosVDRuntimeModule::StartRecording(const FChaosVDStartRecordingCommandMe
 		UE_LOG(LogChaosVDRuntime, Error, TEXT("[%hs] Unsupported Transport mode [%s]"), __func__, *UEnum::GetValueAsString(InRecordingStartCommand.TransportMode))
 			break;
 		}
-
-		CurrentRecordingMode = EChaosVDRecordingMode::Live;
 	}
 	else
 	{
@@ -440,8 +429,6 @@ void FChaosVDRuntimeModule::StartRecording(const FChaosVDStartRecordingCommandMe
 			RecordingStartFailedDelegate.Broadcast(FailureReason);
 		}
 
-		CurrentRecordingMode = EChaosVDRecordingMode::Invalid;
-
 		RelayWriter = nullptr;
 	}
 }
@@ -454,12 +441,11 @@ void FChaosVDRuntimeModule::StopRecording()
 	FTraceAuxiliary::OnTraceStopped.RemoveAll(this);
 	OnTraceStoppedDelegateHandle = FDelegateHandle();
 
-	CurrentRecordingMode = EChaosVDRecordingMode::Invalid;
 	CurrentTransportMode = EChaosVDTransportMode::Invalid;
 
 	if (!bIsRecording)
 	{
-		UE_LOG(LogChaosVDRuntime, Warning, TEXT("[%s] Attempted to stop recorded when there is no CVD recording active."), ANSI_TO_TCHAR(__FUNCTION__));
+		UE_LOG(LogChaosVDRuntime, Warning, TEXT("[%hs] Attempted to stop recorded when there is no CVD recording active."), __func__);
 		return;
 	}
 
diff --git a/Engine/Source/Runtime/Experimental/ChaosVisualDebugger/Public/ChaosVDRecordingDetails.h b/Engine/Source/Runtime/Experimental/ChaosVisualDebugger/Public/ChaosVDRecordingDetails.h
--- a/Engine/Source/Runtime/Experimental/ChaosVisualDebugger/Public/ChaosVDRecordingDetails.h
+++ b/Engine/Source/Runtime/Experimental/ChaosVisualDebugger/Public/ChaosVDRecordingDetails.h
@@ -91,8 +91,7 @@ struct FChaosVDTraceDetails
 	UE_DEPRECATED(5.7, "This property is no longer used and it will removed in the future")
 	bool bIsConnected = false;
 
-	/** Which recording mode is being used by this session (Live or File)*/
-	UPROPERTY()
+	UE_DEPRECATED(5.8, "Recording mode is no longer valid and will not be replaced")
 	EChaosVDRecordingMode Mode = EChaosVDRecordingMode::Invalid;
 
 	/** How the data is being transported to the editor (Trace Server, Custom Relay, Direct, etc.)*/
diff --git a/Engine/Source/Runtime/Experimental/ChaosVisualDebugger/Public/ChaosVDRuntimeModule.h b/Engine/Source/Runtime/Experimental/ChaosVisualDebugger/Public/ChaosVDRuntimeModule.h
--- a/Engine/Source/Runtime/Experimental/ChaosVisualDebugger/Public/ChaosVDRuntimeModule.h
+++ b/Engine/Source/Runtime/Experimental/ChaosVisualDebugger/Public/ChaosVDRuntimeModule.h
@@ -143,9 +143,10 @@ public:
 	/** Return the current connection details for the underlying trace session */
 	UE_API FChaosVDTraceDetails GetCurrentTraceSessionDetails() const;
 
+	UE_DEPRECATED(5.8, "This method is no longer used and will not be replaced")
 	EChaosVDRecordingMode GetCurrentRecordingMode() const
 	{
-		return CurrentRecordingMode;
+		return EChaosVDRecordingMode::Invalid;
 	}
 
 	/** Finds a valid file name for a new file - Used to generate the file name for the Trace recording */
@@ -229,7 +230,6 @@ private:
 
 	TMap<FString, bool> OriginalTraceChannelsState;
 
-	EChaosVDRecordingMode CurrentRecordingMode = EChaosVDRecordingMode::Invalid;
 	EChaosVDTransportMode CurrentTransportMode = EChaosVDTransportMode::Invalid;
 
 	TUniquePtr<Chaos::VD::FRelayTraceWriter> RelayWriter;
