[ChaosVD] Fix CIS localization issue

There was a duplicated localization key in ChaosVDRecordingControls.cpp.

diff --git a/Engine/Plugins/ChaosVD/Source/ChaosVD/Private/Widgets/ChaosVDRecordingControls.cpp b/Engine/Plugins/ChaosVD/Source/ChaosVD/Private/Widgets/ChaosVDRecordingControls.cpp
--- a/Engine/Plugins/ChaosVD/Source/ChaosVD/Private/Widgets/ChaosVDRecordingControls.cpp
+++ b/Engine/Plugins/ChaosVD/Source/ChaosVD/Private/Widgets/ChaosVDRecordingControls.cpp
@@ -309,7 +309,7 @@ TSharedRef<SWidget> FChaosVDRecordingControls::GenerateAdvancedRecordOptionsMenu
 
 		MenuBuilder.AddMenuEntry(
 		LOCTEXT("AutoSaveToDiskLabel", "Save Recorded Data to Disk"),
-		LOCTEXT("AutoStartPlaybackToolTip", "If set to true, any recorded data for the session will be saved automatically as soon it is recorded."),
+		LOCTEXT("AutoSaveToDiskToolTip", "If set to true, any recorded data for the session will be saved automatically as soon it is recorded."),
 		FSlateIcon(),
 		SaveToDiskUIAction, NAME_None, EUserInterfaceActionType::ToggleButton);
 	}
