[ChaosVD] Re-Enabling the transport mode override selector

The change to hide the option was only meant to be applied in UE 5.7, but I didn't use the correct robomerge tage to exclude Main.

#FailedValidations Preflight

diff --git a/Engine/Plugins/ChaosVD/Source/ChaosVD/Private/Settings/ChaosVDGeneralSettings.h b/Engine/Plugins/ChaosVD/Source/ChaosVD/Private/Settings/ChaosVDGeneralSettings.h
--- a/Engine/Plugins/ChaosVD/Source/ChaosVD/Private/Settings/ChaosVDGeneralSettings.h
+++ b/Engine/Plugins/ChaosVD/Source/ChaosVD/Private/Settings/ChaosVDGeneralSettings.h
@@ -45,7 +45,7 @@ public:
 	int32 MaxConnectionRetries = 14;
 
 	/** [DEBUG Setting] Data Transport Override - Used to change the transport mode for the trace data for testing. This is not saved by design */
-	UPROPERTY()
+	UPROPERTY(EditAnywhere, Category = Connection)
 	EChaosVDTransportMode DataTransportModeOverride = EChaosVDTransportMode::Direct;
 
 	/** If True, any traces done to memory will also be saved to disk at the same time in the user's document folder (or a user specified folder). This is not saved by design */
