Implement geometry hierarchy viewer for particles in CVD

diff --git a/Engine/Plugins/ChaosVD/Source/ChaosVD/Private/ChaosVDGeometryBuilder.cpp b/Engine/Plugins/ChaosVD/Source/ChaosVD/Private/ChaosVDGeometryBuilder.cpp
--- a/Engine/Plugins/ChaosVD/Source/ChaosVD/Private/ChaosVDGeometryBuilder.cpp
+++ b/Engine/Plugins/ChaosVD/Source/ChaosVD/Private/ChaosVDGeometryBuilder.cpp
@@ -342,10 +342,10 @@ void FChaosVDGeometryBuilder::DeInitialize()
 	bInitialized = false;
 }
 
-void FChaosVDGeometryBuilder::CreateMeshesFromImplicitObject(const Chaos::FImplicitObject* InImplicitObject, TArray<TSharedPtr<FChaosVDExtractedGeometryDataHandle>>& OutMeshDataHandles,int32 AvailableShapeDataNum , const int32 DesiredLODCount, const Chaos::FRigidTransform3& InTransform, const int32 MeshIndex)
+void FChaosVDGeometryBuilder::CreateMeshesFromImplicitObject(const Chaos::FImplicitObject* InImplicitObject, TArray<TPair<TSharedPtr<FChaosVDExtractedGeometryDataHandle>, TSharedPtr<Chaos::VD::FChaosVDImplicitObjectView>>>& OutMeshDataHandles, Chaos::VD::FChaosVDImplicitObjectView* Root, int32 AvailableShapeDataNum , const int32 DesiredLODCount, const Chaos::FRigidTransform3& InTransform, const int32 MeshIndex)
 {
 	// To start set the leaf and the root to the same ptr. If the object is an union, in the subsequent recursive call the leaf will be set correctly
-	CreateMeshesFromImplicit_Internal(InImplicitObject, InImplicitObject, OutMeshDataHandles, DesiredLODCount, InTransform, MeshIndex, AvailableShapeDataNum);
+	CreateMeshesFromImplicit_Internal(InImplicitObject, InImplicitObject, OutMeshDataHandles, Root, DesiredLODCount, InTransform, MeshIndex, AvailableShapeDataNum);
 }
 
 void FChaosVDGeometryBuilder::AddReferencedObjects(FReferenceCollector& Collector)
@@ -399,13 +399,18 @@ bool FChaosVDGeometryBuilder::HasNegativeScale(const Chaos::FRigidTransform3& In
 	return ScaleSignVector.X * ScaleSignVector.Y * ScaleSignVector.Z < 0;
 }
 
-void FChaosVDGeometryBuilder::CreateMeshesFromImplicit_Internal(const Chaos::FImplicitObject* InRootImplicitObject, const Chaos::FImplicitObject* InLeafImplicitObject, TArray<TSharedPtr<FChaosVDExtractedGeometryDataHandle>>& OutMeshDataHandles, const int32 DesiredLODCount, const Chaos::FRigidTransform3& InTransform, const int32 ParentShapeInstanceIndex, int32 AvailableShapeDataNum)
+void FChaosVDGeometryBuilder::CreateMeshesFromImplicit_Internal(const Chaos::FImplicitObject* InRootImplicitObject, const Chaos::FImplicitObject* InLeafImplicitObject, TArray<TPair<TSharedPtr<FChaosVDExtractedGeometryDataHandle>, TSharedPtr<Chaos::VD::FChaosVDImplicitObjectView>>>& OutMeshDataHandles, Chaos::VD::FChaosVDImplicitObjectView* Root, const int32 DesiredLODCount, const Chaos::FRigidTransform3& InTransform, const int32 ParentShapeInstanceIndex, int32 AvailableShapeDataNum)
 {
 	TRACE_CPUPROFILER_EVENT_SCOPE(FChaosVDGeometryBuilder::CreateMeshesFromImplicit_Internal);
 
 	using namespace Chaos;
 
 	const EImplicitObjectType InnerType = GetInnerType(InLeafImplicitObject->GetType());
+
+	TSharedPtr<Chaos::VD::FChaosVDImplicitObjectView> ImplicitView = MakeShared<Chaos::VD::FChaosVDImplicitObjectView>();
+	ImplicitView->ImplicitObjectType = InnerType;
+
+	Root->AddChild(ImplicitView);
 	
 	if (InnerType == ImplicitObjectType::Union || InnerType == ImplicitObjectType::UnionClustered)
 	{
@@ -437,7 +442,7 @@ void FChaosVDGeometryBuilder::CreateMeshesFromImplicit_Internal(const Chaos::FIm
 					}
 				}
 
-				CreateMeshesFromImplicit_Internal(InRootImplicitObject, UnionImplicit.GetReference(), OutMeshDataHandles, DesiredLODCount, InTransform, CurrentShapeInstanceIndex, AvailableShapeDataNum);	
+				CreateMeshesFromImplicit_Internal(InRootImplicitObject, UnionImplicit.GetReference(), OutMeshDataHandles, ImplicitView.Get(), DesiredLODCount, InTransform, CurrentShapeInstanceIndex, AvailableShapeDataNum);
 			}
 		}
 
@@ -449,7 +454,7 @@ void FChaosVDGeometryBuilder::CreateMeshesFromImplicit_Internal(const Chaos::FIm
 		if (const TImplicitObjectTransformed<FReal, 3>* Transformed = InLeafImplicitObject->template GetObject<TImplicitObjectTransformed<FReal, 3>>())
 		{
 			// For transformed objects, the Instance index is the same so we pass it in without changing it
-			CreateMeshesFromImplicit_Internal(InRootImplicitObject, Transformed->GetTransformedObject(), OutMeshDataHandles, DesiredLODCount, Transformed->GetTransform(), ParentShapeInstanceIndex, AvailableShapeDataNum);
+			CreateMeshesFromImplicit_Internal(InRootImplicitObject, Transformed->GetTransformedObject(), OutMeshDataHandles, ImplicitView.Get(), DesiredLODCount, Transformed->GetTransform(), ParentShapeInstanceIndex, AvailableShapeDataNum);
 		}
 		
 		return;
@@ -461,7 +466,7 @@ void FChaosVDGeometryBuilder::CreateMeshesFromImplicit_Internal(const Chaos::FIm
 		MeshDataHandle->SetShapeInstanceIndex(ParentShapeInstanceIndex);
 		MeshDataHandle->SetRootImplicitObject(InRootImplicitObject);
 
-		OutMeshDataHandles.Add(MeshDataHandle);
+		OutMeshDataHandles.Add(TPair<TSharedPtr<FChaosVDExtractedGeometryDataHandle>, TSharedPtr<Chaos::VD::FChaosVDImplicitObjectView>>(MeshDataHandle, ImplicitView));
 	}
 }
 
diff --git a/Engine/Plugins/ChaosVD/Source/ChaosVD/Private/ChaosVDObjectDetailsTab.cpp b/Engine/Plugins/ChaosVD/Source/ChaosVD/Private/ChaosVDObjectDetailsTab.cpp
--- a/Engine/Plugins/ChaosVD/Source/ChaosVD/Private/ChaosVDObjectDetailsTab.cpp
+++ b/Engine/Plugins/ChaosVD/Source/ChaosVD/Private/ChaosVDObjectDetailsTab.cpp
@@ -46,18 +46,46 @@ TSharedRef<SDockTab> FChaosVDStandAloneObjectDetailsTab::HandleTabSpawnRequest(c
 	AddUnsupportedStruct(FChaosVDQueryDataWrapper::StaticStruct());
 	AddUnsupportedStruct(FChaosVDParticlePairMidPhase::StaticStruct());
 
+	float CurrentSlotSize = 0;
+
 	if (const TSharedPtr<SChaosVDMainTab> MainTabPtr = OwningTabWidget.Pin())
 	{
 		DetailsPanelTab->SetContent
 		(
-			SAssignNew(DetailsPanelView, SChaosVDDetailsView, MainTabPtr.ToSharedRef())
+			SNew(SVerticalBox)
+			+ SVerticalBox::Slot()
+			[
+				SAssignNew(DetailsSplitter, SSplitter)
+					.MinimumSlotHeight(60.0f)
+					.Orientation(Orient_Vertical)
+					.Style(FAppStyle::Get(), "SplitterDark")
+					.PhysicalSplitterHandleSize(2.0f)
+					+ SSplitter::Slot()
+					[
+						SAssignNew(DetailsPanelView, SChaosVDDetailsView, MainTabPtr.ToSharedRef())
+					]
+			]
 		);
+
+		DetailsSplitter->AddSlot(0)
+			.Value(.2f)
+			[
+				SAssignNew(GeometryTreeWidget, SChaosVDGeometryTree, GetChaosVDScene())
+			];
+		CurrentSlotSize = DetailsSplitter->SlotAt(0).GetSizeValue();
 	}
 	else
 	{
 		DetailsPanelTab->SetContent(GenerateErrorWidget());
 	}
 
+	if (CurrentSlotSize != 0)
+	{
+		SavedSplitSize = CurrentSlotSize;
+	}
+
+	UpdateTreeDetailsSplit(nullptr);
+
 	DetailsPanelTab->SetTabIcon(FChaosVDStyle::Get().GetBrush("TabIconDetailsPanel"));
 	
 	HandleTabSpawned(DetailsPanelTab);
@@ -69,9 +97,47 @@ void FChaosVDStandAloneObjectDetailsTab::HandleTabClosed(TSharedRef<SDockTab> In
 {
 	FChaosVDTabSpawnerBase::HandleTabClosed(InTabClosed);
 	
+	GeometryTreeWidget.Reset();
+	DetailsSplitter.Reset();
 	DetailsPanelView.Reset();
 }
 
+void FChaosVDStandAloneObjectDetailsTab::UpdateTreeDetailsSplit(const FChaosVDSceneParticle* Particle)
+{
+	if (!GeometryTreeWidget || !DetailsSplitter || DetailsSplitter->GetChildren()->NumSlot() < 2)
+	{
+		return;
+	}
+	GeometryTreeWidget->SetDataToInspect(Particle);
+	if (Particle == nullptr)
+	{
+		float CurrentSize = DetailsSplitter->SlotAt(0).GetSizeValue();
+		if (CurrentSize != 0)
+		{
+			SavedSplitSize = CurrentSize;
+		}
+		DetailsSplitter->SlotAt(0).SetSizeValue(0);
+		DetailsSplitter->SlotAt(0).SetResizable(false);
+		DetailsSplitter->SlotAt(1).SetResizable(false);
+		GeometryTreeWidget->SetVisibility(EVisibility::Hidden);
+	}
+	else
+	{
+		GeometryTreeWidget->SetVisibility(EVisibility::All);
+		DetailsSplitter->SlotAt(0).SetResizable(true);
+		DetailsSplitter->SlotAt(1).SetResizable(true);
+		float CurrentSize = DetailsSplitter->SlotAt(0).GetSizeValue();
+		if (CurrentSize != 0)
+		{
+			SavedSplitSize = CurrentSize;
+		}
+		if (SavedSplitSize != 0)
+		{
+			DetailsSplitter->SlotAt(0).SetSizeValue(SavedSplitSize);
+		}
+	}
+}
+
 void FChaosVDObjectDetailsTab::HandleActorsSelection(TArrayView<AActor*> SelectedActors)
 {
 	if (SelectedActors.Num() > 0)
@@ -90,6 +156,7 @@ void FChaosVDObjectDetailsTab::HandleActorsSelection(TArrayView<AActor*> Selecte
 	{
 		CurrentSelectedObject = nullptr;
 	}
+	UpdateTreeDetailsSplit(nullptr);
 }
 
 TSharedRef<SDockTab> FChaosVDObjectDetailsTab::HandleTabSpawnRequest(const FSpawnTabArgs& Args)
@@ -149,9 +216,15 @@ void FChaosVDObjectDetailsTab::HandlePostSelectionChange(const UTypedElementSele
 		using namespace Chaos::VD::TypedElementDataUtil;
 	
 		DetailsPanelView->SetSelectedStruct(GetStructOnScopeDataFromTypedElementHandle(SelectedParticlesHandles[0]));
+		
+		const FStructTypedElementData* StructElement = SelectedParticlesHandles[0].GetData<FStructTypedElementData>(true);
+		const FChaosVDSceneParticle* PartPtr = StructElement ? StructElement->GetData<FChaosVDSceneParticle>() : nullptr;
+		UpdateTreeDetailsSplit(PartPtr);
+		
 		return;
 	}
 
+	UpdateTreeDetailsSplit(nullptr);
 	DetailsPanelView->SetSelectedObject(nullptr);
 	DetailsPanelView->SetSelectedStruct(nullptr);
 }
@@ -162,6 +235,7 @@ void FChaosVDObjectDetailsTab::HandleSolverDataSelectionChange(const TSharedPtr<
 	if (!StructOnScope || !IsSupportedStruct(StructOnScope->GetStructPtr()))
 	{
 		DetailsPanelView->SetSelectedStruct(nullptr);
+		UpdateTreeDetailsSplit(nullptr);
 		return;
 	}
 
diff --git a/Engine/Plugins/ChaosVD/Source/ChaosVD/Private/ChaosVDSceneParticle.cpp b/Engine/Plugins/ChaosVD/Source/ChaosVD/Private/ChaosVDSceneParticle.cpp
--- a/Engine/Plugins/ChaosVD/Source/ChaosVD/Private/ChaosVDSceneParticle.cpp
+++ b/Engine/Plugins/ChaosVD/Source/ChaosVD/Private/ChaosVDSceneParticle.cpp
@@ -364,7 +364,7 @@ bool FChaosVDSceneParticle::IsSelected()
 	return false;
 }
 
-void FChaosVDSceneParticle::ProcessUpdatedAndRemovedHandles(TArray<TSharedPtr<FChaosVDExtractedGeometryDataHandle>>& OutExtractedGeometryDataHandles)
+void FChaosVDSceneParticle::ProcessUpdatedAndRemovedHandles(TArray<TPair<TSharedPtr<FChaosVDExtractedGeometryDataHandle>, TSharedPtr<Chaos::VD::FChaosVDImplicitObjectView>>>& OutExtractedGeometryDataHandles)
 {
 	for (TArray<TSharedRef<FChaosVDInstancedMeshData>>::TIterator MeshDataHandleRemoveIterator = MeshDataHandles.CreateIterator(); MeshDataHandleRemoveIterator; ++MeshDataHandleRemoveIterator)
 	{
@@ -374,9 +374,9 @@ void FChaosVDSceneParticle::ProcessUpdatedAndRemovedHandles(TArray<TSharedPtr<FC
 
 		// TODO: This search is n2, but I didn't see this as bottleneck. We should check if it is worth adding this to a TSet, or implementing the < operator so we can sort the array and do a binary search
 		// (avoiding the need to allocate a new container) 
-		for (TArray<TSharedPtr<FChaosVDExtractedGeometryDataHandle>>::TIterator HandleRemoveIterator = OutExtractedGeometryDataHandles.CreateIterator(); HandleRemoveIterator; ++HandleRemoveIterator)
+		for (TArray<TPair<TSharedPtr<FChaosVDExtractedGeometryDataHandle>, TSharedPtr<Chaos::VD::FChaosVDImplicitObjectView>>>::TIterator HandleRemoveIterator = OutExtractedGeometryDataHandles.CreateIterator(); HandleRemoveIterator; ++HandleRemoveIterator)
 		{
-			const TSharedPtr<FChaosVDExtractedGeometryDataHandle> GeometryDataHandle = *HandleRemoveIterator;
+			const TSharedPtr<FChaosVDExtractedGeometryDataHandle> GeometryDataHandle = HandleRemoveIterator->Key;
 			const TSharedRef<FChaosVDExtractedGeometryDataHandle> ExistingComponentGeometryDataHandle = ExistingMeshDataHandle->GetGeometryHandle();
 
 			if (GeometryDataHandle && *GeometryDataHandle == *ExistingMeshDataHandle->GetGeometryHandle())
@@ -386,6 +386,7 @@ void FChaosVDSceneParticle::ProcessUpdatedAndRemovedHandles(TArray<TSharedPtr<FC
 				// Although the geometry is the same, we need to copy over all the new data on the updated handle
 				// Otherwise the ptr to the root implicit object or the Shape Instance Index will be outdated 
 				*ExistingMeshDataHandle->GetGeometryHandle() = *GeometryDataHandle;
+				HandleRemoveIterator->Value->SetMeshData(ExistingMeshDataHandle);
 
 				// If we have a CVD Geometry Component for this handle, just remove it from the list as it means we don't need to re-create it
 				HandleRemoveIterator.RemoveCurrent();
@@ -459,7 +460,7 @@ void FChaosVDSceneParticle::UpdateGeometry(const Chaos::FConstImplicitObjectPtr&
 		return;
 	}
 	
-	TArray<TSharedPtr<FChaosVDExtractedGeometryDataHandle>> OutExtractedGeometryDataHandles;
+	TArray<TPair<TSharedPtr<FChaosVDExtractedGeometryDataHandle>, TSharedPtr<Chaos::VD::FChaosVDImplicitObjectView>>> OutExtractedGeometryDataHandles;
 	OutExtractedGeometryDataHandles.Reserve(ObjectsToGenerateNum);
 
 	bool bHasToUseStaticMeshComponent = Chaos::VisualDebugger::Cvars::bForceStaticMeshComponentUse;
@@ -473,7 +474,12 @@ void FChaosVDSceneParticle::UpdateGeometry(const Chaos::FConstImplicitObjectPtr&
 	constexpr int32 LODsToGenerateNum = 3;
 	constexpr int32 LODsToGenerateNumForInstancedStaticMesh = 0;
 
-	GeometryGenerator->CreateMeshesFromImplicitObject(InImplicitObject, OutExtractedGeometryDataHandles, ParticleDataPtr->CollisionDataPerShape.Num(), bHasToUseStaticMeshComponent ? LODsToGenerateNum : LODsToGenerateNumForInstancedStaticMesh);
+	RootImplicitView = TSharedPtr<Chaos::VD::FChaosVDImplicitObjectView>();
+	
+	TSharedPtr<Chaos::VD::FChaosVDImplicitObjectView> ImplicitRoot = MakeShared<Chaos::VD::FChaosVDImplicitObjectView>();
+	ImplicitRoot->ImplicitObjectType = InImplicitObject->GetType();
+
+	GeometryGenerator->CreateMeshesFromImplicitObject(InImplicitObject, OutExtractedGeometryDataHandles, ImplicitRoot.Get(), ParticleDataPtr->CollisionDataPerShape.Num(), bHasToUseStaticMeshComponent ? LODsToGenerateNum : LODsToGenerateNumForInstancedStaticMesh);
 
 	// This should not happen in theory, but there might be some valid situations where it does. Adding an ensure to catch them and then evaluate if it is really an issue (if it is not I will remove the ensure later on). 
 	if (!ensure(ObjectsToGenerateNum == OutExtractedGeometryDataHandles.Num()))
@@ -486,8 +492,9 @@ void FChaosVDSceneParticle::UpdateGeometry(const Chaos::FConstImplicitObjectPtr&
 
 	if (OutExtractedGeometryDataHandles.Num() > 0)
 	{
-		for (const TSharedPtr<FChaosVDExtractedGeometryDataHandle>& ExtractedGeometryDataHandle : OutExtractedGeometryDataHandles)
+		for (const TPair<TSharedPtr<FChaosVDExtractedGeometryDataHandle>, TSharedPtr<Chaos::VD::FChaosVDImplicitObjectView>>& ExtractedPair : OutExtractedGeometryDataHandles)
 		{
+			TSharedPtr<FChaosVDExtractedGeometryDataHandle> ExtractedGeometryDataHandle = ExtractedPair.Key;
 			//TODO: Time Slice component creation
 			TSharedPtr<FChaosVDInstancedMeshData> MeshDataInstance;
 
@@ -532,6 +539,8 @@ void FChaosVDSceneParticle::UpdateGeometry(const Chaos::FConstImplicitObjectPtr&
 				MeshDataInstance->SetWorldTransform(ParticleTransform);
 			}
 
+			ExtractedPair.Value->SetMeshData(MeshDataInstance.ToSharedRef());
+
 			MeshDataHandles.Add(MeshDataInstance.ToSharedRef());
 		}
 
@@ -539,9 +548,11 @@ void FChaosVDSceneParticle::UpdateGeometry(const Chaos::FConstImplicitObjectPtr&
 		EnumAddFlags(DirtyFlags, EChaosVDSceneParticleDirtyFlags::Visibility);
 		EnumAddFlags(DirtyFlags, EChaosVDSceneParticleDirtyFlags::Coloring);
 		EnumRemoveFlags(DirtyFlags, EChaosVDSceneParticleDirtyFlags::Geometry);
-
+		
 		bIsGeometryDataGenerationStarted = true;
 	}
+
+	RootImplicitView = ImplicitRoot;
 }
 
 void FChaosVDSceneParticle::CalculateAndCacheBounds() const
diff --git a/Engine/Plugins/ChaosVD/Source/ChaosVD/Private/Widgets/SChaosVDGeometryTree.cpp b/Engine/Plugins/ChaosVD/Source/ChaosVD/Private/Widgets/SChaosVDGeometryTree.cpp
new file mode 100644
--- /dev/null
+++ b/Engine/Plugins/ChaosVD/Source/ChaosVD/Private/Widgets/SChaosVDGeometryTree.cpp
@@ -0,0 +1,348 @@
+// Copyright Epic Games, Inc. All Rights Reserved.
+
+#include "SChaosVDGeometryTree.h"
+
+#include "SChaosVDGeometryTreeRow.h"
+#include "Widgets/Images/SImage.h"
+#include "Widgets/Text/STextBlock.h"
+#include "Actors/ChaosVDSolverInfoActor.h"
+
+#define LOCTEXT_NAMESPACE "ChaosVisualDebugger"
+
+SChaosVDGeometryTree::FColumnNames SChaosVDGeometryTree::ColumnNames = SChaosVDGeometryTree::FColumnNames();
+
+void SChaosVDGeometryTree::SetDataToInspect(const FChaosVDSceneParticle* Particle)
+{
+	bool bChangedParticle = SelectedParticle != Particle || !Particle;
+	SelectedParticle = Particle;
+	RefreshTreeData(bChangedParticle);
+}
+
+void SChaosVDGeometryTree::RefreshTreeData(bool bRegenTree)
+{
+	if (bRegenTree)
+	{
+		InternalGeometryTreeData.Reset();
+
+		if (SelectedParticle)
+		{
+			if (const TSharedPtr<FChaosVDScene> ScenePtr = SceneWeakPtr.Pin())
+			{
+				TSharedPtr<Chaos::VD::FChaosVDImplicitObjectView> Root = SelectedParticle->GetRootImplicitView().Pin();
+
+				TSharedPtr<const FChaosVDParticleDataWrapper> ParticleData = SelectedParticle->GetParticleData();
+				Chaos::FConstImplicitObjectPtr RootGeometry = ScenePtr->GetUpdatedGeometry(ParticleData->GeometryHash);
+
+				TSharedPtr<FChaosVDGeometryTreeItem> NewTreeItem = MakeShared<FChaosVDGeometryTreeItem>();
+				NewTreeItem->Name = FText::FromString(SelectedParticle->GetDisplayName());
+				NewTreeItem->Type = Chaos::ImplicitObjectType::Unknown;
+
+				//Implicit Views will not be generated if geometry is not currently streamed in. Traverse ImplicitObjects from scratch as a backup.
+				if (Root) {
+					GenerateTree(Root.Get(), NewTreeItem->Children);
+				}
+				else if (RootGeometry) {
+					GenerateTree(RootGeometry, NewTreeItem->Children);
+				}
+
+				InternalGeometryTreeData.Add(NewTreeItem);
+			}
+		}
+
+		GeometryTreeWidget->SetTreeItemsSource(&InternalGeometryTreeData);
+		ExpandAll(InternalGeometryTreeData);
+
+		if (SelectedParticle && !SelectedParticle->GetSelectedMeshInstance().IsValid() && InternalGeometryTreeData.Num() > 0)
+		{
+			GeometryTreeWidget->SetSelection(InternalGeometryTreeData[0]);
+		}
+	}
+
+	GeometryTreeWidget->RebuildList();
+	GeometryTreeWidget->RequestTreeRefresh();
+}
+
+void SChaosVDGeometryTree::Construct(const FArguments& InArgs, const TWeakPtr<FChaosVDScene>& InScenePtr)
+{
+	SceneWeakPtr = InScenePtr;
+	GeometryItemSelectedDelegate = InArgs._OnItemSelected;
+	GeometryItemFocusedDelegate = InArgs._OnItemFocused;
+
+	constexpr float BottomPadding = 2.0f;
+	constexpr float NoPadding = 0.0f;
+	constexpr float VerticalPadding = 0.0f;
+
+	static FMargin ColumnHeaderTextMargin(NoPadding, VerticalPadding);
+
+	ChildSlot
+		[
+			SNew(SVerticalBox)
+				+ SVerticalBox::Slot()
+				.Padding(NoPadding, NoPadding, NoPadding, BottomPadding)
+				[
+					SAssignNew(GeometryTreeWidget, STreeView<TSharedPtr<FChaosVDGeometryTreeItem>>)
+						.OnGenerateRow(this, &SChaosVDGeometryTree::GenerateGeometryDataRow)
+						.OnSelectionChanged(this, &SChaosVDGeometryTree::GeometryTreeSelectionChanged)
+						.OnGetChildren(this, &SChaosVDGeometryTree::OnGetChildrenForGeometryItem)
+						.TreeItemsSource(&InternalGeometryTreeData)
+						.OnMouseButtonDoubleClick(this, &SChaosVDGeometryTree::HandleFocusRequest)
+						.SelectionMode(ESelectionMode::Type::Single)
+						.HighlightParentNodesForSelection(true)
+						.HeaderRow(
+							SNew(SHeaderRow).Style(&FAppStyle::Get().GetWidgetStyle<FHeaderRowStyle>("PropertyTable.HeaderRow"))
+							+ SHeaderRow::Column(ColumnNames.Type)
+							.SortMode(EColumnSortMode::None)
+							.FillWidth(2)
+							.HAlignCell(HAlign_Left)
+							[
+								SNew(STextBlock)
+									.Margin(ColumnHeaderTextMargin)
+									.Text(LOCTEXT("ImplicitObjectTypeHeader", "Type"))
+							]
+							+ SHeaderRow::Column(ColumnNames.TraceType)
+							.SortMode(EColumnSortMode::None)
+							.FillWidth(1.5)
+							.HAlignCell(HAlign_Left)
+							[
+								SNew(STextBlock)
+									.Margin(ColumnHeaderTextMargin)
+									.Text(LOCTEXT("ImplicitObjectTraceTypeHeader", "Collision Trace Type"))
+							]
+							+ SHeaderRow::Column(ColumnNames.CollisionEnabled)
+							.SortMode(EColumnSortMode::None)
+							.FillWidth(1.5)
+							.HAlignCell(HAlign_Left)
+							[
+								SNew(STextBlock)
+									.Margin(ColumnHeaderTextMargin)
+									.Text(LOCTEXT("ImplicitObjectComplexityCollisionEnabled", "Collision Enabled"))
+							]
+						)
+				]
+		];
+}
+
+TSharedRef<ITableRow> SChaosVDGeometryTree::GenerateGeometryDataRow(TSharedPtr<FChaosVDGeometryTreeItem> GeometryData, const TSharedRef<STableViewBase>& OwnerTable)
+{
+	if (!GeometryData)
+	{
+		return
+			SNew(STableRow< TSharedPtr<FString> >, OwnerTable)
+			[
+				SNew(SVerticalBox)
+					+ SVerticalBox::Slot()
+					[
+						SNew(STextBlock)
+							.Text(LOCTEXT("SChaosVDGeometryHierarchyErrorMessage", "Failed to read data for solver."))
+					]
+			];
+	}
+
+	return SNew(SChaosVDGeometryTreeRow, OwnerTable)
+		.Item(GeometryData);
+}
+
+void SChaosVDGeometryTree::GeometryTreeSelectionChanged(TSharedPtr<FChaosVDGeometryTreeItem> SelectedItem, ESelectInfo::Type Type)
+{
+	SelectItem(SelectedItem, Type);
+	GeometryItemSelectedDelegate.ExecuteIfBound(SelectedItem, Type);
+}
+
+void SChaosVDGeometryTree::OnGetChildrenForGeometryItem(TSharedPtr<FChaosVDGeometryTreeItem> GeomEntry, TArray<TSharedPtr<FChaosVDGeometryTreeItem>>& OutGeoms)
+{
+	OutGeoms.Append(GeomEntry->Children);
+}
+
+void SChaosVDGeometryTree::ExpandAll(const TArray<TSharedPtr<FChaosVDGeometryTreeItem>>& Items)
+{
+	for (const TSharedPtr<FChaosVDGeometryTreeItem>& Item : Items)
+	{
+		if (SelectedParticle && SelectedParticle->GetSelectedMeshInstance().IsValid())
+		{
+			if (TSharedPtr<Chaos::VD::FChaosVDImplicitObjectView> View = Item->ItemWeakPtr.Pin())
+			{
+				if(View->GetMeshData() == SelectedParticle->GetSelectedMeshInstance().Pin())
+				{
+					GeometryTreeWidget->SetSelection(Item);
+				}
+			}
+		}
+
+		GeometryTreeWidget->SetItemExpansion(Item, true);
+		ExpandAll(Item->Children);
+	}
+}
+
+void SChaosVDGeometryTree::SelectItem(const TSharedPtr<FChaosVDGeometryTreeItem>& ItemToSelect, ESelectInfo::Type Type)
+{
+	if (!ItemToSelect.IsValid() || Type == ESelectInfo::Direct)
+	{
+		return;
+	}
+
+	if (const TSharedPtr<FChaosVDScene> ScenePtr = SceneWeakPtr.Pin())
+	{
+		TSharedPtr<FChaosVDSceneParticle> ParticleInstance = ScenePtr->GetParticleInstance(SelectedParticle->GetParticleData()->SolverID, SelectedParticle->GetParticleData()->ParticleIndex);
+		
+		if(ParticleInstance)
+		{
+			if (InternalGeometryTreeData.Num() > 0 && InternalGeometryTreeData[0] == ItemToSelect)
+			{
+				Chaos::VisualDebugger::SelectParticleWithGeometryInstance(ScenePtr.ToSharedRef(), ParticleInstance.Get(), nullptr);
+			}
+			else
+			{
+				if (TSharedPtr<Chaos::VD::FChaosVDImplicitObjectView> View = ItemToSelect->ItemWeakPtr.Pin())
+				{
+					if (const TSharedPtr<FChaosVDInstancedMeshData> MeshDataHandle = View->GetMeshData())
+					{
+						Chaos::VisualDebugger::SelectParticleWithGeometryInstance(ScenePtr.ToSharedRef(), ParticleInstance.Get(), MeshDataHandle);
+					}
+				}
+			}
+		}
+	}
+}
+
+void SChaosVDGeometryTree::Tick(const FGeometry& AllottedGeometry, const double InCurrentTime, const float InDeltaTime)
+{
+	SCompoundWidget::Tick(AllottedGeometry, InCurrentTime, InDeltaTime);
+
+	if (SelectedParticle)
+	{
+		FChaosVDBaseSceneObject::EStreamingState NewParticleStreamingState = SelectedParticle->GetStreamingState();
+		
+		// If the particle data has just been streamed in, force refresh the tree.
+		if (NewParticleStreamingState == FChaosVDBaseSceneObject::EStreamingState::Visible && ParticleStreamingState == FChaosVDBaseSceneObject::EStreamingState::Hidden)
+		{
+			RefreshTreeData(true);
+		}
+
+		ParticleStreamingState = NewParticleStreamingState;
+	}
+}
+
+void SChaosVDGeometryTree::HandleFocusRequest(TSharedPtr<FChaosVDGeometryTreeItem> InFocusedItem)
+{
+	GeometryItemFocusedDelegate.ExecuteIfBound(InFocusedItem);
+	SelectItem(InFocusedItem, ESelectInfo::Type::Direct);
+
+	const TSharedPtr<FChaosVDScene> ScenePtr = SceneWeakPtr.Pin();
+	if (ScenePtr && SelectedParticle)
+	{
+		ScenePtr->OnFocusRequest().Broadcast(SelectedParticle->GetBoundingBox());
+	}
+}
+
+void SChaosVDGeometryTree::GenerateTree(Chaos::VD::FChaosVDImplicitObjectView* InLeafImplicitObject, TArray<TSharedPtr<FChaosVDGeometryTreeItem>>& Nodes)
+{
+	using namespace Chaos;
+
+	if (!InLeafImplicitObject)
+	{
+		return;
+	}
+
+	TSharedPtr<FChaosVDGeometryTreeItem> NewTreeItem = MakeShared<FChaosVDGeometryTreeItem>();
+	NewTreeItem->Name = FText::FromName(FImplicitObject::GetTypeName(InLeafImplicitObject->ImplicitObjectType));
+	NewTreeItem->Type = InLeafImplicitObject->ImplicitObjectType;
+
+	NewTreeItem->ItemWeakPtr = InLeafImplicitObject->AsWeak();
+
+	Nodes.Add(NewTreeItem);
+
+	TArray<TSharedPtr<Chaos::VD::FChaosVDImplicitObjectView>>& Children = *InLeafImplicitObject->GetChildren();
+
+	if (NewTreeItem->Type == Chaos::ImplicitObjectType::Transformed)
+	{
+		if (Children.Num() == 1 && !(Children[0]->ImplicitObjectType == Chaos::ImplicitObjectType::Transformed))
+		{
+			NewTreeItem->Name = FText::FromString(FString(NewTreeItem->Name.ToString()).Append(" ").Append(FImplicitObject::GetTypeName(Children[0]->ImplicitObjectType).ToString()));
+		}
+	}
+
+	if (InLeafImplicitObject->GetMeshData())
+	{ 
+		FChaosVDShapeCollisionData& Data = InLeafImplicitObject->GetMeshData()->GetGeometryCollisionData();
+		NewTreeItem->CollisionEnabled = UEnum::GetDisplayValueAsText(CollisionEnabledFromFlags(Data.bQueryCollision, Data.bSimCollision, Data.bIsProbe));
+		NewTreeItem->TraceType =  UEnum::GetDisplayValueAsText(Data.CollisionTraceType);
+	}
+
+	for (TSharedPtr<Chaos::VD::FChaosVDImplicitObjectView> Child : Children)
+	{
+		GenerateTree(Child.Get(), NewTreeItem->Children);
+	}
+}
+
+void SChaosVDGeometryTree::GenerateTree(const Chaos::FImplicitObject* InLeafImplicitObject, TArray<TSharedPtr<FChaosVDGeometryTreeItem>>& Nodes, int ShapeIndex)
+{
+	using namespace Chaos;
+
+	if (!InLeafImplicitObject || !SelectedParticle)
+	{
+		return;
+	}
+
+	TSharedPtr<FChaosVDGeometryTreeItem> NewTreeItem = MakeShared<FChaosVDGeometryTreeItem>();
+	NewTreeItem->Name = FText::FromName(InLeafImplicitObject->GetTypeName());
+	NewTreeItem->Type = GetInnerType(InLeafImplicitObject->GetType());
+
+	Nodes.Add(NewTreeItem);
+
+	const EImplicitObjectType InnerType = GetInnerType(InLeafImplicitObject->GetType());
+
+	if (InnerType == ImplicitObjectType::Union || InnerType == ImplicitObjectType::UnionClustered)
+	{
+		if (const FImplicitObjectUnion* Union = InLeafImplicitObject->template AsA<FImplicitObjectUnion>())
+		{
+			const bool bIsRootUnion = ShapeIndex == 0;
+			const bool bIsCluster = InnerType == ImplicitObjectType::UnionClustered;
+
+			for (int32 ObjectIndex = 0; ObjectIndex < Union->GetObjects().Num(); ++ObjectIndex)
+			{
+				const FImplicitObjectPtr UnionImplicit = Union->GetObjects()[ObjectIndex];
+				
+				if (bIsRootUnion)
+				{
+					if (bIsCluster)
+					{
+						// Geometry Collections might break the usual rule of how may shape data instances we have per geometry
+						// Sometimes they can create clusters where all particles share a single instance
+						if (SelectedParticle->GetParticleData() && SelectedParticle->GetParticleData()->CollisionDataPerShape.Num() == 1)
+						{
+							ShapeIndex = 0;
+						}
+					}
+					else
+					{
+						ShapeIndex = ObjectIndex;
+					}
+				}
+
+				GenerateTree(UnionImplicit, NewTreeItem->Children, ShapeIndex);
+			}
+		}
+	}
+	else if (InnerType == ImplicitObjectType::Transformed)
+	{
+		if (const TImplicitObjectTransformed<FReal, 3>*Transformed = InLeafImplicitObject->template GetObject<TImplicitObjectTransformed<FReal, 3>>())
+		{
+			if (Transformed->GetTransformedObject()) {
+				NewTreeItem->Name = FText::FromString(FString(NewTreeItem->Name.ToString()).Append(" ").Append(Transformed->GetTransformedObject()->GetTypeName().ToString()));
+			}
+			// For transformed objects, the Instance index is the same so we pass it in without changing it
+			GenerateTree(Transformed->GetTransformedObject(), NewTreeItem->Children, ShapeIndex);
+		}
+	}
+	else
+	{
+		if (SelectedParticle->GetParticleData() && SelectedParticle->GetParticleData()->CollisionDataPerShape.IsValidIndex(ShapeIndex))
+		{
+			const FChaosVDShapeCollisionData& Data = SelectedParticle->GetParticleData()->CollisionDataPerShape[ShapeIndex];
+			NewTreeItem->TraceType = UEnum::GetDisplayValueAsText(Data.CollisionTraceType);
+			NewTreeItem->CollisionEnabled = UEnum::GetDisplayValueAsText(CollisionEnabledFromFlags(Data.bQueryCollision, Data.bSimCollision, Data.bIsProbe));
+		}
+	}
+}
+
+#undef LOCTEXT_NAMESPACE
diff --git a/Engine/Plugins/ChaosVD/Source/ChaosVD/Private/Widgets/SChaosVDGeometryTree.h b/Engine/Plugins/ChaosVD/Source/ChaosVD/Private/Widgets/SChaosVDGeometryTree.h
new file mode 100644
--- /dev/null
+++ b/Engine/Plugins/ChaosVD/Source/ChaosVD/Private/Widgets/SChaosVDGeometryTree.h
@@ -0,0 +1,94 @@
+// Copyright Epic Games, Inc. All Rights Reserved.
+
+#pragma once
+
+#include "Containers/Array.h"
+#include "Templates/SharedPointer.h"
+#include "Widgets/SCompoundWidget.h"
+#include "Widgets/Views/STableViewBase.h"
+#include "Widgets/Views/STreeView.h"
+#include "ChaosVDScene.h"
+#include "ChaosVDGeometryDataComponent.h"
+#include "ChaosVDBaseSceneObject.h"
+#include "Chaos/ImplicitObjectScaled.h"
+#include "Chaos/ImplicitObjectType.h"
+
+class ITableRow;
+class FChaosVDScene;
+struct FChaosVDSceneParticle;
+struct FChaosVDBaseSceneObject;
+
+struct FChaosVDGeometryTreeItem
+{
+	TWeakPtr<Chaos::VD::FChaosVDImplicitObjectView> ItemWeakPtr;
+	TArray<TSharedPtr<FChaosVDGeometryTreeItem>> Children;
+	FText Name = FText::GetEmpty();
+	Chaos::EImplicitObjectType Type;
+	FText CollisionEnabled = FText::GetEmpty();
+	FText TraceType = FText::GetEmpty();
+};
+
+DECLARE_DELEGATE_TwoParams(FChaosVDGeometryTreeItemSelected, const TSharedPtr<FChaosVDGeometryTreeItem>&, ESelectInfo::Type)
+DECLARE_DELEGATE_OneParam(FChaosVDGeometryTreeItemFocused, const TSharedPtr<FChaosVDGeometryTreeItem>&)
+
+/**
+ * Tree for geometry hierarchy
+ */
+class SChaosVDGeometryTree : public SCompoundWidget
+{
+public:
+
+	SLATE_BEGIN_ARGS(SChaosVDGeometryTree) {}
+		SLATE_EVENT(FChaosVDGeometryTreeItemSelected, OnItemSelected)
+		SLATE_EVENT(FChaosVDGeometryTreeItemFocused, OnItemFocused)
+	SLATE_END_ARGS()
+
+	void Construct(const FArguments& InArgs, const TWeakPtr<FChaosVDScene>& InScenePtr);
+
+	void SelectItem(const TSharedPtr<FChaosVDGeometryTreeItem>& ItemToSelect, ESelectInfo::Type Type);
+
+	void ExpandAll(const TArray<TSharedPtr<FChaosVDGeometryTreeItem>>& Items);
+
+	void SetDataToInspect(const FChaosVDSceneParticle* Particle);
+
+	void RefreshTreeData(bool bRegenTree);
+
+	virtual void Tick(const FGeometry& AllottedGeometry, const double InCurrentTime, const float InDeltaTime) override;
+
+
+protected:
+	void HandleFocusRequest(TSharedPtr<FChaosVDGeometryTreeItem> InFocusedItem);
+
+	TSharedRef<ITableRow> GenerateGeometryDataRow(TSharedPtr<FChaosVDGeometryTreeItem> GeometryData, const TSharedRef<STableViewBase>& OwnerTable);
+
+	void GeometryTreeSelectionChanged(TSharedPtr<FChaosVDGeometryTreeItem> SelectedItem, ESelectInfo::Type Type);
+
+	void OnGetChildrenForGeometryItem(TSharedPtr<FChaosVDGeometryTreeItem> GeomEntry, TArray<TSharedPtr<FChaosVDGeometryTreeItem>>& OutGeoms);
+
+	void GenerateTree(Chaos::VD::FChaosVDImplicitObjectView* InLeafImplicitObject, TArray<TSharedPtr<FChaosVDGeometryTreeItem>>& Nodes);
+
+	void GenerateTree(const Chaos::FImplicitObject* InLeafImplicitObject, TArray<TSharedPtr<FChaosVDGeometryTreeItem>>& Nodes, int ShapeIndex = 0);
+
+	TSharedPtr<STreeView<TSharedPtr<FChaosVDGeometryTreeItem>>> GeometryTreeWidget;
+
+	TArray<TSharedPtr<FChaosVDGeometryTreeItem>> InternalGeometryTreeData;
+
+	FChaosVDGeometryTreeItemSelected GeometryItemSelectedDelegate;
+	FChaosVDGeometryTreeItemFocused GeometryItemFocusedDelegate;
+
+	TWeakPtr<FChaosVDScene> SceneWeakPtr;
+
+	const FChaosVDSceneParticle* SelectedParticle = nullptr;
+	FChaosVDBaseSceneObject::EStreamingState ParticleStreamingState = FChaosVDBaseSceneObject::EStreamingState::Hidden;
+
+public:
+
+	struct FColumnNames
+	{
+		const FName Type = FName("Type");
+		const FName TraceType = FName("Collision Trace Type");
+		const FName CollisionEnabled = FName("Collision Enabled");
+	};
+
+	static FColumnNames ColumnNames;
+};
diff --git a/Engine/Plugins/ChaosVD/Source/ChaosVD/Private/Widgets/SChaosVDGeometryTreeRow.cpp b/Engine/Plugins/ChaosVD/Source/ChaosVD/Private/Widgets/SChaosVDGeometryTreeRow.cpp
new file mode 100644
--- /dev/null
+++ b/Engine/Plugins/ChaosVD/Source/ChaosVD/Private/Widgets/SChaosVDGeometryTreeRow.cpp
@@ -0,0 +1,140 @@
+// Copyright Epic Games, Inc. All Rights Reserved.
+
+#include "SChaosVDGeometryTreeRow.h"
+
+#include "SChaosVDGeometryTree.h"
+#include "ChaosVDStyle.h"
+#include "SlateOptMacros.h"
+#include "Widgets/Images/SImage.h"
+#include "Widgets/Layout/SBox.h"
+#include "Widgets/SBoxPanel.h"
+#include "Widgets/Text/STextBlock.h"
+
+#define LOCTEXT_NAMESPACE "ChaosVisualDebugger"
+
+void SChaosVDGeometryTreeRow::Construct(const FArguments& InArgs, const TSharedRef<STableViewBase>& InOwnerTableView)
+{
+	Item = InArgs._Item;
+	FSuperRowType::FArguments Args = FSuperRowType::FArguments()
+									.Style(&FAppStyle::Get().GetWidgetStyle<FTableRowStyle>("SceneOutliner.TableViewRow"));
+
+	SMultiColumnTableRow<TSharedPtr<const FChaosVDGeometryTreeItem>>::Construct(Args, InOwnerTableView);
+}
+
+TSharedRef<SWidget> SChaosVDGeometryTreeRow::GenerateWidgetForColumn(const FName& ColumnName)
+{
+	if (!Item)
+	{
+		return SNullWidget::NullWidget;
+	}
+
+	if (ColumnName == SChaosVDGeometryTree::ColumnNames.Type)
+	{
+		constexpr float NoPadding = 0.0f;
+		constexpr float ExpanderLeftPadding = 6.0f;
+		constexpr float ExpanderIndentAmount = 12.0f;
+		return SNew(SBox)
+			[
+				SNew(SHorizontalBox)
+				+SHorizontalBox::Slot()
+				.AutoWidth()
+				.Padding(ExpanderLeftPadding, NoPadding, NoPadding, NoPadding)
+				[
+					SNew(SExpanderArrow, SharedThis(this)).IndentAmount(ExpanderIndentAmount)
+				]
+					+ SHorizontalBox::Slot()
+					.AutoWidth()
+					[
+						SNew(SImage)
+							.IsEnabled(false)
+							.ColorAndOpacity(FSlateColor::UseForeground())
+							.Image(this, &SChaosVDGeometryTreeRow::GetTypeIcon)
+					]
+				+SHorizontalBox::Slot()
+				.FillWidth(1.0f)
+				[
+					GenerateTextWidgetFromText(Item->Name)
+				]
+			];
+	}
+
+	if (ColumnName == SChaosVDGeometryTree::ColumnNames.TraceType)
+	{
+		return GenerateTextWidgetFromText(Item->TraceType);
+	}
+
+	if (ColumnName == SChaosVDGeometryTree::ColumnNames.CollisionEnabled)
+	{
+		return GenerateTextWidgetFromText(Item->CollisionEnabled);
+	}
+
+	return SNullWidget::NullWidget;
+}
+
+const FSlateBrush* SChaosVDGeometryTreeRow::GetTypeIcon() const
+{
+	if (!Item)
+	{
+		return nullptr;
+	}
+
+	Chaos::EImplicitObjectType Type = Item->Type;
+
+	if (Item->Type == Chaos::ImplicitObjectType::Transformed && Item->Children.Num() == 1)
+	{
+		Type = Item->Children[0]->Type;
+	}
+
+	if (Type == Chaos::ImplicitObjectType::Box)
+	{
+		return FAppStyle::Get().GetBrush(TEXT("ClassIcon.Cube"));
+	}
+	else if (Type == Chaos::ImplicitObjectType::Sphere)
+	{
+		return FAppStyle::Get().GetBrush(TEXT("ClassIcon.Sphere"));
+	}
+	else if (Type == Chaos::ImplicitObjectType::TriangleMesh)
+	{
+		return FAppStyle::Get().GetBrush(TEXT("ClassIcon.StaticMesh"));
+	}
+	else if (Type == Chaos::ImplicitObjectType::HeightField)
+	{
+		return FAppStyle::Get().GetBrush(TEXT("ClassIcon.Landscape"));
+	}
+	else if (Type == Chaos::ImplicitObjectType::Capsule)
+	{
+		return FAppStyle::Get().GetBrush(TEXT("ClassIcon.TriggerCapsule"));
+	}
+	else if (Type == Chaos::ImplicitObjectType::Convex)
+	{
+		return FAppStyle::Get().GetBrush(TEXT("ClassIcon.Cone"));
+	}
+	else if (Type == Chaos::ImplicitObjectType::Union)
+	{
+		return FAppStyle::Get().GetBrush(TEXT("ClassIcon.GeometryCollection"));
+	}
+	else if (Type == Chaos::ImplicitObjectType::Unknown)
+	{
+		return FChaosVDStyle::Get().GetBrush(TEXT("RigidBodyIcon"));
+	}
+
+	return FAppStyle::Get().GetBrush(TEXT("ClassIcon.Sphere"));
+}
+
+TSharedRef<SWidget> SChaosVDGeometryTreeRow::GenerateTextWidgetFromName(FName Name)
+{
+	return GenerateTextWidgetFromText(FText::FromName(Name));
+}
+
+TSharedRef<SWidget> SChaosVDGeometryTreeRow::GenerateTextWidgetFromText(const FText& Text)
+{
+	constexpr float MarginLeft = 4.0f;
+	constexpr float NoMargin = 0.0f;
+
+	return SNew(STextBlock)
+		.Margin(FMargin(MarginLeft, NoMargin, NoMargin, NoMargin))
+		.Text(Text)
+		.OverflowPolicy(ETextOverflowPolicy::MiddleEllipsis);
+}
+
+#undef LOCTEXT_NAMESPACE
\ No newline at end of file
diff --git a/Engine/Plugins/ChaosVD/Source/ChaosVD/Private/Widgets/SChaosVDGeometryTreeRow.h b/Engine/Plugins/ChaosVD/Source/ChaosVD/Private/Widgets/SChaosVDGeometryTreeRow.h
new file mode 100644
--- /dev/null
+++ b/Engine/Plugins/ChaosVD/Source/ChaosVD/Private/Widgets/SChaosVDGeometryTreeRow.h
@@ -0,0 +1,38 @@
+// Copyright Epic Games, Inc. All Rights Reserved.
+
+#pragma once
+
+#include "Widgets/SCompoundWidget.h"
+#include "Widgets/Views/SListView.h"
+#include "Widgets/Views/STableRow.h"
+
+struct FChaosVDQueryDataWrapper;
+struct FChaosVDGeometryTreeItem;
+class FChaosVDImplicitObjectView;
+
+/**
+ * Widget used to represent a row on the geometry hierarchy tree view
+ */
+class SChaosVDGeometryTreeRow : public SMultiColumnTableRow<TSharedPtr<const FChaosVDGeometryTreeItem>>
+{
+public:
+	SLATE_BEGIN_ARGS(SChaosVDGeometryTreeRow)
+		{
+		}
+		SLATE_ARGUMENT(TSharedPtr<FChaosVDGeometryTreeItem>, Item)
+
+	SLATE_END_ARGS()
+
+	void Construct(const FArguments& InArgs, const TSharedRef<STableViewBase>& InOwnerTableView);
+
+	virtual TSharedRef<SWidget> GenerateWidgetForColumn(const FName& ColumnName) override;
+
+private:
+
+	TSharedRef<SWidget> GenerateTextWidgetFromName(FName Name);
+	TSharedRef<SWidget> GenerateTextWidgetFromText(const FText& Text);
+
+	const FSlateBrush* GetTypeIcon() const;
+
+	TSharedPtr<FChaosVDGeometryTreeItem> Item;
+};
diff --git a/Engine/Plugins/ChaosVD/Source/ChaosVD/Public/ChaosVDGeometryBuilder.h b/Engine/Plugins/ChaosVD/Source/ChaosVD/Public/ChaosVDGeometryBuilder.h
--- a/Engine/Plugins/ChaosVD/Source/ChaosVD/Public/ChaosVDGeometryBuilder.h
+++ b/Engine/Plugins/ChaosVD/Source/ChaosVD/Public/ChaosVDGeometryBuilder.h
@@ -70,13 +70,13 @@ public:
 
 	/** Creates Dynamic Mesh components for each object within the provided Implicit object
 	 *	@param InImplicitObject : Implicit object to process
-	 *	@param Owner Actor who will own the generated components
-	 *	@param OutMeshDataHandles Array containing all the generated components
+	 *	@param OutMeshDataHandles Array containing all the generated components paired with their corresponding implicit objects
+	 *  @param Root ImplicitObjectView to attach children to
 	 *	@param DesiredLODCount Number of LODs to generate for the mesh.
 	 *	@param MeshIndex Index of the current component being processed. This is useful when this method is called recursively
 	 *	@param InTransform to apply to the generated components/geometry
 	 */
-	void CreateMeshesFromImplicitObject(const Chaos::FImplicitObject* InImplicitObject, TArray<TSharedPtr<FChaosVDExtractedGeometryDataHandle>>& OutMeshDataHandles, int32 AvailableShapeDataNum = 0, const int32 DesiredLODCount = 0, const Chaos::FRigidTransform3& InTransform = Chaos::FRigidTransform3(), const int32 MeshIndex = 0);
+	void CreateMeshesFromImplicitObject(const Chaos::FImplicitObject* InImplicitObject, TArray<TPair<TSharedPtr<FChaosVDExtractedGeometryDataHandle>, TSharedPtr<Chaos::VD::FChaosVDImplicitObjectView>>>& OutMeshDataHandles, Chaos::VD::FChaosVDImplicitObjectView* Root, int32 AvailableShapeDataNum = 0, const int32 DesiredLODCount = 0, const Chaos::FRigidTransform3& InTransform = Chaos::FRigidTransform3(), const int32 MeshIndex = 0);
 
 	virtual void AddReferencedObjects(FReferenceCollector& Collector) override;
 	virtual FString GetReferencerName() const override
@@ -100,7 +100,7 @@ public:
 
 private:
 	
-	void CreateMeshesFromImplicit_Internal(const Chaos::FImplicitObject* InRootImplicitObject,const Chaos::FImplicitObject* InLeafImplicitObject, TArray<TSharedPtr<FChaosVDExtractedGeometryDataHandle>>& OutMeshDataHandles, const int32 DesiredLODCount = 0, const Chaos::FRigidTransform3& InTransform = Chaos::FRigidTransform3(), const int32 ParentShapeInstanceIndex = 0, int32 AvailableShapeDataNum = 0);
+	void CreateMeshesFromImplicit_Internal(const Chaos::FImplicitObject* InRootImplicitObject,const Chaos::FImplicitObject* InLeafImplicitObject, TArray<TPair<TSharedPtr<FChaosVDExtractedGeometryDataHandle>, TSharedPtr<Chaos::VD::FChaosVDImplicitObjectView>>>& OutMeshDataHandles, Chaos::VD::FChaosVDImplicitObjectView* Root, const int32 DesiredLODCount = 0, const Chaos::FRigidTransform3& InTransform = Chaos::FRigidTransform3(), const int32 ParentShapeInstanceIndex = 0, int32 AvailableShapeDataNum = 0);
 	
 public:
 	/**
diff --git a/Engine/Plugins/ChaosVD/Source/ChaosVD/Public/ChaosVDGeometryDataComponent.h b/Engine/Plugins/ChaosVD/Source/ChaosVD/Public/ChaosVDGeometryDataComponent.h
--- a/Engine/Plugins/ChaosVD/Source/ChaosVD/Public/ChaosVDGeometryDataComponent.h
+++ b/Engine/Plugins/ChaosVD/Source/ChaosVD/Public/ChaosVDGeometryDataComponent.h
@@ -271,6 +271,42 @@ private:
 	friend UChaosVDStaticMeshComponent;
 };
 
+
+namespace Chaos::VD
+{
+	/** Class that represents an Implicit Object to be shown in the Geometry Hierarchy Viewer */
+	class FChaosVDImplicitObjectView : public TSharedFromThis<FChaosVDImplicitObjectView>
+	{
+	public:
+		Chaos::EImplicitObjectType ImplicitObjectType;
+		FChaosVDImplicitObjectBasicView ImplicitObjectInfo;
+
+		void AddChild(TSharedPtr<FChaosVDImplicitObjectView> Child)
+		{
+			Children->Add(Child);
+		}
+
+		TSharedPtr<TArray<TSharedPtr<FChaosVDImplicitObjectView>>> GetChildren()
+		{
+			return Children;
+		}
+
+		void SetMeshData(TSharedPtr<FChaosVDInstancedMeshData> InMeshData)
+		{
+			MeshData = InMeshData;
+		}
+
+		TSharedPtr<FChaosVDInstancedMeshData> GetMeshData()
+		{
+			return MeshData;
+		}
+
+	private:
+		TSharedPtr<TArray<TSharedPtr<FChaosVDImplicitObjectView>>> Children = MakeShared<TArray<TSharedPtr<FChaosVDImplicitObjectView>>>();
+		TSharedPtr<FChaosVDInstancedMeshData> MeshData;
+	};
+}
+
 enum class EChaosVDMeshAttributesFlags : uint8
 {
 	None = 0,
diff --git a/Engine/Plugins/ChaosVD/Source/ChaosVD/Public/ChaosVDObjectDetailsTab.h b/Engine/Plugins/ChaosVD/Source/ChaosVD/Public/ChaosVDObjectDetailsTab.h
--- a/Engine/Plugins/ChaosVD/Source/ChaosVD/Public/ChaosVDObjectDetailsTab.h
+++ b/Engine/Plugins/ChaosVD/Source/ChaosVD/Public/ChaosVDObjectDetailsTab.h
@@ -7,6 +7,7 @@
 #include "ChaosVDTabSpawnerBase.h"
 #include "Delegates/IDelegateInstance.h"
 #include "Widgets/SChaosVDDetailsView.h"
+#include "ChaosVD/Private/Widgets/SChaosVDGeometryTree.h"
 
 struct FChaosVDSolverDataSelectionHandle;
 struct FChaosVDParticleDebugData;
@@ -45,7 +46,12 @@ protected:
 
 	CHAOSVD_API bool IsSupportedStruct(const TWeakObjectPtr<const UStruct>& InWeakStructPtr);
 
+	void UpdateTreeDetailsSplit(const FChaosVDSceneParticle* Particle);
+
 	TSharedPtr<SChaosVDDetailsView> DetailsPanelView;
+	TSharedPtr<SChaosVDGeometryTree> GeometryTreeWidget;
+	TSharedPtr<SSplitter> DetailsSplitter;
+	float SavedSplitSize = 0.2f;
 
 	TSet<TWeakObjectPtr<const UStruct>> UnsupportedStructs;
 };
diff --git a/Engine/Plugins/ChaosVD/Source/ChaosVD/Public/ChaosVDSceneParticle.h b/Engine/Plugins/ChaosVD/Source/ChaosVD/Public/ChaosVDSceneParticle.h
--- a/Engine/Plugins/ChaosVD/Source/ChaosVD/Public/ChaosVDSceneParticle.h
+++ b/Engine/Plugins/ChaosVD/Source/ChaosVD/Public/ChaosVDSceneParticle.h
@@ -148,6 +148,11 @@ public:
 		return ParentParticleInstance;
 	}
 
+	TWeakPtr<Chaos::VD::FChaosVDImplicitObjectView> GetRootImplicitView() const
+	{
+		return RootImplicitView.ToWeakPtr();
+	}
+
 	EChaosVDSceneParticleDirtyFlags GetDirtyFlags() const
 	{
 		return DirtyFlags;
@@ -169,7 +174,7 @@ protected:
 
 	void UpdateParent(const TSharedPtr<FChaosVDParticleDataWrapper>& InRecordedData);
 
-	void ProcessUpdatedAndRemovedHandles(TArray<TSharedPtr<FChaosVDExtractedGeometryDataHandle>>& OutExtractedGeometryDataHandles);
+	void ProcessUpdatedAndRemovedHandles(TArray<TPair<TSharedPtr<FChaosVDExtractedGeometryDataHandle>, TSharedPtr<Chaos::VD::FChaosVDImplicitObjectView>>>& OutExtractedGeometryDataHandles);
 
 	const TArray<TSharedPtr<FChaosVDParticlePairMidPhase>>* GetCollisionMidPhasesArray() const;
 
@@ -200,6 +205,9 @@ protected:
 
 	TWeakPtr<FChaosVDInstancedMeshData> CurrentSelectedGeometryInstance;
 
+	TSharedPtr<Chaos::VD::FChaosVDImplicitObjectView> RootImplicitView;
+
+
 	TWeakPtr<FChaosVDScene> SceneWeakPtr;
 
 	FTransform PendingParticleTransform;
