[ChaosVD] Implemented a connection status toolbar

Now CVD's Status bar shows the connection state when a recording is active, indication:

- If we are waiting for data
- How much data the target application has buffered and has not been sent yet
- Data transfer speed

diff --git a/Engine/Plugins/ChaosVD/Source/ChaosVD/Private/ChaosVDEngine.cpp b/Engine/Plugins/ChaosVD/Source/ChaosVD/Private/ChaosVDEngine.cpp
--- a/Engine/Plugins/ChaosVD/Source/ChaosVD/Private/ChaosVDEngine.cpp
+++ b/Engine/Plugins/ChaosVD/Source/ChaosVD/Private/ChaosVDEngine.cpp
@@ -370,7 +370,7 @@ bool FChaosVDEngine::ConnectToLiveSession_Direct(FGuid RemoteSessionID, EChaosVD
 		ExistingRecordingInstance = PlaybackController->GetCurrentRecording().Pin();
 	}
 
-	NewSessionFromFileDescriptor.SessionName = FChaosVDModule::Get().GetTraceManager()->ConnectToLiveSession_Direct(NewSessionFromFileDescriptor.SessionPort, ExistingRecordingInstance);
+	NewSessionFromFileDescriptor.SessionName = FChaosVDModule::Get().GetTraceManager()->ConnectToLiveSession_Direct(RemoteSessionID, NewSessionFromFileDescriptor.SessionPort, ExistingRecordingInstance);
 
 	bool bSuccess = false;
 
diff --git a/Engine/Plugins/ChaosVD/Source/ChaosVD/Private/Trace/ChaosVDTraceManager.cpp b/Engine/Plugins/ChaosVD/Source/ChaosVD/Private/Trace/ChaosVDTraceManager.cpp
--- a/Engine/Plugins/ChaosVD/Source/ChaosVD/Private/Trace/ChaosVDTraceManager.cpp
+++ b/Engine/Plugins/ChaosVD/Source/ChaosVD/Private/Trace/ChaosVDTraceManager.cpp
@@ -86,10 +86,72 @@ namespace Chaos::VD
 		return true;
 	}
 
+	FSessionStreamStatsUpdater::~FSessionStreamStatsUpdater()
+	{
+	}
+
+	bool FSessionStreamStatsUpdater::Tick(float DeltaTime)
+	{
+		ElapsedTimeSinceLastUpdate += DeltaTime;
+
+		constexpr double IntervalBetweenUpdatesSeconds = 1.0;
+		if (ElapsedTimeSinceLastUpdate < IntervalBetweenUpdatesSeconds)
+		{
+			return true;
+		}
+
+		if (!ensure(DataStreamInstance))
+		{
+			return false;
+		}
+
+		ElapsedTimeSinceLastUpdate = 0.0;
+
+		TSharedPtr<FChaosVDRemoteSessionsManager> RemoteSessionManager = FChaosVDEngineEditorBridge::Get().GetRemoteSessionsManager();
+		TSharedPtr<FChaosVDSessionInfo> SessionInfo = RemoteSessionManager ? RemoteSessionManager->GetSessionInfo(DataStreamInstance->GetOwningRemoteSessionID()).Pin() : nullptr;
+		if (!SessionInfo)
+		{
+			return false;
+		}
+
+		SessionInfo->SetReceivedBytesPerSecond(DataStreamInstance->UpdateAndGetBytesReadSinceLastMeasurement());
+	
+		return true;
+	}
+
+	int32 FDirectSocketStream::Read(void* Data, uint32 Size)
+	{
+		int32 BytesRead =  UE::Trace::FDirectSocketStream::Read(Data, Size);
+
+		if (BytesRead > 0)
+		{
+			BytesReadSinceLastMeasurement += BytesRead;
+		}
+	
+		return BytesRead;
+	}
+
+	FGuid FDirectSocketStream::GetOwningRemoteSessionID() const
+	{
+		return RemoteSessionID;
+	}
+
+	uint64 FDirectSocketStream::UpdateAndGetBytesReadSinceLastMeasurement()
+	{
+		uint64 BytesRead = BytesReadSinceLastMeasurement;
+
+		BytesReadSinceLastMeasurement = 0;
+	
+		return BytesRead;
+	}
+
 	FRelayDataStream::FRelayDataStream(FGuid InRemoteSessionID)
 		: CancelEvent(FPlatformProcess::GetSynchEventFromPool(false)),
 		WaitForDataEvent(FPlatformProcess::GetSynchEventFromPool(false)),
-		RemoteSessionID(InRemoteSessionID)
+		RemoteSessionID(InRemoteSessionID),
+		FileWriter(nullptr),
+		BytesReadSinceLastMeasurement(0),
+		StatsUpdater(this)
 	{
 #if WITH_CHAOS_VISUAL_DEBUGGER
 		TSharedPtr<ITraceDataRelayTransport> RelayTransportInstance = FChaosVDEngineEditorBridge::Get().GetRelayTransportInstance();
@@ -100,12 +162,13 @@ namespace Chaos::VD
 #endif
 	}
 
-
 	FRelayDataStream::FRelayDataStream(FGuid InRemoteSessionID, TUniquePtr<FArchiveFileWriterGeneric>&& InFileWriter)
 		: CancelEvent(FPlatformProcess::GetSynchEventFromPool(false)),
 		WaitForDataEvent(FPlatformProcess::GetSynchEventFromPool(false)),
 		RemoteSessionID(InRemoteSessionID),
-		FileWriter(MoveTemp(InFileWriter))
+		FileWriter(MoveTemp(InFileWriter)),
+		BytesReadSinceLastMeasurement(0),
+		StatsUpdater(this)
 	{
 #if WITH_CHAOS_VISUAL_DEBUGGER
 		TSharedPtr<ITraceDataRelayTransport> RelayTransportInstance = FChaosVDEngineEditorBridge::Get().GetRelayTransportInstance();
@@ -160,6 +223,18 @@ namespace Chaos::VD
 		DataQueue.Enqueue(MoveTemp(DataBuffer));
 	}
 
+	FGuid FRelayDataStream::GetOwningRemoteSessionID() const
+	{
+		return RemoteSessionID;
+	}
+
+	uint64 FRelayDataStream::UpdateAndGetBytesReadSinceLastMeasurement()
+	{
+		uint64 BytesRead = BytesReadSinceLastMeasurement;
+		BytesReadSinceLastMeasurement = 0;
+		return BytesRead;
+	}
+
 	int32 FRelayDataStream::Read(void* Dest, uint32 DestSize)
 	{
 		if (DestSize == 0)
@@ -224,6 +299,8 @@ namespace Chaos::VD
 			// we'll wait for some data to arrive.
 			if (BytesCopied > 0)
 			{
+				BytesReadSinceLastMeasurement += BytesCopied;
+
 				if (FileWriter.IsValid() && !FileWriter->IsError())
 				{
 					FileWriter->Serialize(Dest, BytesCopied);
@@ -605,7 +682,7 @@ TUniquePtr<FArchiveFileWriterGeneric> FChaosVDTraceManager::OpenTraceFileForWrit
 	return MoveTemp(FileWriter);
 }
 
-FString FChaosVDTraceManager::ConnectToLiveSession_Direct(uint16& OutSessionPort, const TSharedPtr<FChaosVDRecording>& ExistingRecordingPtr)
+FString FChaosVDTraceManager::ConnectToLiveSession_Direct(FGuid RemoteSessionID, uint16& OutSessionPort, const TSharedPtr<FChaosVDRecording>& ExistingRecordingPtr)
 {
 	TUniquePtr<FArchiveFileWriterGeneric> TraceFileWriter = nullptr;
 
@@ -622,7 +699,7 @@ FString FChaosVDTraceManager::ConnectToLiveSession_Direct(uint16& OutSessionPort
 		}
 	}
 
-	TUniquePtr<UE::Trace::FDirectSocketStream> DirectSocketStream = MakeUnique<UE::Trace::FDirectSocketStream>(MoveTemp(TraceFileWriter));
+	TUniquePtr<Chaos::VD::FDirectSocketStream> DirectSocketStream = MakeUnique<Chaos::VD::FDirectSocketStream>(RemoteSessionID, MoveTemp(TraceFileWriter));
 	OutSessionPort = DirectSocketStream->StartListening(CurrentDirectModeStartingPort);
 
 	if (OutSessionPort == 0)
diff --git a/Engine/Plugins/ChaosVD/Source/ChaosVD/Private/Widgets/SChaosVDMainTab.cpp b/Engine/Plugins/ChaosVD/Source/ChaosVD/Private/Widgets/SChaosVDMainTab.cpp
--- a/Engine/Plugins/ChaosVD/Source/ChaosVD/Private/Widgets/SChaosVDMainTab.cpp
+++ b/Engine/Plugins/ChaosVD/Source/ChaosVD/Private/Widgets/SChaosVDMainTab.cpp
@@ -31,6 +31,7 @@
 #include "HAL/PlatformApplicationMisc.h"
 #include "IDesktopPlatform.h"
 #include "PropertyEditorModule.h"
+#include "SConnectionStatusBar.h"
 #include "Misc/MessageDialog.h"
 #include "StatusBarSubsystem.h"
 #include "ToolMenu.h"
@@ -133,11 +134,12 @@ void SChaosVDMainTab::Construct(const FArguments& InArgs, TSharedPtr<FChaosVDEng
 	
 	if (UStatusBarSubsystem* StatusBarSubsystem = GEditor ? GEditor->GetEditorSubsystem<UStatusBarSubsystem>() : nullptr)
 	{
-		 StatusBarWidget = StatusBarSubsystem->MakeStatusBarWidget(StatusBarID, TabManager->GetOwnerTab().ToSharedRef());
+		using namespace UE::Editor::Toolbars;
+		ECreateStatusBarOptions StatusBarOptions = ECreateStatusBarOptions::HideContentBrowser | ECreateStatusBarOptions::HideOutputLog| ECreateStatusBarOptions::HideSourceControl;
 
-		// Status bars come with the output log and content browser drawers by default, therefore we need to remove them otherwise they will be on the tool's window
-		StatusBarSubsystem->UnregisterDrawer(StatusBarID, "ContentBrowser");
-		StatusBarSubsystem->UnregisterDrawer(StatusBarID, "OutputLog");
+		StatusBarWidget = StatusBarSubsystem->MakeStatusBarWidget(StatusBarID, TabManager->GetOwnerTab().ToSharedRef(), StatusBarOptions);
+
+		ExtendStatusToolbar();
 	}
 	else
 	{
@@ -213,6 +215,17 @@ void SChaosVDMainTab::Construct(const FArguments& InArgs, TSharedPtr<FChaosVDEng
 	}
 }
 
+SChaosVDMainTab::~SChaosVDMainTab()
+{
+	if (TSharedPtr<FChaosVDRemoteSessionsManager> RemoteSessionManager = FChaosVDEngineEditorBridge::Get().GetRemoteSessionsManager())
+	{
+		RemoteSessionManager->StopSessionDiscovery();
+	}
+
+	FChaosVDExtensionsManager::Get().OnExtensionRegistered().RemoveAll(this);
+	CleanUpDisableCPUThrottlingDelegate();
+}
+
 void SChaosVDMainTab::ExtendMainToolbar()
 {
 	if (!RecordingControls)
@@ -224,15 +237,19 @@ void SChaosVDMainTab::ExtendMainToolbar()
 	RecordingControls->AddToMenu(MainToolBarName);
 	AddSettingsMenu();
 }
-SChaosVDMainTab::~SChaosVDMainTab()
+
+void SChaosVDMainTab::ExtendStatusToolbar()
 {
-	if (TSharedPtr<FChaosVDRemoteSessionsManager> RemoteSessionManager = FChaosVDEngineEditorBridge::Get().GetRemoteSessionsManager())
+	FName StatusBarToolbarID = FName(StatusBarID.ToString() + TEXT(".Toolbar"));
+
+	if (UToolMenu* StatusToolBarMenu = UToolMenus::Get()->ExtendMenu(StatusBarToolbarID))
 	{
-		RemoteSessionManager->StopSessionDiscovery();
-	}
+		FToolMenuSection& ConnectionStatusSection = StatusToolBarMenu->AddSection("ConnectionStatus");
 
-	FChaosVDExtensionsManager::Get().OnExtensionRegistered().RemoveAll(this);
-	CleanUpDisableCPUThrottlingDelegate();
+		constexpr bool bNoIndent = true;
+		constexpr bool bSearchable = false;
+		ConnectionStatusSection.AddEntry(FToolMenuEntry::InitWidget("ConnectionStatusBar", SNew(Chaos::VD::SConnectionStatusBar), FText::GetEmpty(), bNoIndent, bSearchable));
+	}
 }
 
 void SChaosVDMainTab::BindUICommands(const TSharedRef<FUICommandList>& InGlobalUICommandsRef)
diff --git a/Engine/Plugins/ChaosVD/Source/ChaosVD/Private/Widgets/SConnectionStatusBar.cpp b/Engine/Plugins/ChaosVD/Source/ChaosVD/Private/Widgets/SConnectionStatusBar.cpp
new file mode 100644
--- /dev/null
+++ b/Engine/Plugins/ChaosVD/Source/ChaosVD/Private/Widgets/SConnectionStatusBar.cpp
@@ -0,0 +1,209 @@
+﻿// Copyright Epic Games, Inc. All Rights Reserved.
+
+#include "SConnectionStatusBar.h"
+
+#include "Chaos/ChaosVDEngineEditorBridge.h"
+#include "Components/HorizontalBox.h"
+#include "Widgets/SOverlay.h"
+#include "Widgets/Images/SImage.h"
+#include "Widgets/Text/STextBlock.h"
+
+#define LOCTEXT_NAMESPACE "ChaosVisualDebugger"
+
+namespace Chaos::VD
+{
+
+void SConnectionStatusBar::Construct(const FArguments& InArgs)
+{
+	ChildSlot
+	[
+		SNew(SHorizontalBox)
+	
+			+ SHorizontalBox::Slot()
+			.AutoWidth()
+			.VAlign(VAlign_Center)
+			.Padding(0, 0, 3, 0)
+			[
+				SNew(SOverlay)
+				+ SOverlay::Slot()
+				.HAlign(HAlign_Center)
+				.VAlign(VAlign_Top)
+				[
+					SNew(SImage)
+					.ColorAndOpacity(FSlateColor::UseForeground())
+					.Image_Raw(this, &SConnectionStatusBar::GetStateBackgroundIcon)
+				]
+				+ SOverlay::Slot()
+				.HAlign(HAlign_Center)
+				.VAlign(VAlign_Top)
+				[
+					SNew(SImage)
+					.Image(FAppStyle::Get().GetBrush("DerivedData.RemoteCache.Busy"))
+					.ColorAndOpacity_Raw(this, &SConnectionStatusBar::GetWaitingIconColor)
+				]
+				+ SOverlay::Slot()
+				.HAlign(HAlign_Center)
+				.VAlign(VAlign_Top)
+				[
+					SNew(SImage)
+					.Image(FAppStyle::Get().GetBrush("DerivedData.RemoteCache.Uploading"))
+					.ColorAndOpacity_Raw(this, &SConnectionStatusBar::GetUploadingIconColor)
+				]
+				+ SOverlay::Slot()
+				.HAlign(HAlign_Center)
+				.VAlign(VAlign_Top)
+				[
+					SNew(SImage)
+					.Image(FAppStyle::Get().GetBrush("DerivedData.RemoteCache.Downloading"))
+					.ColorAndOpacity_Raw(this, &SConnectionStatusBar::GetDownloadIconColor)
+				]
+				+ SOverlay::Slot()
+				.HAlign(HAlign_Center)
+				.VAlign(VAlign_Top)
+				[
+					SNew(SImage)
+					.Image(FAppStyle::Get().GetBrush("Zen.Status.Idle"))
+					.ColorAndOpacity_Raw(this, &SConnectionStatusBar::GetReadyIconColor)
+				]
+			]
+			+ SHorizontalBox::Slot()
+			.AutoWidth()
+			.VAlign(VAlign_Center)
+			.Padding(0, 0, 10, 0)
+			[
+				SNew(STextBlock)
+				.Text_Raw(this, &SConnectionStatusBar::GetTitleText)
+				.ToolTipText_Raw(this, &SConnectionStatusBar::GetTitleToolTipText)
+			]
+	];
+}
+
+const FSlateBrush* SConnectionStatusBar::GetStateBackgroundIcon() const
+{
+	return State == EState::WaitingForData ? FAppStyle::Get().GetBrush("DerivedData.RemoteCache.IdleBG") : FAppStyle::Get().GetBrush("DerivedData.RemoteCache.IdleBG");
+}
+
+FText SConnectionStatusBar::GetTitleText() const
+{
+	static const FNumberFormattingOptions SizeFormattingOptions = FNumberFormattingOptions().SetMinimumFractionalDigits(2).SetMaximumFractionalDigits(2);
+
+	switch (State)
+	{
+	case EState::Ready:
+		return LOCTEXT("ConnectionReadyStatusText", "Ready");
+	case EState::ReceivingData:
+		{
+			FText BytesPerSecondText =  FText::AsMemory(CurrentReceivedDataBytesPerSecond, &SizeFormattingOptions,nullptr, IEC);
+			FText PendingBytesText =  FText::AsMemory(CurrentPendingBytes, &SizeFormattingOptions,nullptr, IEC);
+			FText TotalBytesText =  FText::AsMemory(CurrentTotalReceivedData, &SizeFormattingOptions,nullptr, IEC);
+			return FText::FormatOrdered(LOCTEXT("ConnectionDownloadedStatusText", "Receiving {0}/s | {1} Pending | Total {2}"), BytesPerSecondText, PendingBytesText, TotalBytesText);
+		}
+	case EState::WaitingForData:
+	default:
+		{
+			FText PendingBytesText =  FText::AsMemory(CurrentPendingBytes, &SizeFormattingOptions,nullptr, IEC);
+			FText TotalBytesText =  FText::AsMemory(CurrentTotalReceivedData, &SizeFormattingOptions,nullptr, IEC);
+			return FText::FormatOrdered(LOCTEXT("ConnectionWaitingStatusText", "Waiting for Data | {0} Pending | Total {1}"), PendingBytesText, TotalBytesText);
+		}
+	}
+}
+
+FText SConnectionStatusBar::GetTitleToolTipText() const
+{
+	return LOCTEXT("DownloadingStatusToolTipText", "Current state of the connection to the recording target.");
+}
+
+FSlateColor SConnectionStatusBar::GetDownloadIconColor() const
+{
+	return State == EState::ReceivingData ? FLinearColor::White.CopyWithNewOpacity(FMath::MakePulsatingValue(ElapsedTime, 2)) : FLinearColor(0.0,0.0,0.0,0.0);
+}
+
+FSlateColor SConnectionStatusBar::GetUploadingIconColor() const
+{
+	// CVD does not upload data, but if we are downloading data we want to slightly see the upload arrow
+	return State == EState::ReceivingData ? FLinearColor::White.CopyWithNewOpacity(0.2f) : FLinearColor(0.0,0.0,0.0,0.0);
+}
+
+FSlateColor SConnectionStatusBar::GetWaitingIconColor() const
+{
+	return State == EState::WaitingForData ? FSlateColor::UseForeground() : FLinearColor(0.0,0.0,0.0,0.0);
+}
+
+FSlateColor SConnectionStatusBar::GetReadyIconColor() const
+{
+	return State == EState::Ready ? FSlateColor::UseForeground() : FLinearColor(0.0,0.0,0.0,0.0);
+}
+
+void SConnectionStatusBar::Tick(const FGeometry& AllottedGeometry, const double InCurrentTime, const float InDeltaTime)
+{
+	SCompoundWidget::Tick(AllottedGeometry, InCurrentTime, InDeltaTime);
+
+	ElapsedTime += InDeltaTime;
+
+	UpdateConnectionStatus(InDeltaTime);
+}
+
+void SConnectionStatusBar::UpdateConnectionStatus(double DeltaTime)
+{
+	ElapsedTimeSinceLastUpdate += DeltaTime;
+
+	constexpr double UpdateIntervalSeconds = 1.0;
+	if (ElapsedTimeSinceLastUpdate < UpdateIntervalSeconds)
+	{
+		return;
+	}
+
+	ElapsedTimeSinceLastUpdate = 0;
+
+	TSharedPtr<FChaosVDRemoteSessionsManager> RemoteSessionManager = FChaosVDEngineEditorBridge::Get().GetRemoteSessionsManager();
+	if (!RemoteSessionManager)
+	{
+		State = EState::Ready;
+		return;
+	}
+
+	bool bIsAnyRecording = false;
+	int64 PendingBytesNum = 0;
+	uint64 BytesReadSinceLastMeasurement = 0;
+	RemoteSessionManager->EnumerateActiveSessions([this, &BytesReadSinceLastMeasurement, &bIsAnyRecording, &PendingBytesNum](const TSharedRef<FChaosVDSessionInfo>& InSessionInfoRef)
+	{
+		if (!EnumHasAnyFlags(InSessionInfoRef->GetSessionTypeAttributes(), EChaosVDRemoteSessionAttributes::IsMultiSessionWrapper))
+		{
+			if (InSessionInfoRef->IsRecording())
+			{
+				bIsAnyRecording = true;
+				PendingBytesNum += InSessionInfoRef->LastKnownRecordingState.BufferedDataBytesSize;
+				BytesReadSinceLastMeasurement += InSessionInfoRef->GetReceivedBytesPerSecond();
+			}			
+		}
+
+		return true;
+	});
+
+	CurrentPendingBytes = PendingBytesNum;
+
+	if (bIsAnyRecording)
+	{
+		if (BytesReadSinceLastMeasurement > 0)
+		{
+			CurrentReceivedDataBytesPerSecond = BytesReadSinceLastMeasurement;
+			CurrentTotalReceivedData += BytesReadSinceLastMeasurement;
+			State = EState::ReceivingData;
+		}
+		else
+		{
+			State = EState::WaitingForData;
+		}
+	}
+	else
+	{
+		State = EState::Ready;
+		CurrentPendingBytes = 0;
+		CurrentReceivedDataBytesPerSecond = 0;
+		CurrentTotalReceivedData = 0;
+	}
+}
+
+}
+
+#undef LOCTEXT_NAMESPACE
diff --git a/Engine/Plugins/ChaosVD/Source/ChaosVD/Private/Widgets/SConnectionStatusBar.h b/Engine/Plugins/ChaosVD/Source/ChaosVD/Private/Widgets/SConnectionStatusBar.h
new file mode 100644
--- /dev/null
+++ b/Engine/Plugins/ChaosVD/Source/ChaosVD/Private/Widgets/SConnectionStatusBar.h
@@ -0,0 +1,54 @@
+﻿// Copyright Epic Games, Inc. All Rights Reserved.
+
+#pragma once
+
+#include "CoreMinimal.h"
+#include "Widgets/SCompoundWidget.h"
+
+namespace Chaos::VD
+{
+
+/**
+ * Status Bar widget that provides information about active CVD connections
+ */
+class SConnectionStatusBar : public SCompoundWidget
+{
+public:
+	SLATE_BEGIN_ARGS(SConnectionStatusBar)
+	{
+	}
+	SLATE_END_ARGS()
+
+	void Construct(const FArguments& InArgs);
+
+private:
+
+	const FSlateBrush* GetStateBackgroundIcon() const;
+	FText GetTitleText() const;
+	FText GetTitleToolTipText() const;
+	FSlateColor GetDownloadIconColor() const;
+	FSlateColor GetUploadingIconColor() const;
+	FSlateColor GetWaitingIconColor() const;
+	FSlateColor GetReadyIconColor() const;
+
+	virtual void Tick(const FGeometry& AllottedGeometry, const double InCurrentTime, const float InDeltaTime) override;
+
+	void UpdateConnectionStatus(double DeltaTime);
+
+	double ElapsedTime = 0.0;
+	double ElapsedTimeSinceLastUpdate = 0.0;
+
+	enum class EState
+	{
+		Ready,
+		ReceivingData,
+		WaitingForData
+	};
+	EState State = EState::Ready;
+
+	int64 CurrentPendingBytes = 0;
+	uint64 CurrentReceivedDataBytesPerSecond = 0;
+	uint64 CurrentTotalReceivedData = 0;
+};
+
+}
diff --git a/Engine/Plugins/ChaosVD/Source/ChaosVD/Public/Trace/ChaosVDTraceManager.h b/Engine/Plugins/ChaosVD/Source/ChaosVD/Public/Trace/ChaosVDTraceManager.h
--- a/Engine/Plugins/ChaosVD/Source/ChaosVD/Public/Trace/ChaosVDTraceManager.h
+++ b/Engine/Plugins/ChaosVD/Source/ChaosVD/Public/Trace/ChaosVDTraceManager.h
@@ -27,7 +27,6 @@ namespace Chaos::VD
 		{
 		}
 
-
 		virtual ~FAsyncProgressNotification() override;
 
 		/** Updates the current progress
@@ -45,9 +44,65 @@ namespace Chaos::VD
 		FText NotificationTitle;
 	};
 
+	class IRemoteSessionDataStream
+	{
+	public:
+		virtual ~IRemoteSessionDataStream() = default;
+		virtual FGuid GetOwningRemoteSessionID() const = 0;
+		virtual uint64 UpdateAndGetBytesReadSinceLastMeasurement() = 0;
+	};
+
+	class FSessionStreamStatsUpdater : FTSTickerObjectBase
+	{
+	public:
+		virtual ~FSessionStreamStatsUpdater() override;
+
+		explicit FSessionStreamStatsUpdater(IRemoteSessionDataStream* const DataStreamInstance)
+			: DataStreamInstance(DataStreamInstance)
+		{
+		}
+
+	private:
+		virtual bool Tick(float DeltaTime) override;
+	
+		IRemoteSessionDataStream* DataStreamInstance;
+
+		double ElapsedTimeSinceLastUpdate = 0.0;
+	};
+
+	class FDirectSocketStream : public UE::Trace::FDirectSocketStream, public IRemoteSessionDataStream
+	{
+	public:
+
+		FDirectSocketStream() = delete;
+
+		explicit FDirectSocketStream(FGuid InRemoteSessionID) : BytesReadSinceLastMeasurement(0), RemoteSessionID(InRemoteSessionID), StatsUpdater(this)
+		{
+		}
+
+		explicit FDirectSocketStream(FGuid InRemoteSessionID, TUniquePtr<FArchiveFileWriterGeneric>&& InFileWriter)
+			: UE::Trace::FDirectSocketStream(MoveTemp(InFileWriter)),
+			BytesReadSinceLastMeasurement(0),
+			RemoteSessionID(InRemoteSessionID),
+			StatsUpdater(this)
+		{
+		}
+
+		virtual int32 Read(void* Data, uint32 Size) override;
+
+		virtual FGuid GetOwningRemoteSessionID() const override;
+		virtual uint64 UpdateAndGetBytesReadSinceLastMeasurement() override;
+	
+	private:
+		std::atomic<uint64> BytesReadSinceLastMeasurement;
+		FGuid RemoteSessionID;
+
+		FSessionStreamStatsUpdater StatsUpdater;
+	};
+
 	/** A trace data stream object that listens for data coming from a relay stream
 	 * via CVD's Engine Editor Bridge object */
-	class FRelayDataStream : public UE::Trace::IInDataStream
+	class FRelayDataStream : public UE::Trace::IInDataStream,  public IRemoteSessionDataStream
 	{
 	public:
 		explicit FRelayDataStream(FGuid InRemoteSessionID);
@@ -58,9 +113,11 @@ namespace Chaos::VD
 		void EnqueueRelayedData(const TConstArrayView<uint8> InTraceDataBuffer);
 		void EnqueueRelayedData(TArray<uint8>&& InTraceDataBuffer);
 
+		virtual FGuid GetOwningRemoteSessionID() const override;
+		virtual uint64 UpdateAndGetBytesReadSinceLastMeasurement() override;
+
 	private:
 		virtual int32 Read(void* Dest, uint32 DestSize) override;
-
 		virtual void Close() override;
 		virtual bool WaitUntilReady() override;
 
@@ -74,6 +131,9 @@ namespace Chaos::VD
 		FGuid RemoteSessionID = FGuid();
 
 		TUniquePtr<FArchiveFileWriterGeneric> FileWriter;
+
+		std::atomic<uint64> BytesReadSinceLastMeasurement;
+		FSessionStreamStatsUpdater StatsUpdater;
 	};
 }
 
@@ -145,13 +205,21 @@ public:
 	 */
 	FString ConnectToLiveSession(FStringView InSessionHost, uint32 SessionID, const TSharedPtr<FChaosVDRecording>& ExistingRecordingPtr = nullptr);
 
+
+	UE_DEPRECATED(5.8, "Please use the version that takes a Remote Session ID")
+	FString ConnectToLiveSession_Direct( uint16& OutSessionPort, const TSharedPtr<FChaosVDRecording>& ExistingRecordingPtr = nullptr)
+	{
+		return FString();
+	};
+
 	/**
 	 * Opens up a direct trace stream and waits for data
+	 * @param RemoteSessionID Session ID to which we want to connect to
 	 * @param OutSessionPort Number of the port we open for this session
 	 * @param ExistingRecordingPtr Ptr to an existing recording instance to append the new CVD data to, if any
 	 * @return Session Name
 	 */
-	FString ConnectToLiveSession_Direct(uint16& OutSessionPort, const TSharedPtr<FChaosVDRecording>& ExistingRecordingPtr = nullptr);
+	FString ConnectToLiveSession_Direct(FGuid RemoteSessionID, uint16& OutSessionPort, const TSharedPtr<FChaosVDRecording>& ExistingRecordingPtr = nullptr);
 
 	/**
 	 * Opens up a relay trace stream and waits for data
diff --git a/Engine/Plugins/ChaosVD/Source/ChaosVD/Public/Widgets/SChaosVDMainTab.h b/Engine/Plugins/ChaosVD/Source/ChaosVD/Public/Widgets/SChaosVDMainTab.h
--- a/Engine/Plugins/ChaosVD/Source/ChaosVD/Public/Widgets/SChaosVDMainTab.h
+++ b/Engine/Plugins/ChaosVD/Source/ChaosVD/Public/Widgets/SChaosVDMainTab.h
@@ -175,6 +175,7 @@ public:
 private:
 
 	void ExtendMainToolbar();
+	void ExtendStatusToolbar();
 
 	TSharedPtr<SNotificationItem> CreateProcessingTraceDataNotification();
 
diff --git a/Engine/Source/Runtime/Experimental/ChaosSolverEngine/Private/Chaos/ChaosVDEngineEditorBridge.cpp b/Engine/Source/Runtime/Experimental/ChaosSolverEngine/Private/Chaos/ChaosVDEngineEditorBridge.cpp
--- a/Engine/Source/Runtime/Experimental/ChaosSolverEngine/Private/Chaos/ChaosVDEngineEditorBridge.cpp
+++ b/Engine/Source/Runtime/Experimental/ChaosSolverEngine/Private/Chaos/ChaosVDEngineEditorBridge.cpp
@@ -176,6 +176,7 @@ bool FChaosVDEngineEditorBridge::BroadcastSessionStatus(float DeltaTime)
 	StatusMessage.InstanceId = FApp::GetInstanceId();
 	StatusMessage.bIsRecording = RuntimeModule.IsRecording();
 	StatusMessage.ElapsedTime = RuntimeModule.GetAccumulatedRecordingTime();
+	StatusMessage.BufferedDataBytesSize = RuntimeModule.GetBufferedDataBytesSize();
 
 	RemoteSessionsManager->PublishRecordingStatusUpdate(StatusMessage);
 
diff --git a/Engine/Source/Runtime/Experimental/ChaosSolverEngine/Private/Chaos/ChaosVDRemoteSessionsManager.cpp b/Engine/Source/Runtime/Experimental/ChaosSolverEngine/Private/Chaos/ChaosVDRemoteSessionsManager.cpp
--- a/Engine/Source/Runtime/Experimental/ChaosSolverEngine/Private/Chaos/ChaosVDRemoteSessionsManager.cpp
+++ b/Engine/Source/Runtime/Experimental/ChaosSolverEngine/Private/Chaos/ChaosVDRemoteSessionsManager.cpp
@@ -42,6 +42,11 @@ bool FChaosVDSessionInfo::IsRecording() const
 	return LastKnownRecordingState.bIsRecording;
 }
 
+uint64 FChaosVDSessionInfo::GetBufferedBytesNum() const
+{
+	return LastKnownRecordingState.BufferedDataBytesSize;
+}
+
 EChaosVDRecordingMode FChaosVDSessionInfo::GetRecordingMode() const
 {
 	PRAGMA_DISABLE_DEPRECATION_WARNINGS
@@ -68,6 +73,16 @@ bool FChaosVDSessionInfo::IsConnected() const
 	return false;
 }
 
+void FChaosVDSessionInfo::SetReceivedBytesPerSecond(uint64 InNewBytesPerSecond)
+{
+	ReceivedBytesPerSecond = InNewBytesPerSecond;
+}
+
+uint64 FChaosVDSessionInfo::GetReceivedBytesPerSecond() const
+{
+	return ReceivedBytesPerSecond;
+}
+
 bool FChaosVDMultiSessionInfo::IsRecording() const
 {
 	bool bIsRecording = false;
@@ -84,6 +99,18 @@ bool FChaosVDMultiSessionInfo::IsRecording() const
 	return bIsRecording;
 }
 
+uint64 FChaosVDMultiSessionInfo::GetBufferedBytesNum() const
+{
+	uint64 BufferedBytesNum = 0;
+	EnumerateInnerSessions([&BufferedBytesNum](const TSharedRef<FChaosVDSessionInfo>& InSessionRef)
+	{
+		BufferedBytesNum += InSessionRef->GetBufferedBytesNum();
+		return true;
+	});
+
+	return BufferedBytesNum;
+}
+
 EChaosVDRecordingMode FChaosVDMultiSessionInfo::GetRecordingMode() const
 {
 	PRAGMA_DISABLE_DEPRECATION_WARNINGS
diff --git a/Engine/Source/Runtime/Experimental/ChaosSolverEngine/Public/Chaos/ChaosVDRemoteSessionsManager.h b/Engine/Source/Runtime/Experimental/ChaosSolverEngine/Public/Chaos/ChaosVDRemoteSessionsManager.h
--- a/Engine/Source/Runtime/Experimental/ChaosSolverEngine/Public/Chaos/ChaosVDRemoteSessionsManager.h
+++ b/Engine/Source/Runtime/Experimental/ChaosSolverEngine/Public/Chaos/ChaosVDRemoteSessionsManager.h
@@ -75,6 +75,9 @@ struct FChaosVDRecordingStatusMessage
 	UPROPERTY()
 	float ElapsedTime = false;
 
+	UPROPERTY()
+	int64 BufferedDataBytesSize = 0;
+
 	UE_DEPRECATED(5.7, "Please get the trace details directly from the session info object")
 	FChaosVDTraceDetails TraceDetails;
 };
@@ -249,6 +252,7 @@ public:
 	CHAOSSOLVERENGINE_API const FChaosVDTraceDetails& GetConnectionDetails();
 
 	CHAOSSOLVERENGINE_API virtual bool IsRecording() const;
+	CHAOSSOLVERENGINE_API virtual uint64 GetBufferedBytesNum() const;
 
 	UE_DEPRECATED(5.8, "This method will no longer be used and there is not a replacement planned.")
 	CHAOSSOLVERENGINE_API virtual EChaosVDRecordingMode GetRecordingMode() const;
@@ -266,10 +270,16 @@ public:
 		return SessionTypeAttributes;
 	}
 
+	CHAOSSOLVERENGINE_API void SetReceivedBytesPerSecond(uint64 InNewBytesPerSecond);
+	CHAOSSOLVERENGINE_API uint64 GetReceivedBytesPerSecond() const;
+
 protected:
 	UE_DEPRECATED(5.8, "This method will no longer be used and there is not a replacement planned.")
 	EChaosVDRecordingMode LastRequestedRecordingMode = EChaosVDRecordingMode::Invalid;
 	FChaosVDTraceDetails LastKnownConnectionDetails;
+
+	uint64 ReceivedBytesPerSecond = 0;
+	
 	const EChaosVDRemoteSessionAttributes SessionTypeAttributes;
 
 	friend class FChaosVDRemoteSessionsManager;
@@ -288,6 +298,7 @@ struct FChaosVDMultiSessionInfo : public FChaosVDSessionInfo
 	virtual ~FChaosVDMultiSessionInfo() override = default;
 
 	CHAOSSOLVERENGINE_API virtual bool IsRecording() const override;
+	CHAOSSOLVERENGINE_API virtual uint64 GetBufferedBytesNum() const override;
 
 	UE_DEPRECATED(5.8, "This method will no longer be used and there is not a replacement planned.")
 	CHAOSSOLVERENGINE_API virtual EChaosVDRecordingMode GetRecordingMode() const override;
diff --git a/Engine/Source/Runtime/Experimental/ChaosVisualDebugger/Private/ChaosVDRuntimeModule.cpp b/Engine/Source/Runtime/Experimental/ChaosVisualDebugger/Private/ChaosVDRuntimeModule.cpp
--- a/Engine/Source/Runtime/Experimental/ChaosVisualDebugger/Private/ChaosVDRuntimeModule.cpp
+++ b/Engine/Source/Runtime/Experimental/ChaosVisualDebugger/Private/ChaosVDRuntimeModule.cpp
@@ -153,6 +153,16 @@ int32 FChaosVDRuntimeModule::GenerateUniqueID()
 	return NewID;
 }
 
+int64 FChaosVDRuntimeModule::GetBufferedDataBytesSize() const
+{
+	if (RelayWriter)
+	{
+		return RelayWriter->GetQueuedBytesNum();
+	}
+
+	return 0;
+}
+
 FString FChaosVDRuntimeModule::GetLastRecordingFileNamePath() const
 {
 	return FString();
diff --git a/Engine/Source/Runtime/Experimental/ChaosVisualDebugger/Public/ChaosVDRuntimeModule.h b/Engine/Source/Runtime/Experimental/ChaosVisualDebugger/Public/ChaosVDRuntimeModule.h
--- a/Engine/Source/Runtime/Experimental/ChaosVisualDebugger/Public/ChaosVDRuntimeModule.h
+++ b/Engine/Source/Runtime/Experimental/ChaosVisualDebugger/Public/ChaosVDRuntimeModule.h
@@ -136,6 +136,8 @@ public:
 		return AccumulatedRecordingTime;
 	}
 
+	UE_API int64 GetBufferedDataBytesSize() const;
+
 	/** [Deprecated] Returns the full path of the active recording file*/
 	UE_DEPRECATED(5.7, "This method is no longer used. Use GetCurrentTraceSessionDetails")
 	UE_API FString GetLastRecordingFileNamePath() const;
