[ChaosVD] Fixed CIS warning in SConnectionStatusbar.cpp

The status icon background icon was returning the same icon regardless of the wait state condition. This CL updates the code to return either the Busy state icon or the Idle state icon as intended

diff --git a/Engine/Plugins/ChaosVD/Source/ChaosVD/Private/Widgets/SConnectionStatusBar.cpp b/Engine/Plugins/ChaosVD/Source/ChaosVD/Private/Widgets/SConnectionStatusBar.cpp
--- a/Engine/Plugins/ChaosVD/Source/ChaosVD/Private/Widgets/SConnectionStatusBar.cpp
+++ b/Engine/Plugins/ChaosVD/Source/ChaosVD/Private/Widgets/SConnectionStatusBar.cpp
@@ -80,7 +80,7 @@ void SConnectionStatusBar::Construct(const FArguments& InArgs)
 
 const FSlateBrush* SConnectionStatusBar::GetStateBackgroundIcon() const
 {
-	return State == EState::WaitingForData ? FAppStyle::Get().GetBrush("DerivedData.RemoteCache.IdleBG") : FAppStyle::Get().GetBrush("DerivedData.RemoteCache.IdleBG");
+	return State == EState::WaitingForData ? FAppStyle::Get().GetBrush("DerivedData.RemoteCache.BusyBG") : FAppStyle::Get().GetBrush("DerivedData.RemoteCache.IdleBG");
 }
 
 FText SConnectionStatusBar::GetTitleText() const
