[ChaosVD] Moving the general settings object to the public API so extensions can use it

In addition, added separate transport mode settings for Game Servers and Game Clients, every other target mode can only use the Direct Transport mode.

diff --git a/Engine/Plugins/ChaosVD/Source/ChaosVD/Private/Settings/ChaosVDGeneralSettings.cpp b/Engine/Plugins/ChaosVD/Source/ChaosVD/Private/Settings/ChaosVDGeneralSettings.cpp
--- a/Engine/Plugins/ChaosVD/Source/ChaosVD/Private/Settings/ChaosVDGeneralSettings.cpp
+++ b/Engine/Plugins/ChaosVD/Source/ChaosVD/Private/Settings/ChaosVDGeneralSettings.cpp
@@ -1,5 +1,53 @@
 // Copyright Epic Games, Inc. All Rights Reserved.
 
-#include "ChaosVDGeneralSettings.h"
+#include "Settings/ChaosVDGeneralSettings.h"
+
+#include "ChaosVDModule.h"
 
 #include UE_INLINE_GENERATED_CPP_BY_NAME(ChaosVDGeneralSettings)
+
+EChaosVDTransportMode UChaosVDGeneralSettings::GetTransportModeForTarget(EBuildTargetType TargetType) const
+{
+	switch (TargetType)
+	{
+	case EBuildTargetType::Server:
+		return GameServerDataTransportMode;
+	case EBuildTargetType::Game:
+	case EBuildTargetType::Client:
+		return GameClientDataTransportMode;
+	case EBuildTargetType::Editor:
+	case EBuildTargetType::Unknown:
+	case EBuildTargetType::Program:
+	default:
+		return EChaosVDTransportMode::Direct;
+	}
+}
+
+void UChaosVDGeneralSettings::SetTransportModeForTarget(EBuildTargetType TargetType, EChaosVDTransportMode NewTransportMode)
+{
+	switch (TargetType)
+	{
+	case EBuildTargetType::Server:
+		{
+			GameServerDataTransportMode = NewTransportMode;
+			break;
+		}
+	case EBuildTargetType::Game:
+	case EBuildTargetType::Client:
+		{
+			GameClientDataTransportMode = NewTransportMode;
+			break;
+		}
+	case EBuildTargetType::Editor:
+	case EBuildTargetType::Unknown:
+	case EBuildTargetType::Program:
+	default:
+		{
+			const FString ErrorMessage = FString::Printf(TEXT("[%hs Changing transport mode in the selected target type is not supported. Target type [%s]]"), __func__, LexToString(TargetType));
+
+			ensureMsgf(false, TEXT("[%s]"), *ErrorMessage);
+			UE_LOG(LogChaosVDEditor, Error, TEXT("[%s]"),*ErrorMessage);
+			return;
+		}
+	}
+}
diff --git a/Engine/Plugins/ChaosVD/Source/ChaosVD/Private/Settings/ChaosVDGeneralSettings.h b/Engine/Plugins/ChaosVD/Source/ChaosVD/Private/Settings/ChaosVDGeneralSettings.h
deleted file mode 100644
diff --git a/Engine/Plugins/ChaosVD/Source/ChaosVD/Private/Widgets/ChaosVDRecordingControls.cpp b/Engine/Plugins/ChaosVD/Source/ChaosVD/Private/Widgets/ChaosVDRecordingControls.cpp
--- a/Engine/Plugins/ChaosVD/Source/ChaosVD/Private/Widgets/ChaosVDRecordingControls.cpp
+++ b/Engine/Plugins/ChaosVD/Source/ChaosVD/Private/Widgets/ChaosVDRecordingControls.cpp
@@ -967,7 +967,7 @@ void FChaosVDRecordingControls::SetSessionRecordingState(bool bIsRecording, cons
 		if (UChaosVDGeneralSettings* Settings = FChaosVDSettingsManager::Get().GetSettingsObject<UChaosVDGeneralSettings>())
 		{
 			RemainingRetries = Settings->MaxConnectionRetries;
-			DataTransportOverrideMode = Settings->DataTransportModeOverride;
+			DataTransportOverrideMode = Settings->GetTransportModeForTarget(SessionInfoRef->BuildTargetType);
 		}
 		else
 		{
@@ -998,7 +998,7 @@ void FChaosVDRecordingControls::SetSessionRecordingState(bool bIsRecording, cons
 
 		FChaosVDStartRecordingCommandMessage RecordingParams;
 		RecordingParams.RecordingMode = EChaosVDRecordingMode::Live;
-		RecordingParams.TransportMode = DataTransportOverrideMode != EChaosVDTransportMode::Invalid ? DataTransportOverrideMode : EChaosVDTransportMode::TraceServer;
+		RecordingParams.TransportMode = DataTransportOverrideMode != EChaosVDTransportMode::Invalid ? DataTransportOverrideMode : EChaosVDTransportMode::Direct;
 
 		check(GLog);
 		bool bOutCanBindAll = false;
diff --git a/Engine/Plugins/ChaosVD/Source/ChaosVD/Public/Settings/ChaosVDGeneralSettings.h b/Engine/Plugins/ChaosVD/Source/ChaosVD/Public/Settings/ChaosVDGeneralSettings.h
new file mode 100644
--- /dev/null
+++ b/Engine/Plugins/ChaosVD/Source/ChaosVD/Public/Settings/ChaosVDGeneralSettings.h
@@ -0,0 +1,79 @@
+// Copyright Epic Games, Inc. All Rights Reserved.
+
+#pragma once
+
+#include "Settings/ChaosVDCoreSettings.h"
+#include "ChaosVDRecordingDetails.h"
+#include "ChaosVDGeneralSettings.generated.h"
+
+/**
+ * General settings that controls how CVD behaves
+ */
+UCLASS(config=ChaosVD, PerObjectConfig)
+class UChaosVDGeneralSettings : public UChaosVDSettingsObjectBase
+{
+	GENERATED_BODY()
+
+public:
+
+	/** Returns the configured transport mode for the provided target type
+	 * @param TargetType Build target type to evaluate
+	 */
+	CHAOSVD_API EChaosVDTransportMode GetTransportModeForTarget(EBuildTargetType TargetType) const;
+
+	/** Sets a specific transport mode for the provided target type
+	 * @param TargetType Build target type to evaluate
+	 * @param NewTransportMode New transport mode to use for the provided target
+	 */
+	CHAOSVD_API void SetTransportModeForTarget(EBuildTargetType TargetType, EChaosVDTransportMode NewTransportMode);
+	
+	/** If true, CVD will only load frames that have solver data in them - Only takes effect before loading a file */
+	UPROPERTY(Config, EditAnywhere, Category = "Session Data Loading")
+	bool bTrimEmptyFrames = true;
+
+	/** How many Game thread frames CVD should queue internally before making them available in the visualization and timeline controls - Only takes effect before loading a file */
+	UPROPERTY(Config, EditAnywhere, Category = "Session Data Loading")
+	int32 MaxGameThreadFramesToQueueNum = 10;
+
+	/** If true, CVD will only load collision geometry that is visible */
+	UPROPERTY(Config, EditAnywhere, Category = "Performance | Streaming")
+	bool bStreamingSystemEnabled = true;
+
+	/** Extent size of the box used for calculate what should be streamed in */
+	UPROPERTY(Config, EditAnywhere, Category = "Performance | Streaming")
+	float StreamingBoxExtentSize = 10000.0f;
+
+	/** If set to true CVD will process any updates to the streaming accel structure in worker threads, in between streaming updates */
+	UPROPERTY(Config, EditAnywhere, Category = "Performance | Streaming")
+	bool bProcessPendingOperationsQueueInWorkerThread = true;
+
+	/** If set to true CVD will keep the scene outliner up to date as the recording is played. If during the recording a
+	 * significant amount of objects are loaded/unloaded, the performance impact will be significant enough to degrade the playback stability. */
+	UPROPERTY(Config, EditAnywhere, Category = "Performance | Scene Outliner")
+	bool bUpdateSceneOutlinerDuringPlayback = false;
+
+	/** How many times CVD will attempt to connect to a live trace or load from file session if the first attempt failed */
+	UPROPERTY(Config, EditAnywhere, Category = Connection)
+	int32 MaxConnectionRetries = 14;
+
+	/** If True, any traces done to memory will also be saved to disk at the same time in the user's document folder (or a user specified folder). This is not saved by design */
+	UPROPERTY(EditAnywhere, Category = Connection)
+	bool bSaveRecordingsToDisk = true;
+
+	/** Custom path where recordings will be saved from. By default, recordings are saved in a Chaos Visual Debugger folder in the documents folder */
+	UPROPERTY(Config, EditAnywhere, Category = Connection)
+	FFilePath DefaultSavePathOverride;
+
+	/** If True, playback will start automatically after a recording started, and we successfully connect to it */
+	UPROPERTY(Config, EditAnywhere, Category = Connection)
+	bool bAutoStartLiveSessionsPlayback = true;
+
+private:
+	/** [DEBUG Setting] Default Data Transport for Game Servers - Used to change the transport mode for the trace data. This is not saved by design */
+	UPROPERTY(EditAnywhere, Category = Connection)
+	EChaosVDTransportMode GameServerDataTransportMode = EChaosVDTransportMode::Direct;
+
+	/** [DEBUG Setting] Default Data Transport for Game Clients - Used to change the transport mode for the trace data. This is not saved by design */
+	UPROPERTY(EditAnywhere, Category = Connection)
+	EChaosVDTransportMode GameClientDataTransportMode = EChaosVDTransportMode::Direct;
+};
diff --git a/Engine/Source/Runtime/Experimental/ChaosVisualDebugger/Public/ChaosVDRecordingDetails.h b/Engine/Source/Runtime/Experimental/ChaosVisualDebugger/Public/ChaosVDRecordingDetails.h
--- a/Engine/Source/Runtime/Experimental/ChaosVisualDebugger/Public/ChaosVDRecordingDetails.h
+++ b/Engine/Source/Runtime/Experimental/ChaosVisualDebugger/Public/ChaosVDRecordingDetails.h
@@ -18,11 +18,11 @@ enum class EChaosVDRecordingMode : uint8
 UENUM()
 enum class EChaosVDTransportMode : uint8
 {
-	Invalid,
+	Invalid UMETA(Hidden),
 	/** The data will be traced directly to a file in the local file system */
-	FileSystem,
+	FileSystem UMETA(Hidden),
 	/** Data will be traced to the selected trace store server */
-	TraceServer,
+	TraceServer UMETA(Hidden),
 	/** Data will be traced directly to the editor, accessible via sockets */
 	Direct,
 	/** Data will be traced via the trace relay system using a custom transport, like normal UE Networking or a direct socket connection */
