[ChaosVD] Added a viewport overlay indicator for the state of CVD's Streaming system.

- The new indicator is a button inside the viewport that allows you to enable and disable streaming directly from the viewport without having to go to the settings menu. This button also provides useful information about the consequences of disabling streaming and when you should do it.

- Changed the default far clipping view value to infinity (CVD's streaming system already controls what should be rendered, and having two system affecting it was causing confusion. Users can still change the value in the Viewport's tool bar if desired).

- Improved the UI feedback in the scene outliner, telling the reason why it is automatically disabled in some scenarios and when it is not 100% up to date and is still processing pending data.

- Updated the Open file menu folder presets. Since the changes made last week to remove trace server support, we don't store trace files in the profiling folder anymore

diff --git a/Engine/Plugins/ChaosVD/Source/ChaosVD/Private/ChaosVDPlaybackViewportClient.cpp b/Engine/Plugins/ChaosVD/Source/ChaosVD/Private/ChaosVDPlaybackViewportClient.cpp
--- a/Engine/Plugins/ChaosVD/Source/ChaosVD/Private/ChaosVDPlaybackViewportClient.cpp
+++ b/Engine/Plugins/ChaosVD/Source/ChaosVD/Private/ChaosVDPlaybackViewportClient.cpp
@@ -29,8 +29,10 @@ FChaosVDPlaybackViewportClient::FChaosVDPlaybackViewportClient(const TSharedPtr<
 		GEngine->OnActorMoving().AddRaw(this, &FChaosVDPlaybackViewportClient::HandleActorMoving);
 	}
 
-	constexpr float DefaultFarClipPlaneOverride = 20000.0f;
+	// Set the default to infinity (the engine treats 0 as inifinity), because CVD's streaming system already handles what should be rendered
+	constexpr float DefaultFarClipPlaneOverride = 0.0f;
 	OverrideFarClipPlane(DefaultFarClipPlaneOverride);
+	OverrideNearClipPlane(UE_KINDA_SMALL_NUMBER);
 }
 
 FChaosVDPlaybackViewportClient::~FChaosVDPlaybackViewportClient()
diff --git a/Engine/Plugins/ChaosVD/Source/ChaosVD/Private/ChaosVDWorldOutlinerMode.cpp b/Engine/Plugins/ChaosVD/Source/ChaosVD/Private/ChaosVDWorldOutlinerMode.cpp
--- a/Engine/Plugins/ChaosVD/Source/ChaosVD/Private/ChaosVDWorldOutlinerMode.cpp
+++ b/Engine/Plugins/ChaosVD/Source/ChaosVD/Private/ChaosVDWorldOutlinerMode.cpp
@@ -17,6 +17,7 @@
 #include "Elements/Framework/TypedElementSelectionSet.h"
 #include "Elements/Interfaces/TypedElementDataStorageCompatibilityInterface.h"
 #include "Settings/ChaosVDGeneralSettings.h"
+#include "Styling/StyleColors.h"
 #include "TEDS/ChaosVDParticleEditorDataFactory.h"
 #include "TEDS/ChaosVDSelectionInterface.h"
 #include "TEDS/ChaosVDStructTypedElementData.h"
@@ -363,4 +364,25 @@ void FChaosVDWorldOutlinerMode::OnDataStorageUpdateCompleted()
 	}
 }
 
+FText FChaosVDWorldOutlinerMode::GetStatusText() const
+{
+	if (CanPopulate())
+	{
+		const int32 PendingUpdated = RowsPendingAddition.Num() + RowsPendingRemoval.Num();
+		if (PendingUpdated > 0)
+		{
+			return FText::Format(LOCTEXT("SceneOutlinerUpdatingMessage", "Processing {0} changes..."), FText::AsNumber(PendingUpdated));
+		}
+
+		return FText::GetEmpty();
+	}
+	
+	return LOCTEXT("SceneOutlinerDisabledWarning", "Outliner updates are disabled during playback or scrubbing");
+}
+
+FSlateColor FChaosVDWorldOutlinerMode::GetStatusTextColor() const
+{
+	return CanPopulate() ? FStyleColors::Foreground : FStyleColors::Warning;
+}
+
 #undef LOCTEXT_NAMESPACE
diff --git a/Engine/Plugins/ChaosVD/Source/ChaosVD/Private/ChaosVDWorldOutlinerMode.h b/Engine/Plugins/ChaosVD/Source/ChaosVD/Private/ChaosVDWorldOutlinerMode.h
--- a/Engine/Plugins/ChaosVD/Source/ChaosVD/Private/ChaosVDWorldOutlinerMode.h
+++ b/Engine/Plugins/ChaosVD/Source/ChaosVD/Private/ChaosVDWorldOutlinerMode.h
@@ -79,6 +79,15 @@ public:
 
 	void OnDataStorageUpdateCompleted();
 
+	virtual bool ShowStatusBar() const override
+	{
+		return true;
+	}
+
+	virtual FText GetStatusText() const override;
+
+	virtual FSlateColor GetStatusTextColor() const override;
+
 private:
 
 	virtual void HandlePostSelectionChange(const UTypedElementSelectionSet* ChangesSelectionSet) override;
diff --git a/Engine/Plugins/ChaosVD/Source/ChaosVD/Private/Widgets/SChaosBrowseTraceFileSourceModal.h b/Engine/Plugins/ChaosVD/Source/ChaosVD/Private/Widgets/SChaosBrowseTraceFileSourceModal.h
--- a/Engine/Plugins/ChaosVD/Source/ChaosVD/Private/Widgets/SChaosBrowseTraceFileSourceModal.h
+++ b/Engine/Plugins/ChaosVD/Source/ChaosVD/Private/Widgets/SChaosBrowseTraceFileSourceModal.h
@@ -19,10 +19,10 @@ enum class EChaosVDBrowseFileModalResponse
 {
 	/** Opens the file picker at the last opened folder */
 	LastOpened,
-	/** Opens the file picker at this project profiling folder (Saved/Profiling) */
-	Profiling,
+	/** Opens the file picker at the currently configure save folder. The documents folder is the default */
+	CurrentSaveFolder,
 	/** Opens the file picker at local trace store folder (This is where live recordings are located) */
-	TraceStore,
+	LocalTraceStore,
 	Cancel UMETA(Hidden)
 };
 
diff --git a/Engine/Plugins/ChaosVD/Source/ChaosVD/Private/Widgets/SChaosVDMainTab.cpp b/Engine/Plugins/ChaosVD/Source/ChaosVD/Private/Widgets/SChaosVDMainTab.cpp
--- a/Engine/Plugins/ChaosVD/Source/ChaosVD/Private/Widgets/SChaosVDMainTab.cpp
+++ b/Engine/Plugins/ChaosVD/Source/ChaosVD/Private/Widgets/SChaosVDMainTab.cpp
@@ -24,6 +24,7 @@
 #include "Components/ChaosVDSceneQueryDataComponent.h"
 #include "DesktopPlatformModule.h"
 #include "Editor.h"
+#include "IAssetViewport.h"
 #include "Framework/Application/SlateApplication.h"
 #include "Framework/Application/SWindowTitleBar.h"
 #include "Framework/Commands/UICommandList.h"
@@ -31,6 +32,7 @@
 #include "HAL/PlatformApplicationMisc.h"
 #include "IDesktopPlatform.h"
 #include "PropertyEditorModule.h"
+#include "SChaosVDPlaybackViewport.h"
 #include "SConnectionStatusBar.h"
 #include "Misc/MessageDialog.h"
 #include "StatusBarSubsystem.h"
@@ -226,6 +228,24 @@ SChaosVDMainTab::~SChaosVDMainTab()
 	CleanUpDisableCPUThrottlingDelegate();
 }
 
+void SChaosVDMainTab::AddViewportOverlayWidget(TSharedRef<SWidget> InOverlaidWidget, int32 ZOrder, TSharedPtr<IAssetViewport> InViewport)
+{
+	if (InViewport)
+	{
+		InViewport->AddOverlayWidget(InOverlaidWidget, ZOrder);
+		return;
+	}
+
+	{
+		TSharedPtr<FChaosVDPlaybackViewportTab> PlaybackViewportTab = GetTabSpawnerInstance<FChaosVDPlaybackViewportTab>(FChaosVDTabID::PlaybackViewport).Pin();
+		TSharedPtr<SChaosVDPlaybackViewport> PlaybackViewportWidget = PlaybackViewportTab ? PlaybackViewportTab->GetPlaybackViewportWidget().Pin() : nullptr;
+		if (PlaybackViewportWidget)
+		{
+			PlaybackViewportWidget->AddViewportOverlayWidget(InOverlaidWidget, ZOrder);
+		}
+	}
+}
+
 void SChaosVDMainTab::ExtendMainToolbar()
 {
 	if (!RecordingControls)
@@ -665,18 +685,20 @@ FReply SChaosVDMainTab::BrowseAndOpenChaosVDRecording()
 				BrowseChaosVDRecordingFromFolder(TEXT(""), LoadingMode);
 				break;
 			}
-		case EChaosVDBrowseFileModalResponse::Profiling:
+		case EChaosVDBrowseFileModalResponse::CurrentSaveFolder:
 			{
-				BrowseChaosVDRecordingFromFolder(*FPaths::ProfilingDir(), LoadingMode);
+				UChaosVDGeneralSettings* Settings = FChaosVDSettingsManager::Get().GetSettingsObject<UChaosVDGeneralSettings>();
+				const FString FolderPath = Settings && !Settings->DefaultSavePathOverride.FilePath.IsEmpty() ? Settings->DefaultSavePathOverride.FilePath : FChaosVDTraceManager::GetDefaultSaveDirectoryPath();
+
+				BrowseChaosVDRecordingFromFolder(*FolderPath, LoadingMode);
 				break;
 			}
-		case EChaosVDBrowseFileModalResponse::TraceStore:
+		case EChaosVDBrowseFileModalResponse::LocalTraceStore:
 			{
-				//TODO: Support remote Trace Stores
 				const FString TraceStorePath = FChaosVDModule::Get().GetTraceManager()->GetLocalTraceStoreDirPath();
 				if (TraceStorePath.IsEmpty())
 				{
-					UE_LOG(LogChaosVDEditor, Error, TEXT("[%s] Failed to access Trace Store..."), ANSI_TO_TCHAR(__FUNCTION__));
+					UE_LOG(LogChaosVDEditor, Error, TEXT("[%hs] Failed to access Trace Store..."), __func__);
 					FMessageDialog::Open(EAppMsgType::Ok, LOCTEXT("OpenTraceStoreFailedMessage", "Failed to access the Trace Store, The default profiling folder will be open. \n Please see the logs for mor details... "));
 				}
 
diff --git a/Engine/Plugins/ChaosVD/Source/ChaosVD/Private/Widgets/SChaosVDPlaybackViewport.cpp b/Engine/Plugins/ChaosVD/Source/ChaosVD/Private/Widgets/SChaosVDPlaybackViewport.cpp
--- a/Engine/Plugins/ChaosVD/Source/ChaosVD/Private/Widgets/SChaosVDPlaybackViewport.cpp
+++ b/Engine/Plugins/ChaosVD/Source/ChaosVD/Private/Widgets/SChaosVDPlaybackViewport.cpp
@@ -26,6 +26,7 @@
 #include "Widgets/Layout/SSeparator.h"
 #include "Widgets/Text/STextBlock.h"
 #include "ToolMenus.h"
+#include "SStreamingStatus.h"
 
 #define LOCTEXT_NAMESPACE "ChaosVisualDebugger"
 
@@ -395,6 +396,27 @@ void SChaosVDPlaybackViewport::OnFocusViewportToSelection()
 	PlaybackViewportClient->FocusOnSelectedObject();
 }
 
+void SChaosVDPlaybackViewport::AddViewportOverlayWidget(TSharedRef<SWidget> InOverlaidWidget, int32 ZOrder)
+{
+	if (ViewportOverlay)
+	{
+		ViewportOverlay->AddSlot(ZOrder)
+		[
+			InOverlaidWidget
+		];
+	}
+}
+
+void SChaosVDPlaybackViewport::PopulateViewportOverlays(TSharedRef<SOverlay> Overlay)
+{
+	SEditorViewport::PopulateViewportOverlays(Overlay);
+
+	Overlay->AddSlot(0)
+	[
+		SNew(Chaos::VD::SStreamingStatus)
+	];
+}
+
 FReply SChaosVDPlaybackViewport::OnDragOver(const FGeometry& MyGeometry, const FDragDropEvent& DragDropEvent)
 {
 	TSharedPtr<SChaosVDMainTab> AsCVDToolkitHost = StaticCastSharedPtr<SChaosVDMainTab>(EditorModeTools->GetToolkitHost());
diff --git a/Engine/Plugins/ChaosVD/Source/ChaosVD/Private/Widgets/SChaosVDPlaybackViewport.h b/Engine/Plugins/ChaosVD/Source/ChaosVD/Private/Widgets/SChaosVDPlaybackViewport.h
--- a/Engine/Plugins/ChaosVD/Source/ChaosVD/Private/Widgets/SChaosVDPlaybackViewport.h
+++ b/Engine/Plugins/ChaosVD/Source/ChaosVD/Private/Widgets/SChaosVDPlaybackViewport.h
@@ -65,7 +65,11 @@ public:
 
 	virtual void OnFocusViewportToSelection() override;
 
+	void AddViewportOverlayWidget(TSharedRef<SWidget> InOverlaidWidget, int32 ZOrder);
+
 protected:
+
+	virtual void PopulateViewportOverlays(TSharedRef<SOverlay> Overlay) override;
 	
 	virtual FReply OnDragOver(const FGeometry& MyGeometry, const FDragDropEvent& DragDropEvent) override;
 	virtual FReply OnDrop(const FGeometry& MyGeometry, const FDragDropEvent& DragDropEvent) override;
diff --git a/Engine/Plugins/ChaosVD/Source/ChaosVD/Private/Widgets/SStreamingStatus.cpp b/Engine/Plugins/ChaosVD/Source/ChaosVD/Private/Widgets/SStreamingStatus.cpp
new file mode 100644
--- /dev/null
+++ b/Engine/Plugins/ChaosVD/Source/ChaosVD/Private/Widgets/SStreamingStatus.cpp
@@ -0,0 +1,143 @@
+﻿// Copyright Epic Games, Inc. All Rights Reserved.
+
+#include "SStreamingStatus.h"
+
+#include "ChaosVDSettingsManager.h"
+#include "Components/VerticalBox.h"
+#include "Misc/MessageDialog.h"
+#include "Settings/ChaosVDGeneralSettings.h"
+#include "Styling/AppStyle.h"
+#include "Widgets/Images/SImage.h"
+#include "Widgets/Input/SButton.h"
+#include "Widgets/Layout/SBorder.h"
+#include "Widgets/Text/STextBlock.h"
+
+#define LOCTEXT_NAMESPACE "ChaosVisualDebugger"
+
+namespace Chaos::VD
+{
+void SStreamingStatus::Construct(const FArguments& InArgs)
+{
+	ChildSlot
+	[
+		SNew(SVerticalBox)
+		+SVerticalBox::Slot()
+		.VAlign(VAlign_Bottom)
+		.HAlign(HAlign_Right)
+		.Padding(4.0f)
+		[
+			SNew(SButton)
+			.ButtonStyle(&FCoreStyle::Get().GetWidgetStyle<FButtonStyle>("HoverHintOnly"))
+			.Cursor(EMouseCursor::Default)
+			.OnClicked_Raw(this, &SStreamingStatus::OpenToggleStreamingStatePopup)
+			.ContentPadding(0.f)
+			[
+				SNew(SBorder)
+				.BorderImage(FAppStyle::Get().GetBrush("FloatingBorder"))
+				.Padding(4.f)
+				[
+					SNew(SHorizontalBox)
+					+SHorizontalBox::Slot()
+					.Padding(FMargin(2, 1, 0, 1))
+					.AutoWidth()
+					.VAlign(VAlign_Center)
+					[
+						SNew(SImage)
+						.Image_Raw(this,&SStreamingStatus::GetStreamingStatusImage)
+						.DesiredSizeOverride(FVector2D(16, 16))
+					]
+					+SHorizontalBox::Slot()
+					.Padding(FMargin(6, 1, 2, 1))
+					.VAlign(VAlign_Center)
+					.HAlign(HAlign_Left)
+					.FillWidth(1)
+					[
+						SNew(STextBlock)
+						.Text_Raw(this, &SStreamingStatus::GetStreamingStatusText)
+						.ToolTipText_Raw(this, &SStreamingStatus::GetStreamingStatusTooltipText)
+						.ColorAndOpacity_Raw(this,&SStreamingStatus::GetStreamingStatusColor)
+					]
+				]
+			]
+		]
+	];
+
+	if (UChaosVDGeneralSettings* GeneralSettings = FChaosVDSettingsManager::Get().GetSettingsObject<UChaosVDGeneralSettings>())
+	{
+		GeneralSettings->OnSettingsChanged().AddSP(this, &SStreamingStatus::HandleSettingsChanged);
+		HandleSettingsChanged(GeneralSettings);
+	}
+}
+
+void SStreamingStatus::HandleSettingsChanged(UObject* SettingsObject)
+{
+	if (UChaosVDGeneralSettings* GeneralSettings = Cast<UChaosVDGeneralSettings>(SettingsObject))
+	{
+		bCachedStreamingEnabledState = GeneralSettings->bStreamingSystemEnabled;
+	}
+}
+
+FText SStreamingStatus::GetStreamingStatusText() const
+{
+	return bCachedStreamingEnabledState ?  LOCTEXT("StreamingStatusEnabledLabel", "Streaming Enabled") : LOCTEXT("StreamingStatusDisabledLabel", "Streaming Disabled");
+}
+
+FSlateColor SStreamingStatus::GetStreamingStatusColor() const
+{
+	return bCachedStreamingEnabledState ? FAppStyle::Get().GetSlateColor("Colors.AccentGreen") : FAppStyle::Get().GetSlateColor("Colors.Warning");
+}
+
+const FSlateBrush* SStreamingStatus::GetStreamingStatusImage() const
+{
+	return bCachedStreamingEnabledState ? FAppStyle::Get().GetBrush("Icons.SuccessWithColor") :  FAppStyle::Get().GetBrush("Icons.WarningWithColor");
+}
+
+FText SStreamingStatus::GetStreamingStatusTooltipText() const
+{
+	return bCachedStreamingEnabledState ?  LOCTEXT("StreamingStatusEnabledToolTip", "CVD's Streaming System is enabled. Only objects within the configured range are visible. This is the recommended config to get the best performance, but you might want disabled it if you are dbugging Level Streaming related issues")
+											: LOCTEXT("StreamingStatusDisabledTooltip", "CVD's Streaming system is disabled. Performance will be degraded");
+}
+
+FReply SStreamingStatus::OpenToggleStreamingStatePopup()
+{
+	UChaosVDGeneralSettings* GeneralSettings = FChaosVDSettingsManager::Get().GetSettingsObject<UChaosVDGeneralSettings>();
+	if (!GeneralSettings)
+	{
+		FMessageDialog::Open(EAppMsgCategory::Error, EAppMsgType::Ok, LOCTEXT("DisableStreamingErrorMessageText", "Failled to obtain the settings object. The streaming state cannot be changed at this moment."));
+		return FReply::Handled();
+	}
+		
+	if (bCachedStreamingEnabledState)
+	{
+		EAppReturnType::Type Response = FMessageDialog::Open(EAppMsgCategory::Warning, EAppMsgType::OkCancel,
+														LOCTEXT("DisableStreamingMessageText", "Do you want to disable CVD's Streaming system? \n\nDisabling the streaming system will degrade performance in large recordings \nand is not advised unless you are trying to debug level streaming related issues \nto see an accurate representation of what was loaded at the time the recording was made."),
+														LOCTEXT("DisableStreamingTitleText", "CVD Streaming System"));
+
+		if (Response == EAppReturnType::Ok)
+		{
+			FProperty* SaveToDiskProperty = FindFieldChecked<FProperty>(UChaosVDGeneralSettings::StaticClass(), "bStreamingSystemEnabled");
+			GeneralSettings->PreEditChange(SaveToDiskProperty);
+			GeneralSettings->bStreamingSystemEnabled = false;
+			GeneralSettings->PostEditChange();
+		}
+	}
+	else
+	{
+		EAppReturnType::Type Response = FMessageDialog::Open(EAppMsgCategory::Info, EAppMsgType::OkCancel,
+												LOCTEXT("EnableStreamingMessageText", "Do you want to enable CVD's Streaming system? \n\nEnabling the streaming system is the recommended setting as it will improve performance significantly, \nbut if you are trying to debug level streaming related issues you might need to keep it disabled \nto see an accurate representation of what was loaded at the time the recording was made."),
+												LOCTEXT("EnableStreamingTitleText", "CVD Streaming System"));
+
+		if (Response == EAppReturnType::Ok)
+		{
+			FProperty* SaveToDiskProperty = FindFieldChecked<FProperty>(UChaosVDGeneralSettings::StaticClass(), "bStreamingSystemEnabled");
+			GeneralSettings->PreEditChange(SaveToDiskProperty);
+			GeneralSettings->bStreamingSystemEnabled = true;
+			GeneralSettings->PostEditChange();
+		}
+	}
+	
+	return FReply::Handled();
+}
+}
+
+#undef LOCTEXT_NAMESPACE
diff --git a/Engine/Plugins/ChaosVD/Source/ChaosVD/Private/Widgets/SStreamingStatus.h b/Engine/Plugins/ChaosVD/Source/ChaosVD/Private/Widgets/SStreamingStatus.h
new file mode 100644
--- /dev/null
+++ b/Engine/Plugins/ChaosVD/Source/ChaosVD/Private/Widgets/SStreamingStatus.h
@@ -0,0 +1,34 @@
+﻿// Copyright Epic Games, Inc. All Rights Reserved.
+
+#pragma once
+
+#include "Widgets/SCompoundWidget.h"
+
+namespace Chaos::VD
+{
+/**
+ * Widget to be used as overlay in the main viewport to indicate the user the streaming system status
+ */
+class SStreamingStatus : public SCompoundWidget
+{
+public:
+	SLATE_BEGIN_ARGS(SStreamingStatus)
+	{
+	}
+	SLATE_END_ARGS()
+
+	void Construct(const FArguments& InArgs);
+
+private:
+
+	void HandleSettingsChanged(UObject* SettingsObject);
+	FText GetStreamingStatusText() const;
+	FSlateColor GetStreamingStatusColor() const;
+	const FSlateBrush* GetStreamingStatusImage() const;
+	FText GetStreamingStatusTooltipText() const;
+	FReply OpenToggleStreamingStatePopup();
+
+	bool bCachedStreamingEnabledState = false;
+};
+}
+
diff --git a/Engine/Plugins/ChaosVD/Source/ChaosVD/Public/Widgets/SChaosVDMainTab.h b/Engine/Plugins/ChaosVD/Source/ChaosVD/Public/Widgets/SChaosVDMainTab.h
--- a/Engine/Plugins/ChaosVD/Source/ChaosVD/Public/Widgets/SChaosVDMainTab.h
+++ b/Engine/Plugins/ChaosVD/Source/ChaosVD/Public/Widgets/SChaosVDMainTab.h
@@ -172,6 +172,8 @@ public:
 		return RecordingControls;
 	}
 
+	CHAOSVD_API virtual void AddViewportOverlayWidget(TSharedRef<SWidget> InOverlaidWidget, int32 ZOrder, TSharedPtr<IAssetViewport> InViewport = nullptr) override;
+
 private:
 
 	void ExtendMainToolbar();
