[NetPhysics] Fix bug where TickOffset didn't get updated in CVD.

diff --git a/Engine/Plugins/ChaosVD/Source/ChaosVD/Private/Trace/ChaosVDTraceAnalyzer.cpp b/Engine/Plugins/ChaosVD/Source/ChaosVD/Private/Trace/ChaosVDTraceAnalyzer.cpp
--- a/Engine/Plugins/ChaosVD/Source/ChaosVD/Private/Trace/ChaosVDTraceAnalyzer.cpp
+++ b/Engine/Plugins/ChaosVD/Source/ChaosVD/Private/Trace/ChaosVDTraceAnalyzer.cpp
@@ -316,7 +316,7 @@ bool FChaosVDTraceAnalyzer::OnEvent(uint16 RouteId, EStyle Style, const FOnEvent
 
 			if (SolverID != INDEX_NONE)
 			{
-				ChaosVDTraceProvider->GetCurrentTickOffsetsBySolverID().FindOrAdd(SolverID, TickOffset);
+				ChaosVDTraceProvider->GetCurrentTickOffsetsBySolverID().FindOrAdd(SolverID) = TickOffset;
 			}
 
 			break;
