Change Browse Live Sessions Menu to ToolMenu

diff --git a/Engine/Plugins/ChaosVD/Source/ChaosVD/Private/Widgets/SChaosVDMainTab.cpp b/Engine/Plugins/ChaosVD/Source/ChaosVD/Private/Widgets/SChaosVDMainTab.cpp
--- a/Engine/Plugins/ChaosVD/Source/ChaosVD/Private/Widgets/SChaosVDMainTab.cpp
+++ b/Engine/Plugins/ChaosVD/Source/ChaosVD/Private/Widgets/SChaosVDMainTab.cpp
@@ -235,7 +235,7 @@ void SChaosVDMainTab::AddViewportOverlayWidget(TSharedRef<SWidget> InOverlaidWid
 		InViewport->AddOverlayWidget(InOverlaidWidget, ZOrder);
 		return;
 	}
-
+	
 	{
 		TSharedPtr<FChaosVDPlaybackViewportTab> PlaybackViewportTab = GetTabSpawnerInstance<FChaosVDPlaybackViewportTab>(FChaosVDTabID::PlaybackViewport).Pin();
 		TSharedPtr<SChaosVDPlaybackViewport> PlaybackViewportWidget = PlaybackViewportTab ? PlaybackViewportTab->GetPlaybackViewportWidget().Pin() : nullptr;
@@ -625,38 +625,59 @@ void SChaosVDMainTab::CombineOpenSessions()
 
 void SChaosVDMainTab::GenerateMainWindowMenu()
 {
-	FMenuBarBuilder MenuBarBuilder = FMenuBarBuilder(GlobalCommandList);
-	MenuBarBuilder.AddPullDownMenu(
+	FToolMenuContext MenuContext(GlobalCommandList);
+
+	static const FName MainMenuName("ChaosVD.MainMenuBar");
+	UToolMenus* ToolMenus = UToolMenus::Get();
+
+	if (ToolMenus->IsMenuRegistered(MainMenuName))
+	{
+		return;
+	}
+
+	UToolMenu* MenuBar = ToolMenus->RegisterMenu(MainMenuName, NAME_None, EMultiBoxType::MenuBar);
+	MenuBar->StyleName = FName("WindowMenuBar");
+
+	UToolMenu* FileEntry = MenuBar->AddSubMenu(
+		"MainMenuBar",
+		NAME_None,
+		"File",
 		LOCTEXT("FileMenuLabel", "File"),
-		FText::GetEmpty(),
-		FNewMenuDelegate::CreateSPLambda(this, [this](FMenuBuilder& MenuBuilder)
-		{
-			MenuBuilder.AddSubMenu(LOCTEXT("RecentFilesMenuLabel", "Recent Files"), LOCTEXT("RecentFilesMenuLabelToolTip", "Shows a list of recently used CVD Files"), FNewMenuDelegate::CreateSP(this, &SChaosVDMainTab::GenerateRecentFilesMenu));
-
-			const FChaosVDCommands& Commands = FChaosVDCommands::Get();
-			MenuBuilder.BeginSection(NAME_None, LOCTEXT("LegacySessionSectionLabel", "Legacy Sessions Management"));
-				MenuBuilder.AddMenuEntry(Commands.BrowseLiveSessions, NAME_None);
-				MenuBuilder.AddMenuEntry(Commands.DisconnectFromLiveSession, NAME_None);
-			MenuBuilder.EndSection();
-		}),
-		"File"
+		FText::GetEmpty()
 	);
 
-	MenuBarBuilder.AddPullDownMenu(
+	UToolMenu* RecentSubMenu = FileEntry->AddSubMenu("File", NAME_None, "Recent Files", LOCTEXT("RecentFilesMenuLabel", "Recent Files"), LOCTEXT("RecentFilesMenuLabelToolTip", "Shows a list of recently used CVD Files"));
+	RecentSubMenu->AddDynamicSection(NAME_None, FNewSectionConstructChoice(FNewToolMenuDelegateLegacy::CreateSP(this, &SChaosVDMainTab::GenerateRecentFilesMenu)));
+
+	const FChaosVDCommands& Commands = FChaosVDCommands::Get();
+	FToolMenuSection& SessionsSection = FileEntry->AddSection("LegacySessionSection", LOCTEXT("LegacySessionSectionLabel", "Legacy Sessions Management"));
+	SessionsSection.AddMenuEntry(Commands.BrowseLiveSessions);
+	SessionsSection.AddMenuEntry(Commands.DisconnectFromLiveSession);
+
+	UToolMenu* WindowEntry = MenuBar->AddSubMenu(
+		"MainMenuBar",
+		NAME_None,
+		"Window",
 		LOCTEXT("WindowMenuLabel", "Window"),
-		FText::GetEmpty(),
-		FNewMenuDelegate::CreateLambda([this](FMenuBuilder& MenuBuilder)
-		{
-			TabManager->PopulateLocalTabSpawnerMenu(MenuBuilder);
-		}),
-		"Window"
+		FText::GetEmpty()
 	);
 
-	TabManager->SetAllowWindowMenuBar(true);
-	TabManager->SetMenuMultiBox(MenuBarBuilder.GetMultiBox(), MenuBarBuilder.MakeWidget());
+	WindowEntry->AddDynamicSection(NAME_None, FNewSectionConstructChoice(FNewToolMenuDelegateLegacy::CreateSPLambda(this, [this](FMenuBuilder& MenuBuilder, UToolMenu* Menu)
+	{
+		TabManager->PopulateLocalTabSpawnerMenu(MenuBuilder);
+	})));
+
+	TSharedRef<SWidget> MenuBarWidget = UToolMenus::Get()->GenerateWidget(MainMenuName, MenuContext);
+	if (MenuBarWidget != SNullWidget::NullWidget)
+	{
+		TSharedRef<SMultiBoxWidget> MultiBoxWidget = StaticCastSharedRef<SMultiBoxWidget>(MenuBarWidget);
+
+		TabManager->SetAllowWindowMenuBar(true);
+		TabManager->SetMenuMultiBox(ConstCastSharedRef<FMultiBox>(MultiBoxWidget->GetMultiBox()), MultiBoxWidget);
+	}
 }
 
-void SChaosVDMainTab::GenerateRecentFilesMenu(FMenuBuilder& MenuBuilder)
+void SChaosVDMainTab::GenerateRecentFilesMenu(FMenuBuilder& MenuBuilder, UToolMenu* Menu)
 {
 	if (UChaosVDMiscSettings* MiscSettings = FChaosVDSettingsManager::Get().GetSettingsObject<UChaosVDMiscSettings>())
 	{
diff --git a/Engine/Plugins/ChaosVD/Source/ChaosVD/Public/Widgets/SChaosVDMainTab.h b/Engine/Plugins/ChaosVD/Source/ChaosVD/Public/Widgets/SChaosVDMainTab.h
--- a/Engine/Plugins/ChaosVD/Source/ChaosVD/Public/Widgets/SChaosVDMainTab.h
+++ b/Engine/Plugins/ChaosVD/Source/ChaosVD/Public/Widgets/SChaosVDMainTab.h
@@ -7,6 +7,7 @@
 #include "Templates/SharedPointer.h"
 #include "Toolkits/IToolkitHost.h"
 #include "Widgets/SCompoundWidget.h"
+#include "ToolMenus.h"
 
 #include "SChaosVDMainTab.generated.h"
 
@@ -207,7 +208,7 @@ private:
 	void CombineOpenSessions();
 
 	void GenerateMainWindowMenu();
-	void GenerateRecentFilesMenu(FMenuBuilder& MenuBuilder);
+	void GenerateRecentFilesMenu(FMenuBuilder& MenuBuilder, UToolMenu* Menu);
 
 	FReply BrowseAndOpenChaosVDRecording();
 
