Fix CVD Menu Not Created on Second Open

diff --git a/Engine/Plugins/ChaosVD/Source/ChaosVD/Private/Widgets/SChaosVDMainTab.cpp b/Engine/Plugins/ChaosVD/Source/ChaosVD/Private/Widgets/SChaosVDMainTab.cpp
--- a/Engine/Plugins/ChaosVD/Source/ChaosVD/Private/Widgets/SChaosVDMainTab.cpp
+++ b/Engine/Plugins/ChaosVD/Source/ChaosVD/Private/Widgets/SChaosVDMainTab.cpp
@@ -80,7 +80,8 @@
 
 #define LOCTEXT_NAMESPACE "ChaosVisualDebugger"
 
-const FName SChaosVDMainTab::MainToolBarName = FName("ChaosVD.MainToolBar");
+const FName SChaosVDMainTab::MainToolBarName = FName("ChaosVD.MainToolBar"); 
+const FName SChaosVDMainTab::MainMenuName = FName("ChaosVD.MainMenuBar");
 
 void SChaosVDMainTab::Construct(const FArguments& InArgs, TSharedPtr<FChaosVDEngine> InChaosVDEngine)
 {
@@ -226,6 +227,14 @@ SChaosVDMainTab::~SChaosVDMainTab()
 
 	FChaosVDExtensionsManager::Get().OnExtensionRegistered().RemoveAll(this);
 	CleanUpDisableCPUThrottlingDelegate();
+
+	if (UToolMenus* ToolMenus = UToolMenus::Get())
+	{
+		if(ToolMenus->IsMenuRegistered(MainMenuName))
+		{
+			ToolMenus->RemoveMenu(MainMenuName);
+		}
+	}
 }
 
 void SChaosVDMainTab::AddViewportOverlayWidget(TSharedRef<SWidget> InOverlaidWidget, int32 ZOrder, TSharedPtr<IAssetViewport> InViewport)
@@ -627,15 +636,19 @@ void SChaosVDMainTab::GenerateMainWindowMenu()
 {
 	FToolMenuContext MenuContext(GlobalCommandList);
 
-	static const FName MainMenuName("ChaosVD.MainMenuBar");
 	UToolMenus* ToolMenus = UToolMenus::Get();
 
-	if (ToolMenus->IsMenuRegistered(MainMenuName))
+	UToolMenu* MenuBar;
+
+	if (!ToolMenus->IsMenuRegistered(MainMenuName))
 	{
-		return;
+		MenuBar = ToolMenus->RegisterMenu(MainMenuName, NAME_None, EMultiBoxType::MenuBar);
+	}
+	else
+	{
+		MenuBar = ToolMenus->ExtendMenu(MainMenuName);
 	}
 
-	UToolMenu* MenuBar = ToolMenus->RegisterMenu(MainMenuName, NAME_None, EMultiBoxType::MenuBar);
 	MenuBar->StyleName = FName("WindowMenuBar");
 
 	UToolMenu* FileEntry = MenuBar->AddSubMenu(
diff --git a/Engine/Plugins/ChaosVD/Source/ChaosVD/Public/Widgets/SChaosVDMainTab.h b/Engine/Plugins/ChaosVD/Source/ChaosVD/Public/Widgets/SChaosVDMainTab.h
--- a/Engine/Plugins/ChaosVD/Source/ChaosVD/Public/Widgets/SChaosVDMainTab.h
+++ b/Engine/Plugins/ChaosVD/Source/ChaosVD/Public/Widgets/SChaosVDMainTab.h
@@ -261,6 +261,8 @@ private:
 	TSharedPtr<FChaosVDRecordingControls> RecordingControls;
 
 	static const FName MainToolBarName;
+
+	static const FName MainMenuName;
 };
 
 template <typename TabType>
