Make CVD Scene Outliner update during playback by default

diff --git a/Engine/Plugins/ChaosVD/Source/ChaosVD/Public/Settings/ChaosVDGeneralSettings.h b/Engine/Plugins/ChaosVD/Source/ChaosVD/Public/Settings/ChaosVDGeneralSettings.h
--- a/Engine/Plugins/ChaosVD/Source/ChaosVD/Public/Settings/ChaosVDGeneralSettings.h
+++ b/Engine/Plugins/ChaosVD/Source/ChaosVD/Public/Settings/ChaosVDGeneralSettings.h
@@ -50,7 +50,7 @@ public:
 	/** If set to true CVD will keep the scene outliner up to date as the recording is played. If during the recording a
 	 * significant amount of objects are loaded/unloaded, the performance impact will be significant enough to degrade the playback stability. */
 	UPROPERTY(Config, EditAnywhere, Category = "Performance | Scene Outliner")
-	bool bUpdateSceneOutlinerDuringPlayback = false;
+	bool bUpdateSceneOutlinerDuringPlayback = true;
 
 	/** How many times CVD will attempt to connect to a live trace or load from file session if the first attempt failed */
 	UPROPERTY(Config, EditAnywhere, Category = Connection)
