Add camera tracing feature to Chaos Visual Debugger

#Synced-CL 48457324

diff --git a/Engine/Plugins/ChaosVD/Source/ChaosVD/Private/Actors/ChaosVDSolverInfoActor.h b/Engine/Plugins/ChaosVD/Source/ChaosVD/Private/Actors/ChaosVDSolverInfoActor.h
--- a/Engine/Plugins/ChaosVD/Source/ChaosVD/Private/Actors/ChaosVDSolverInfoActor.h
+++ b/Engine/Plugins/ChaosVD/Source/ChaosVD/Private/Actors/ChaosVDSolverInfoActor.h
@@ -16,6 +16,8 @@ struct FChaosVDParticleDataWrapper;
 class UChaosVDParticleDataComponent;
 class UChaosVDSceneQueryDataComponent;
 class UChaosVDSolverCollisionDataComponent;
+class UChaosVDCameraDataComponent;
+struct FChaosVDCameraDataWrapper;
 struct FChaosVDGameFrameData;
 
 enum class EChaosVDParticleType : uint8;
diff --git a/Engine/Plugins/ChaosVD/Source/ChaosVD/Private/ChaosVDPlaybackController.cpp b/Engine/Plugins/ChaosVD/Source/ChaosVD/Private/ChaosVDPlaybackController.cpp
--- a/Engine/Plugins/ChaosVD/Source/ChaosVD/Private/ChaosVDPlaybackController.cpp
+++ b/Engine/Plugins/ChaosVD/Source/ChaosVD/Private/ChaosVDPlaybackController.cpp
@@ -899,6 +899,11 @@ bool FChaosVDPlaybackController::Tick(float DeltaTime)
 				// If this a live session and automatic playback is enabled, start playing
 				if (IsPlayingLiveSession() && bStartLiveSessionPlaybackAutomatically)
 				{
+					constexpr int32 GameFrameToLoad = 0;
+					constexpr int32 SolverStageToLoad = 0;
+					//Ensure Data Container Actors are created in time before the first frame callback.
+					GoToTrackFrameAndSync(PlaybackSelfInstigatorID, EChaosVDTrackType::Game, GameTrackID, GameFrameToLoad, SolverStageToLoad);
+
 					HandleFramePlaybackControlInput(EChaosVDPlaybackButtonsID::Play, GameTrackInfo.ToSharedRef(), PlaybackSelfInstigatorID);
 					bPlayedFirstFrame = true;
 				}
@@ -910,6 +915,17 @@ bool FChaosVDPlaybackController::Tick(float DeltaTime)
 					GoToTrackFrameAndSync(PlaybackSelfInstigatorID, EChaosVDTrackType::Game, GameTrackID, GameFrameToLoad, SolverStageToLoad);
 					bPlayedFirstFrame = true;
 				}
+
+				if (bPlayedFirstFrame)
+				{
+					const TWeakPtr<FChaosVDPlaybackController> InControllerWeakPtr = DoesSharedInstanceExist() ? AsWeak() : nullptr;
+
+					FChaosVDExtensionsManager::Get().EnumerateExtensions([InControllerWeakPtr](const TSharedRef<FChaosVDExtension>& Extension)
+					{
+						Extension->HandleRecordingFirstFrameLoaded(InControllerWeakPtr);
+						return true;
+					});
+				}
 			}
 		}
 
diff --git a/Engine/Plugins/ChaosVD/Source/ChaosVD/Private/ChaosVDPlaybackViewportClient.cpp b/Engine/Plugins/ChaosVD/Source/ChaosVD/Private/ChaosVDPlaybackViewportClient.cpp
--- a/Engine/Plugins/ChaosVD/Source/ChaosVD/Private/ChaosVDPlaybackViewportClient.cpp
+++ b/Engine/Plugins/ChaosVD/Source/ChaosVD/Private/ChaosVDPlaybackViewportClient.cpp
@@ -1,6 +1,6 @@
 ﻿// Copyright Epic Games, Inc. All Rights Reserved.
 
-#include "ChaosVDPlaybackViewportClient.h"
+#include "ChaosVD/Public/ChaosVDPlaybackViewportClient.h"
 
 #include "CanvasTypes.h"
 #include "Actors/ChaosVDSolverInfoActor.h"
diff --git a/Engine/Plugins/ChaosVD/Source/ChaosVD/Private/ChaosVDPlaybackViewportClient.h b/Engine/Plugins/ChaosVD/Source/ChaosVD/Private/ChaosVDPlaybackViewportClient.h
deleted file mode 100644
diff --git a/Engine/Plugins/ChaosVD/Source/ChaosVD/Private/Widgets/SChaosVDMainTab.cpp b/Engine/Plugins/ChaosVD/Source/ChaosVD/Private/Widgets/SChaosVDMainTab.cpp
--- a/Engine/Plugins/ChaosVD/Source/ChaosVD/Private/Widgets/SChaosVDMainTab.cpp
+++ b/Engine/Plugins/ChaosVD/Source/ChaosVD/Private/Widgets/SChaosVDMainTab.cpp
@@ -210,6 +210,12 @@ void SChaosVDMainTab::Construct(const FArguments& InArgs, TSharedPtr<FChaosVDEng
 	TabManager->TryInvokeTab(FChaosVDTabID::SolversTrack);
 	TabManager->TryInvokeTab(FChaosVDTabID::DetailsPanel);
 
+	FChaosVDExtensionsManager::Get().EnumerateExtensions([this](const TSharedRef<FChaosVDExtension>& Extension)
+	{
+		Extension->PostMainTabInitialization(StaticCastSharedRef<SChaosVDMainTab>(AsShared()));
+		return true;
+	});
+
 	SetUpDisableCPUThrottlingDelegate();
 
 	if (TSharedPtr<FChaosVDRemoteSessionsManager> RemoteSessionManager = FChaosVDEngineEditorBridge::Get().GetRemoteSessionsManager())
diff --git a/Engine/Plugins/ChaosVD/Source/ChaosVD/Public/ChaosVDPlaybackController.h b/Engine/Plugins/ChaosVD/Source/ChaosVD/Public/ChaosVDPlaybackController.h
--- a/Engine/Plugins/ChaosVD/Source/ChaosVD/Public/ChaosVDPlaybackController.h
+++ b/Engine/Plugins/ChaosVD/Source/ChaosVD/Public/ChaosVDPlaybackController.h
@@ -22,6 +22,7 @@ class FString;
 
 enum class EChaosVDPlaybackButtonsID : uint8;
 
+DECLARE_MULTICAST_DELEGATE_OneParam(FChaosVDInitialFirstFramePlayed, TWeakPtr<FChaosVDPlaybackController>)
 DECLARE_MULTICAST_DELEGATE_OneParam(FChaosVDPlaybackControllerUpdated, TWeakPtr<FChaosVDPlaybackController>)
 DECLARE_MULTICAST_DELEGATE_ThreeParams(FChaosVDPlaybackControllerFrameUpdated, TWeakPtr<FChaosVDPlaybackController>, TWeakPtr<const FChaosVDTrackInfo>, FGuid);
 
@@ -289,6 +290,11 @@ public:
 	{
 		return ControllerFrameUpdatedDelegate;
 	}
+
+	FChaosVDInitialFirstFramePlayed& OnFirstFramePlayed()
+	{
+		return ControllerFirstFramePlayedDelegate;
+	}
 	
 	CHAOSVD_API virtual bool Tick(float DeltaTime) override;
 
@@ -511,6 +517,9 @@ protected:
 	/** Delegate called when the data in a track changes */
 	FChaosVDPlaybackControllerFrameUpdated ControllerFrameUpdatedDelegate;
 
+	/** Delegate called when the first frame is played on recording load */
+	FChaosVDInitialFirstFramePlayed ControllerFirstFramePlayedDelegate;
+
 	/** Set to true when the recording data controlled by this Playback Controller is updated, the update delegate will be called on the GT */
 	std::atomic<bool> bHasPendingGTUpdateBroadcast;
 
diff --git a/Engine/Plugins/ChaosVD/Source/ChaosVD/Public/ChaosVDPlaybackViewportClient.h b/Engine/Plugins/ChaosVD/Source/ChaosVD/Public/ChaosVDPlaybackViewportClient.h
new file mode 100644
--- /dev/null
+++ b/Engine/Plugins/ChaosVD/Source/ChaosVD/Public/ChaosVDPlaybackViewportClient.h
@@ -0,0 +1,89 @@
+﻿// Copyright Epic Games, Inc. All Rights Reserved.
+
+#pragma once
+#include "EditorViewportClient.h"
+
+class SChaosVDMainTab;
+class FComponentVisualizer;
+class UChaosVDCoreSettings;
+struct FChaosVDGameFrameData;
+class FChaosVDScene;
+enum class EChaosVDActorTrackingMode;
+
+/** Client viewport class used for to handle a Chaos Visual Debugger world Interaction/Rendering.
+ * It re-routes interaction events to our Chaos VD scene
+ */
+class FChaosVDPlaybackViewportClient : public FEditorViewportClient
+{
+public:
+
+	FChaosVDPlaybackViewportClient(const TSharedPtr<FEditorModeTools>& InModeTools, const TSharedPtr<SEditorViewport>& InEditorViewportWidget);
+	virtual ~FChaosVDPlaybackViewportClient() override;
+
+	virtual void ProcessClick(FSceneView& View, HHitProxy* HitProxy, FKey Key, EInputEvent Event, uint32 HitX, uint32 HitY) override;
+
+	void SetScene(TWeakPtr<FChaosVDScene> InScene);
+
+	virtual UWorld* GetWorld() const override
+	{
+		return CVDWorld;
+	}
+
+	virtual void Draw(const FSceneView* View, FPrimitiveDrawInterface* PDI) override;
+	virtual void DrawCanvas(FViewport& InViewport, FSceneView& View, FCanvas& Canvas) override;
+
+	CHAOSVD_API void ToggleObjectTrackingIfSelected();
+
+	bool IsAutoTrackingSelectedObject() const
+	{
+		return bAutoTrackSelectedObject;
+	}
+
+	CHAOSVD_API void SetAutoTrackingViewDistance(float NewDistance);
+
+	float GetAutoTrackingViewDistance() const
+	{
+		return TrackingViewDistance;
+	}
+
+	void GoToLocation(const FVector& InLocation);
+
+	void UpdateObjectTracking();
+
+	void FocusOnSelectedObject();
+
+	FBox GetSelectionBounds() const;
+	
+	virtual void UpdateMouseDelta() override;
+
+	void HandleCVDSceneUpdated();
+
+	bool GetCanSelectTranslucentGeometry() const
+	{
+		return bAllowTranslucentHitProxies;
+	}
+
+	void SetCanSelectTranslucentGeometry(bool bCanSelect);
+	void ToggleCanSelectTranslucentGeometry();
+
+protected:
+
+	virtual bool Internal_InputKey(const FInputKeyEventArgs& EventArgs) override;
+
+private:
+
+	void ExecuteVisualizationCallback(const TSharedRef<SChaosVDMainTab>& InMainTabToolkitHostRef,
+		const TSharedRef<FChaosVDScene>& SceneRef, const TFunction<void(const UActorComponent*, const TSharedRef<FComponentVisualizer>&)>& VisitorCallback);
+
+	void HandleFocusRequest(FBox BoxToFocusOn);
+	void HandleActorMoving(AActor* MovedActor) const;
+
+	FDelegateHandle FocusRequestDelegateHandle;
+	UWorld* CVDWorld;
+	TWeakPtr<FChaosVDScene> CVDScene;
+
+	bool bAutoTrackSelectedObject = false;
+	float TrackingViewDistance = 120.0f;
+
+	bool bAllowTranslucentHitProxies = true;
+};
diff --git a/Engine/Plugins/ChaosVD/Source/ChaosVD/Public/ExtensionsSystem/ChaosVDExtension.h b/Engine/Plugins/ChaosVD/Source/ChaosVD/Public/ExtensionsSystem/ChaosVDExtension.h
--- a/Engine/Plugins/ChaosVD/Source/ChaosVD/Public/ExtensionsSystem/ChaosVDExtension.h
+++ b/Engine/Plugins/ChaosVD/Source/ChaosVD/Public/ExtensionsSystem/ChaosVDExtension.h
@@ -48,6 +48,9 @@ public:
 	/** Registers all available Tab Spawner instances in this extension, if any */
 	virtual void RegisterCustomTabSpawners(const TSharedRef<SChaosVDMainTab>& InParentTabWidget) {}
 
+	/** Initialization by Extension after MainTab is constructed */
+	virtual void PostMainTabInitialization(const TSharedRef<SChaosVDMainTab>& InParentTabWidget) {}
+
 	/** Injects any customization implementations for CVD's details panels*/
 	virtual void SetCustomPropertyLayouts(IDetailsView* DetailsView, TSharedRef<SChaosVDMainTab> InCVDToolKit) {}
 
@@ -57,6 +60,9 @@ public:
 	/** Handles playback state changes on a CVD Instance */
 	virtual void HandleControllerTrackFrameUpdated(TWeakPtr<FChaosVDPlaybackController> InController, TWeakPtr<const FChaosVDTrackInfo> UpdatedTrackInfo, FGuid InstigatorGuid) {}
 
+	/** Handles initial first frame playback after a recording is loading */
+	virtual void HandleRecordingFirstFrameLoaded(TWeakPtr<FChaosVDPlaybackController> InController) {}
+
 protected:
 	FName ExtensionName;
 };
diff --git a/Engine/Plugins/ChaosVD/Source/ChaosVDBuiltInExtensions/Private/CameraTraces/ChaosVDCameraTracesExtension.cpp b/Engine/Plugins/ChaosVD/Source/ChaosVDBuiltInExtensions/Private/CameraTraces/ChaosVDCameraTracesExtension.cpp
new file mode 100644
--- /dev/null
+++ b/Engine/Plugins/ChaosVD/Source/ChaosVDBuiltInExtensions/Private/CameraTraces/ChaosVDCameraTracesExtension.cpp
@@ -0,0 +1,395 @@
+// Copyright Epic Games, Inc. All Rights Reserved.
+
+#include "ChaosVDCameraTracesExtension.h"
+#include "Actors/ChaosVDDataContainerBaseActor.h"
+#include "Trace/ChaosVDCameraDataProcessor.h"
+#include "Settings/ChaosVDCameraDataSettings.h"
+#include "Visualizers/ChaosVDCameraDataComponentVisualizer.h"
+#include "Widgets/SChaosVDMainTab.h"
+#include "ChaosVD/Public/ChaosVDPlaybackViewportClient.h"
+#include "EditorModeManager.h"
+#include "EditorViewportClient.h"
+#include "SEditorViewport.h"
+#include "ToolMenu.h"
+#include "ToolMenus.h"
+#include "ToolMenuEntry.h"
+#include "ToolMenuSection.h"
+#include "Widgets/SWidget.h"
+#include "ChaosVDPlaybackController.h"
+#include "ChaosVDScene.h"
+#include "Components/ChaosVdCameraDataComponent.h"
+#include "ChaosVDEngine.h"
+
+#define LOCTEXT_NAMESPACE "ChaosVisualDebugger"
+
+const FName FChaosVDCameraTracesExtension::ChaosViewportType = FName("SChaosVDPlaybackViewport");
+const FName FChaosVDCameraTracesExtension::CameraSeparatorName = FName("ChaosVDCameraMenuSeparator");
+
+FChaosVDCameraTracesExtension::FChaosVDCameraTracesExtension() : FChaosVDExtension()
+{
+	DataComponentsClasses.Add(UChaosVDCameraDataComponent::StaticClass());
+
+	ExtensionName = FName(TEXT("FChaosVDCameraTracesExtension"));
+}
+
+FChaosVDCameraTracesExtension::~FChaosVDCameraTracesExtension()
+{
+	DataComponentsClasses.Reset();
+}
+
+void FChaosVDCameraTracesExtension::RegisterDataProcessorsInstancesForProvider(const TSharedRef<FChaosVDTraceProvider>& InTraceProvider)
+{
+	FChaosVDExtension::RegisterDataProcessorsInstancesForProvider(InTraceProvider);
+	
+	TSharedPtr<FChaosVDCameraDataProcessor> CameraDataProcessor = MakeShared<FChaosVDCameraDataProcessor>();
+	CameraDataProcessor->SetTraceProvider(InTraceProvider);
+    InTraceProvider->RegisterDataProcessor(CameraDataProcessor);
+
+}
+
+void FChaosVDCameraTracesExtension::HandleRecordingFirstFrameLoaded(TWeakPtr<FChaosVDPlaybackController> InController)
+{
+	if (UChaosVDCameraDataSettings* EditorSettings = FChaosVDSettingsManager::Get().GetSettingsObject<UChaosVDCameraDataSettings>())
+	{
+		if(EditorSettings->bViewportStartAtCameraTrace)
+		{
+			if (TSharedPtr<FChaosVDPlaybackController> Playback = InController.Pin())
+			{
+				if (Playback->IsRecordingLoaded())
+				{
+					if (FChaosVDPlaybackViewportClient* ViewportClient = GetPlaybackViewportClient())
+					{
+						FViewportCameraTransform& ViewTransform = ViewportClient->GetViewTransform();
+						FVector Position;
+						FQuat Rotation;
+						float FieldOfView;
+
+						if (!ViewportClient->IsAutoTrackingSelectedObject() && FindCameraData(Position, Rotation, FieldOfView, ECameraFindFlags::FindFirst))
+						{
+							ViewTransform.SetLocation(Position);
+							ViewTransform.SetRotation(FRotator(Rotation));
+						}
+					}
+				}
+			}
+		}
+	}
+}
+
+TConstArrayView<TSubclassOf<UActorComponent>> FChaosVDCameraTracesExtension::GetSolverDataComponentsClasses()
+{
+	return DataComponentsClasses;
+}
+
+bool FChaosVDCameraTracesExtension::FindCameraData(FVector& Position, FQuat& Rotation, float& FieldOfView, ECameraFindFlags CameraFindFlags = ECameraFindFlags::Default) const
+{
+	if (SelectedCamera.IsSet() || CameraFindFlags == ECameraFindFlags::FindFirst)
+	{
+		if (TSharedPtr<FChaosVDScene> ScenePtr = WeakScenePtr.Pin())
+		{
+			TConstArrayView<TObjectPtr<AChaosVDDataContainerBaseActor>> DataContainerActors = ScenePtr->GetDataContainerActorsView();
+			TInlineComponentArray<const UChaosVDCameraDataComponent*> CameraComponents;
+
+			for (const TObjectPtr<AChaosVDDataContainerBaseActor>& DataContainerActor : DataContainerActors)
+			{
+				constexpr bool bIncludeFromChildActors = false;
+
+				DataContainerActor->ForEachComponent(bIncludeFromChildActors, [&CameraComponents](UActorComponent* Component)
+				{
+					if (const UChaosVDCameraDataComponent* CameraComponent = Cast<UChaosVDCameraDataComponent>(Component))
+					{
+						CameraComponents.Emplace(CameraComponent);
+					}
+				});
+			}
+
+			if (CameraComponents.Num() > 0)
+			{
+				for (const UChaosVDCameraDataComponent* Container : CameraComponents)
+				{
+					TConstArrayView<TSharedPtr<FChaosVDCameraDataWrapper>> DataWrappers = Container->GetCameraData();
+					for (const TSharedPtr<FChaosVDCameraDataWrapper>& DataWrapper : DataWrappers)
+					{
+						if ((SelectedCamera.IsSet() && DataWrapper->Camera == SelectedCamera.GetValue()) || CameraFindFlags == ECameraFindFlags::FindFirst)
+						{
+							Position = DataWrapper->Position;
+							Rotation = DataWrapper->Rotation;
+							FieldOfView = DataWrapper->FOV;
+							return true;
+						}
+					}
+				}
+			}
+		}
+	}
+	return false;
+}
+
+void FChaosVDCameraTracesExtension::SetTrackedCamera(FChaosVDCameraIdentifier InTargetCamera)
+{
+	SelectedCamera = InTargetCamera;
+
+	HandleCVDSceneUpdated();
+}
+
+void FChaosVDCameraTracesExtension::ClearTrackedCamera()
+{
+	SelectedCamera.Reset();
+}
+
+FChaosVDPlaybackViewportClient* FChaosVDCameraTracesExtension::GetPlaybackViewportClient()
+{
+	if (TSharedPtr<FEditorModeTools> EditorModeToolsPtr = EditorModeToolsWeakPtr.Pin())
+	{
+		if (FEditorViewportClient* EditorClient = EditorModeToolsPtr->GetHoveredViewportClient())
+		{
+			if (TSharedPtr<SEditorViewport> EditorViewport = EditorClient->GetEditorViewportWidget())
+			{
+				if (EditorViewport->GetType() == ChaosViewportType)
+				{
+					return static_cast<FChaosVDPlaybackViewportClient*>(EditorModeToolsPtr->GetHoveredViewportClient());
+				}
+			}
+		}
+	}
+	return nullptr;
+}
+
+void FChaosVDCameraTracesExtension::HandleCVDSceneUpdated()
+{
+	if (FChaosVDPlaybackViewportClient* ViewportClient = GetPlaybackViewportClient())
+	{
+		if (!ViewportClient->IsAutoTrackingSelectedObject())
+		{
+			FViewportCameraTransform& ViewTransform = ViewportClient->GetViewTransform();
+			FVector Position;
+			FQuat Rotation;
+			float FieldOfView;
+
+			if (FindCameraData(Position, Rotation, FieldOfView))
+			{
+				ViewTransform.SetLocation(Position);
+				ViewTransform.SetRotation(FRotator(Rotation));
+				ViewportClient->ViewFOV = FieldOfView;
+			}
+			else
+			{
+				ViewportClient->ViewFOV = ViewportClient->FOVAngle;
+			}
+
+			ViewportClient->Invalidate();
+		}
+		else
+		{
+			ClearTrackedCamera();
+		}
+	}
+}
+
+void FChaosVDCameraTracesExtension::OnSolverVisibilityUpdated(int32 SolverID, bool bNewVisibility)
+{
+	HandleCVDSceneUpdated();
+}
+
+void FChaosVDCameraTracesExtension::HandleSettingsChanged(UObject* SettingsObject)
+{
+	if (UChaosVDCameraDataSettings* Settings = Cast<UChaosVDCameraDataSettings>(SettingsObject))
+	{
+		bViewportStartAtCameraTrace = Settings->bViewportStartAtCameraTrace;
+	}
+}
+
+void FChaosVDCameraTracesExtension::ToggleOffObjectTracking()
+{
+	if (FChaosVDPlaybackViewportClient* ViewportClient = GetPlaybackViewportClient())
+	{
+		if (ViewportClient->IsAutoTrackingSelectedObject())
+		{
+			ViewportClient->ToggleObjectTrackingIfSelected();
+		}
+	}
+}
+
+void FChaosVDCameraTracesExtension::CreateCameraTracesMenu()
+{
+	if (UChaosVDCameraDataSettings* Settings = FChaosVDSettingsManager::Get().GetSettingsObject<UChaosVDCameraDataSettings>())
+	{
+		Settings->OnSettingsChanged().AddSP(this, &FChaosVDCameraTracesExtension::HandleSettingsChanged);
+		HandleSettingsChanged(Settings);
+	}
+
+	ToggleOffObjectTracking();
+
+	UToolMenus* ToolMenus = UToolMenus::Get();
+	if (UToolMenu* Menu = ToolMenus->ExtendMenu(FName("ChaosVDViewportToolbarBase")))
+	{
+		FToolMenuEntry Entry = FToolMenuEntry::InitSubMenu("CameraTraceSettings",
+			LOCTEXT("CameraTracingSubMenuLabel", "Camera Trace Settings"),
+			LOCTEXT("CameraTracingSubMenuTooltip", "Additional settings for camera traces."),
+			FNewToolMenuDelegate::CreateLambda([this](UToolMenu* Menu)
+				{
+					FToolMenuSection& OptionsSection = Menu->AddDynamicSection("ViewportOptions", FNewToolMenuDelegate::CreateLambda(
+						[this](UToolMenu* InDynamicMenu)
+						{
+							FToolMenuSection& Section = InDynamicMenu->AddSection("CameraTraces");
+
+							Section.AddMenuEntry(
+								NAME_None,
+								LOCTEXT("CameraTracingViewportStartAtCameraLabel", "Get viewport starting position from trace"),
+								FText::GetEmpty(),
+								FSlateIcon(FAppStyle::Get().GetStyleSetName(), TEXT("Icons.Toolbar.Settings")),
+								FUIAction(
+									FExecuteAction::CreateLambda([this]()
+										{
+
+											if (UChaosVDCameraDataSettings* EditorSettings = FChaosVDSettingsManager::Get().GetSettingsObject<UChaosVDCameraDataSettings>())
+											{
+												UChaosVDCameraDataSettings::SetViewportStartAtCameraTrace(!EditorSettings->bViewportStartAtCameraTrace);
+												bViewportStartAtCameraTrace = EditorSettings->bViewportStartAtCameraTrace;
+											}
+										}),
+									FCanExecuteAction(),
+									FIsActionChecked::CreateLambda([this]()
+										{
+											return bViewportStartAtCameraTrace;
+										})
+								),
+								EUserInterfaceActionType::Check
+							);
+
+							Section.AddSeparator(CameraSeparatorName);
+
+							bool bFoundEntry = false;
+
+							if (TSharedPtr<FChaosVDScene> ScenePtr = this->WeakScenePtr.Pin()) {
+								constexpr bool bIncludeFromChildActors = false;
+								TInlineComponentArray<const UChaosVDCameraDataComponent*> CameraComponents;
+								TConstArrayView<TObjectPtr<AChaosVDDataContainerBaseActor>> DataContainerActors = ScenePtr->GetDataContainerActorsView();
+								for (const TObjectPtr<AChaosVDDataContainerBaseActor>& DataContainerActor : DataContainerActors)
+								{
+									DataContainerActor->ForEachComponent(bIncludeFromChildActors, [&CameraComponents](UActorComponent* Component)
+									{
+										if (const UChaosVDCameraDataComponent* CameraComponent = Cast<UChaosVDCameraDataComponent>(Component))
+										{
+											CameraComponents.Emplace(Cast<UChaosVDCameraDataComponent>(CameraComponent));
+										}
+									});
+								}
+
+								TSet<FChaosVDCameraIdentifier> CameraSet;
+
+								if (CameraComponents.Num() > 0)
+								{
+									for (const UChaosVDCameraDataComponent* Container : CameraComponents)
+									{
+										TConstArrayView<TSharedPtr<FChaosVDCameraDataWrapper>> DataWrappers = Container->GetCameraData();
+										for (const TSharedPtr<FChaosVDCameraDataWrapper>& DataWrapper : DataWrappers)
+										{
+											FChaosVDCameraIdentifier CameraWrapper = DataWrapper->Camera;
+
+											if(!CameraSet.Contains(CameraWrapper))
+											{
+												FUIAction Action(
+													FExecuteAction::CreateLambda([this, CameraWrapper]()
+														{
+															ToggleOffObjectTracking();
+															SetTrackedCamera(CameraWrapper);
+														}),
+													FCanExecuteAction(),
+													FIsActionChecked::CreateLambda([this, CameraWrapper]()
+														{
+															return SelectedCamera.IsSet() && SelectedCamera.GetValue() == CameraWrapper;
+														})
+												);
+
+												bFoundEntry |= (SelectedCamera.IsSet() && SelectedCamera.GetValue() == CameraWrapper);
+
+												Section.AddMenuEntry(
+													NAME_None,
+													FText::FromName(CameraWrapper.CameraName),
+													FText::FromString(FString::Printf(TEXT("%s - %s"), *CameraWrapper.MapAssetPath.ToString(), *CameraWrapper.ActorName.ToString())),
+													FSlateIcon(FAppStyle::Get().GetStyleSetName(), "ClassIcon.CameraComponent"),
+													Action,
+													EUserInterfaceActionType::RadioButton
+												);
+												CameraSet.Add(CameraWrapper);
+											}
+										}
+									}
+								}
+							}
+
+							FToolMenuEntry PerspectiveEntry = FToolMenuEntry::InitMenuEntry(
+								NAME_None,
+								LOCTEXT("CameraTracingSubMenuPerspectiveLabel", "Perspective View"),
+								FText::GetEmpty(),
+								FSlateIcon(FAppStyle::Get().GetStyleSetName(), "ClassIcon.CameraComponent"),
+								FUIAction(
+									FExecuteAction::CreateLambda([this]()
+										{
+											ClearTrackedCamera();
+										}),
+									FCanExecuteAction(),
+									FIsActionChecked::CreateLambda([this, bFoundEntry]()
+										{
+											return !bFoundEntry;
+										})
+								),
+								EUserInterfaceActionType::RadioButton
+							);
+
+							PerspectiveEntry.InsertPosition.Position = EToolMenuInsertType::After;
+							PerspectiveEntry.InsertPosition.Name = CameraSeparatorName;
+							Section.AddEntry(PerspectiveEntry);
+
+						}
+					));
+				}));
+
+		Entry.ToolBarData.LabelOverride = LOCTEXT("CameraTracingMenuLabel", "Camera Tracing");
+		Entry.Icon = FSlateIcon(FAppStyle::Get().GetStyleSetName(), "ClassIcon.CameraComponent");
+
+		for (FToolMenuSection& Sec : Menu->Sections)
+		{
+			if (Sec.Name.IsEqual(FName("Left")))
+			{
+				Sec.AddEntry(Entry);
+			}
+		}
+	}
+}
+
+void FChaosVDCameraTracesExtension::PostMainTabInitialization(const TSharedRef<SChaosVDMainTab>& InParentTabWidget)
+{
+	CreateCameraTracesMenu();
+}
+
+void FChaosVDCameraTracesExtension::RegisterComponentVisualizers(const TSharedRef<SChaosVDMainTab>& InCVDToolKit)
+{
+	FChaosVDExtension::RegisterComponentVisualizers(InCVDToolKit);
+	
+	InCVDToolKit->RegisterComponentVisualizer(UChaosVDCameraDataComponent::StaticClass()->GetFName(), MakeShared<FChaosVDCameraDataComponentVisualizer>());
+	TSharedPtr<FChaosVDRecordingControls> a = InCVDToolKit->GetRecordingControls();
+	TSharedRef<FChaosVDEngine> EngineInstance = InCVDToolKit->GetChaosVDEngineInstance();
+
+	WeakScenePtr = nullptr;
+	if (TSharedPtr<FChaosVDScene> ScenePtr = EngineInstance->GetCurrentScene())
+	{
+		WeakScenePtr = ScenePtr;
+	}
+
+	EditorModeToolsWeakPtr = InCVDToolKit->GetEditorModeManager().AsWeak();
+
+	TSharedPtr<FChaosVDPlaybackController> PlaybackControllerPtr = EngineInstance->GetPlaybackController();
+
+	if (PlaybackControllerPtr)
+	{
+		if (TSharedPtr<FChaosVDScene> ScenePtr = PlaybackControllerPtr->GetControllerScene().Pin())
+		{
+			ScenePtr->OnSceneUpdated().AddSP(this, &FChaosVDCameraTracesExtension::HandleCVDSceneUpdated);
+			ScenePtr->OnSolverVisibilityUpdated().AddSP(this, &FChaosVDCameraTracesExtension::OnSolverVisibilityUpdated);
+		}
+	}
+}
+
+#undef LOCTEXT_NAMESPACE
\ No newline at end of file
diff --git a/Engine/Plugins/ChaosVD/Source/ChaosVDBuiltInExtensions/Private/CameraTraces/ChaosVDCameraTracesExtension.h b/Engine/Plugins/ChaosVD/Source/ChaosVDBuiltInExtensions/Private/CameraTraces/ChaosVDCameraTracesExtension.h
new file mode 100644
--- /dev/null
+++ b/Engine/Plugins/ChaosVD/Source/ChaosVDBuiltInExtensions/Private/CameraTraces/ChaosVDCameraTracesExtension.h
@@ -0,0 +1,60 @@
+// Copyright Epic Games, Inc. All Rights Reserved.
+
+#pragma once
+
+#pragma once
+#include "ExtensionsSystem/ChaosVDExtension.h"
+#include "Templates/SubclassOf.h"
+#include "Templates/SharedPointer.h"
+#include "DataWrappers/ChaosVDCameraDataWrapper.h"
+
+class FChaosVDTraceProvider;
+class UActorComponent;
+class SChaosVDMainTab;
+class FChaosVDPlaybackViewportClient;
+class FEditorModeTools;
+class FChaosVDPlaybackController;
+class FChaosVDScene;
+struct FChaosVDCameraIdentifier;
+
+class FChaosVDCameraTracesExtension final : public FChaosVDExtension, public TSharedFromThis<FChaosVDCameraTracesExtension>
+{
+public:
+	
+	FChaosVDCameraTracesExtension();
+	virtual ~FChaosVDCameraTracesExtension() override;
+
+	virtual void RegisterDataProcessorsInstancesForProvider(const TSharedRef<FChaosVDTraceProvider>& InTraceProvider) override;
+	virtual TConstArrayView<TSubclassOf<UActorComponent>> GetSolverDataComponentsClasses() override;
+	virtual void RegisterComponentVisualizers(const TSharedRef<SChaosVDMainTab>& InCVDToolKit) override;
+	virtual void PostMainTabInitialization(const TSharedRef<SChaosVDMainTab>& InParentTabWidget) override;
+	virtual void HandleRecordingFirstFrameLoaded(TWeakPtr<FChaosVDPlaybackController> InController) override;
+
+protected:
+	enum class ECameraFindFlags : uint8
+	{
+		Default,
+		FindFirst
+	};
+
+	TWeakPtr<FChaosVDScene> WeakScenePtr;
+	TWeakPtr<FEditorModeTools> EditorModeToolsWeakPtr;
+
+	TOptional<FChaosVDCameraIdentifier> SelectedCamera;
+	bool bViewportStartAtCameraTrace = true;
+	static const FName ChaosViewportType;
+	static const FName CameraSeparatorName;
+
+	void SetTrackedCamera(FChaosVDCameraIdentifier InTargetCamera);
+	void ClearTrackedCamera();
+	bool FindCameraData(FVector& Position, FQuat& Rotation, float& FieldOfView, ECameraFindFlags CameraFindFlags) const;
+	void HandleCVDSceneUpdated();
+	void OnSolverVisibilityUpdated(int32 SolverID, bool bNewVisibility);
+	void CreateCameraTracesMenu();
+	void HandleSettingsChanged(UObject* SettingsObject);
+	void ToggleOffObjectTracking();
+	FChaosVDPlaybackViewportClient* GetPlaybackViewportClient();
+
+private:
+	TArray<TSubclassOf<UActorComponent>> DataComponentsClasses;
+};
diff --git a/Engine/Plugins/ChaosVD/Source/ChaosVDBuiltInExtensions/Private/CameraTraces/Components/ChaosVDCameraDataComponent.cpp b/Engine/Plugins/ChaosVD/Source/ChaosVDBuiltInExtensions/Private/CameraTraces/Components/ChaosVDCameraDataComponent.cpp
new file mode 100644
--- /dev/null
+++ b/Engine/Plugins/ChaosVD/Source/ChaosVDBuiltInExtensions/Private/CameraTraces/Components/ChaosVDCameraDataComponent.cpp
@@ -0,0 +1,36 @@
+// Copyright Epic Games, Inc. All Rights Reserved.
+
+#include "ChaosVDCameraDataComponent.h"
+#include "DataWrappers/ChaosVDCameraDataWrapper.h"
+#include "ChaosVDRecording.h"
+#include "Algo/Copy.h"
+
+#include UE_INLINE_GENERATED_CPP_BY_NAME(ChaosVDCameraDataComponent)
+
+UChaosVDCameraDataComponent::UChaosVDCameraDataComponent()
+{
+	PrimaryComponentTick.bCanEverTick = false;
+	
+	SetCanEverAffectNavigation(false);
+	bNavigationRelevant = false;
+}
+
+void UChaosVDCameraDataComponent::UpdateCameraData(TConstArrayView<TSharedPtr<FChaosVDCameraDataWrapper>> CameraDataView)
+{
+	RecordedCameraData.Reset(CameraDataView.Num());
+	Algo::Copy(CameraDataView, RecordedCameraData);
+}
+
+void UChaosVDCameraDataComponent::UpdateFromNewGameFrameData(const FChaosVDGameFrameData& InGameFrameData)
+{
+	if (TSharedPtr<FChaosVDCameraDataContainer> CameraDataContainer = InGameFrameData.GetCustomDataHandler().GetData<FChaosVDCameraDataContainer>())
+	{
+		const TArray<TSharedPtr<FChaosVDCameraDataWrapper>>& RecordedCameraTraceData = CameraDataContainer->CameraData;
+		UpdateCameraData(RecordedCameraTraceData);
+	}
+}
+
+void UChaosVDCameraDataComponent::ClearData()
+{
+	RecordedCameraData.Reset();
+}
diff --git a/Engine/Plugins/ChaosVD/Source/ChaosVDBuiltInExtensions/Private/CameraTraces/Components/ChaosVDCameraDataComponent.h b/Engine/Plugins/ChaosVD/Source/ChaosVDBuiltInExtensions/Private/CameraTraces/Components/ChaosVDCameraDataComponent.h
new file mode 100644
--- /dev/null
+++ b/Engine/Plugins/ChaosVD/Source/ChaosVDBuiltInExtensions/Private/CameraTraces/Components/ChaosVDCameraDataComponent.h
@@ -0,0 +1,33 @@
+// Copyright Epic Games, Inc. All Rights Reserved.
+
+#pragma once
+
+#include "Components/ChaosVDSolverDataComponent.h"
+#include "ChaosVDCameraDataComponent.generated.h"
+
+struct FChaosVDCameraDataWrapper;
+struct FChaosVDGameFrameData;
+
+UCLASS()
+class UChaosVDCameraDataComponent : public UChaosVDSolverDataComponent
+{
+	GENERATED_BODY()
+
+public:
+	// Sets default values for this component's properties
+	UChaosVDCameraDataComponent();
+
+	void UpdateCameraData(TConstArrayView<TSharedPtr<FChaosVDCameraDataWrapper>> CameraDataView);
+
+	TConstArrayView<TSharedPtr<FChaosVDCameraDataWrapper>> GetCameraData() const
+	{ 
+		return RecordedCameraData;
+	}
+
+	virtual void UpdateFromNewGameFrameData(const FChaosVDGameFrameData& InGameFrameData) override;
+
+	virtual void ClearData() override;
+protected:
+	
+	TArray<TSharedPtr<FChaosVDCameraDataWrapper>> RecordedCameraData;
+};
diff --git a/Engine/Plugins/ChaosVD/Source/ChaosVDBuiltInExtensions/Private/CameraTraces/Settings/ChaosVDCameraDataSettings.cpp b/Engine/Plugins/ChaosVD/Source/ChaosVDBuiltInExtensions/Private/CameraTraces/Settings/ChaosVDCameraDataSettings.cpp
new file mode 100644
--- /dev/null
+++ b/Engine/Plugins/ChaosVD/Source/ChaosVDBuiltInExtensions/Private/CameraTraces/Settings/ChaosVDCameraDataSettings.cpp
@@ -0,0 +1,42 @@
+// Copyright Epic Games, Inc. All Rights Reserved.
+
+#include "ChaosVDCameraDataSettings.h"
+
+#include "ChaosVDSettingsManager.h"
+#include "Utils/ChaosVDUserInterfaceUtils.h"
+
+#include UE_INLINE_GENERATED_CPP_BY_NAME(ChaosVDCameraDataSettings)
+
+
+void UChaosVDCameraDataSettings::SetDataVisualizationFlags(EChaosVDCameraDataVisualizationFlags NewFlags)
+{
+	if (UChaosVDCameraDataSettings* Settings = FChaosVDSettingsManager::Get().GetSettingsObject<UChaosVDCameraDataSettings>())
+	{
+		Settings->DebugDrawFlags = static_cast<uint32>(NewFlags);
+		Settings->BroadcastSettingsChanged();
+	}
+}
+
+void UChaosVDCameraDataSettings::SetViewportStartAtCameraTrace(bool bInViewportStartAtCameraTrace)
+{
+	if (UChaosVDCameraDataSettings* Settings = FChaosVDSettingsManager::Get().GetSettingsObject<UChaosVDCameraDataSettings>())
+	{
+		Settings->bViewportStartAtCameraTrace = bInViewportStartAtCameraTrace;
+		Settings->BroadcastSettingsChanged();
+	}
+}
+
+EChaosVDCameraDataVisualizationFlags UChaosVDCameraDataSettings::GetDataVisualizationFlags()
+{
+	if (UChaosVDCameraDataSettings* Settings = FChaosVDSettingsManager::Get().GetSettingsObject<UChaosVDCameraDataSettings>())
+	{
+		return static_cast<EChaosVDCameraDataVisualizationFlags>(Settings->DebugDrawFlags);
+	}
+
+	return EChaosVDCameraDataVisualizationFlags::None;
+}
+
+bool UChaosVDCameraDataSettings::CanVisualizationFlagBeChangedByUI(uint32 Flag)
+{
+	return Chaos::VisualDebugger::Utils::ShouldVisFlagBeEnabledInUI(Flag, DebugDrawFlags, EChaosVDCameraDataVisualizationFlags::EnableDraw);
+}
diff --git a/Engine/Plugins/ChaosVD/Source/ChaosVDBuiltInExtensions/Private/CameraTraces/Settings/ChaosVDCameraDataSettings.h b/Engine/Plugins/ChaosVD/Source/ChaosVDBuiltInExtensions/Private/CameraTraces/Settings/ChaosVDCameraDataSettings.h
new file mode 100644
--- /dev/null
+++ b/Engine/Plugins/ChaosVD/Source/ChaosVDBuiltInExtensions/Private/CameraTraces/Settings/ChaosVDCameraDataSettings.h
@@ -0,0 +1,48 @@
+// Copyright Epic Games, Inc. All Rights Reserved.
+
+#pragma once
+
+#include "Settings/ChaosVDCoreSettings.h"
+#include "ChaosVDCameraDataSettings.generated.h"
+
+/** Set of visualization flags options for camera data */
+UENUM(meta = (Bitflags, UseEnumValuesAsMaskValuesInEditor = "true"))
+enum class EChaosVDCameraDataVisualizationFlags : uint32
+{
+	None					= 0 UMETA(Hidden),
+	EnableDraw				= 1 << 0,
+};
+ENUM_CLASS_FLAGS(EChaosVDCameraDataVisualizationFlags);
+
+UCLASS(config=ChaosVD, PerObjectConfig)
+class UChaosVDCameraDataSettings : public UChaosVDVisualizationSettingsObjectBase
+{
+	GENERATED_BODY()
+
+public:
+
+	/** If true and if there is an existing camera trace, the initial viewport position will be at the rather than at the world origin*/
+	UPROPERTY(config, EditAnywhere, Category = DebugDraw)
+	bool bViewportStartAtCameraTrace = true;
+
+	/** The depth priority used for while drawing data. Can be World or Foreground (with this one the shapes will be drawn on top of the geometry and be always visible) */
+	UPROPERTY(config, EditAnywhere, Category=DebugDraw)
+	TEnumAsByte<ESceneDepthPriorityGroup> DepthPriority = ESceneDepthPriorityGroup::SDPG_Foreground;
+
+	/** Length of the camera direction vector when drawn. */
+	UPROPERTY(config, EditAnywhere, Category=DebugDraw)
+	float DirectionVectorScale = 2.0f;
+
+	static void SetDataVisualizationFlags(EChaosVDCameraDataVisualizationFlags NewFlags);
+	static EChaosVDCameraDataVisualizationFlags GetDataVisualizationFlags();
+
+	static void SetViewportStartAtCameraTrace(bool bInViewportStartAtCameraTrace);
+
+	virtual bool CanVisualizationFlagBeChangedByUI(uint32 Flag) override;
+
+private:
+
+	/** Set of flags to enable/disable visualization of debug draw data shapes */
+	UPROPERTY(config, meta = (Bitmask, BitmaskEnum = "/Script/ChaosVDBuiltInExtensions.EChaosVDCameraDataVisualizationFlags"))
+	uint32 DebugDrawFlags = static_cast<uint32>(EChaosVDCameraDataVisualizationFlags::None);
+};
diff --git a/Engine/Plugins/ChaosVD/Source/ChaosVDBuiltInExtensions/Private/CameraTraces/Trace/ChaosVDCameraDataProcessor.cpp b/Engine/Plugins/ChaosVD/Source/ChaosVDBuiltInExtensions/Private/CameraTraces/Trace/ChaosVDCameraDataProcessor.cpp
new file mode 100644
--- /dev/null
+++ b/Engine/Plugins/ChaosVD/Source/ChaosVDBuiltInExtensions/Private/CameraTraces/Trace/ChaosVDCameraDataProcessor.cpp
@@ -0,0 +1,43 @@
+// Copyright Epic Games, Inc. All Rights Reserved.
+
+#include "ChaosVDCameraDataProcessor.h"
+
+#include "ChaosVDRecording.h"
+#include "ChaosVisualDebugger/ChaosVDMemWriterReader.h"
+#include "ChaosVisualDebugger/ChaosVDSerializedNameTable.h"
+#include "ChaosVisualDebugger/ChaosVisualDebuggerTrace.h"
+#include "DataWrappers/ChaosVDCameraDataWrapper.h"
+#include "Trace/ChaosVDTraceProvider.h"
+
+FChaosVDCameraDataProcessor::FChaosVDCameraDataProcessor() : FChaosVDDataProcessorBase(FChaosVDCameraDataWrapper::WrapperTypeName)
+{
+}
+
+bool FChaosVDCameraDataProcessor::ProcessRawData(const TArray<uint8>& InData)
+{
+	FChaosVDDataProcessorBase::ProcessRawData(InData);
+
+	const TSharedPtr<FChaosVDTraceProvider> ProviderSharedPtr = TraceProvider.Pin();
+	if (!ensure(ProviderSharedPtr.IsValid()))
+	{
+		return false;
+	}
+
+	const TSharedPtr<FChaosVDCameraDataWrapper> CameraData = MakeShared<FChaosVDCameraDataWrapper>();
+	const bool bSuccess = Chaos::VisualDebugger::ReadDataFromBuffer(InData, *CameraData, ProviderSharedPtr.ToSharedRef());
+
+	if (bSuccess)
+	{
+		if (TSharedPtr<FChaosVDGameFrameData> CurrentFrameData = ProviderSharedPtr->GetCurrentGameFrame().Pin())
+		{
+			if (TSharedPtr<FChaosVDCameraDataContainer> CamDataContainer = CurrentFrameData->GetCustomDataHandler().GetOrAddDefaultData<FChaosVDCameraDataContainer>())
+			{
+				CamDataContainer->CameraData.Add(CameraData);
+				CurrentFrameData->MarkDirty();
+			}
+		}
+	}
+
+	return bSuccess;
+}
+
diff --git a/Engine/Plugins/ChaosVD/Source/ChaosVDBuiltInExtensions/Private/CameraTraces/Trace/ChaosVDCameraDataProcessor.h b/Engine/Plugins/ChaosVD/Source/ChaosVDBuiltInExtensions/Private/CameraTraces/Trace/ChaosVDCameraDataProcessor.h
new file mode 100644
--- /dev/null
+++ b/Engine/Plugins/ChaosVD/Source/ChaosVDBuiltInExtensions/Private/CameraTraces/Trace/ChaosVDCameraDataProcessor.h
@@ -0,0 +1,17 @@
+// Copyright Epic Games, Inc. All Rights Reserved.
+
+#pragma once
+
+#include "Trace/DataProcessors/ChaosVDDataProcessorBase.h"
+
+/**
+ * Data processor implementation that is able to deserialize traced camera data
+ */
+class FChaosVDCameraDataProcessor final : public FChaosVDDataProcessorBase
+{
+public:
+	explicit FChaosVDCameraDataProcessor();
+
+	virtual bool ProcessRawData(const TArray<uint8>& InData) override;
+};
+
diff --git a/Engine/Plugins/ChaosVD/Source/ChaosVDBuiltInExtensions/Private/CameraTraces/Visualizers/ChaosVDCameraDataComponentVisualizer.cpp b/Engine/Plugins/ChaosVD/Source/ChaosVDBuiltInExtensions/Private/CameraTraces/Visualizers/ChaosVDCameraDataComponentVisualizer.cpp
new file mode 100644
--- /dev/null
+++ b/Engine/Plugins/ChaosVD/Source/ChaosVDBuiltInExtensions/Private/CameraTraces/Visualizers/ChaosVDCameraDataComponentVisualizer.cpp
@@ -0,0 +1,169 @@
+// Copyright Epic Games, Inc. All Rights Reserved.
+
+#include "ChaosVDCameraDataComponentVisualizer.h"
+
+#include "CameraTraces/Components/ChaosVdCameraDataComponent.h"
+#include "CameraTraces/Settings/ChaosVDCameraDataSettings.h"
+#include "Actors/ChaosVDDataContainerBaseActor.h"
+#include "ChaosVDScene.h"
+#include "ChaosVDStyle.h"
+#include "ChaosVDTabsIDs.h"
+#include "SceneView.h"
+#include "Visualizers/ChaosVDDebugDrawUtils.h"
+
+#include UE_INLINE_GENERATED_CPP_BY_NAME(ChaosVDCameraDataComponentVisualizer)
+
+#define LOCTEXT_NAMESPACE "ChaosVisualDebugger"
+
+class AChaosVDSolverInfoActor;
+
+bool FChaosVDCameraSelectionHandle::IsSelected()
+{
+	if (bool bIsPrimaryDataSelected = FChaosVDSolverDataSelectionHandle::IsSelected())
+	{
+		if (TSharedPtr<FChaosVDSolverDataSelection> OwnerPtr = Owner.Pin())
+		{
+			TSharedPtr<FChaosVDSolverDataSelectionHandle> CurrentSelectedDataHandle = OwnerPtr->GetCurrentSelectionHandle();
+
+			FChaosVDCameraSelectionContext* CurrentSelectionContext = CurrentSelectedDataHandle->GetContextData<FChaosVDCameraSelectionContext>();
+			FChaosVDCameraSelectionContext* HandleSelectionContext = GetContextData<FChaosVDCameraSelectionContext>();
+
+			return CurrentSelectionContext && HandleSelectionContext && (*CurrentSelectionContext) == (*HandleSelectionContext);
+		}
+	}
+
+	return false;
+}
+
+void FChaosVDCameraSelectionHandle::CreateStructViewForDetailsPanelIfNeeded()
+{
+	if (StructDataView)
+	{
+		return;
+	}
+
+	StructDataView = MakeShared<FChaosVDSelectionMultipleView>();
+
+	if (FChaosVDCameraDataWrapper* CameraData = GetData<FChaosVDCameraDataWrapper>())
+	{
+		StructDataView->AddData(CameraData);
+	}
+
+	if (FChaosVDCameraSelectionContext* SelectionContext = GetContextData<FChaosVDCameraSelectionContext>())
+	{
+		StructDataView->AddData(const_cast<FChaosVDCameraDataWrapper*>(SelectionContext->CameraData));
+	}
+
+	StructDataViewStructOnScope = MakeShared<FStructOnScope>(FChaosVDSelectionMultipleView::StaticStruct(), reinterpret_cast<uint8*>(StructDataView.Get()));
+}
+
+TSharedPtr<FStructOnScope> FChaosVDCameraSelectionHandle::GetCustomDataReadOnlyStructViewForDetails()
+{
+	CreateStructViewForDetailsPanelIfNeeded();
+
+	return StructDataViewStructOnScope;
+}
+
+FChaosVDCameraDataComponentVisualizer::FChaosVDCameraDataComponentVisualizer()
+{
+	RegisterVisualizerMenus();
+	InspectorTabID = FChaosVDTabID::DetailsPanel;
+}
+
+void FChaosVDCameraDataComponentVisualizer::RegisterVisualizerMenus()
+{
+
+	FName MenuSection("CameraDataVisualization.Show");
+	FText MenuSectionLabel = LOCTEXT("CameraDataShowMenuLabel", "Camera Data Visualization");
+	FText FlagsMenuLabel = LOCTEXT("CameraDataFlagsMenuLabel", "Camera Data Flags");
+	FText FlagsMenuTooltip = LOCTEXT("CameraDataFlagsMenuToolTip", "Set of flags to enable/disable visibility of camera data visualization");
+	FSlateIcon FlagsMenuIcon = FSlateIcon(FAppStyle::Get().GetStyleSetName(), "ClassIcon.CameraComponent");
+
+	FText SettingsMenuLabel = LOCTEXT("CameraDataSettingsMenuLabel", "Camera Data Visualization Settings");
+	FText SettingsMenuTooltip = LOCTEXT("CameraDataSettingsMenuToolTip", "Options to change how the camera data is debug drawn");
+	
+	CreateGenericVisualizerMenu<UChaosVDCameraDataSettings, EChaosVDCameraDataVisualizationFlags>(FName("ChaosVDViewportToolbarBase.Show"), MenuSection, MenuSectionLabel, FlagsMenuLabel, FlagsMenuTooltip, FlagsMenuIcon, SettingsMenuLabel, SettingsMenuTooltip);
+}
+
+void FChaosVDCameraDataComponentVisualizer::DrawVisualization(const UActorComponent* Component, const FSceneView* View, FPrimitiveDrawInterface* PDI)
+{
+	const UChaosVDCameraDataComponent* DataComponent = Cast<UChaosVDCameraDataComponent>(Component);
+	if (!DataComponent)
+	{
+		return;
+	}
+	
+	AChaosVDDataContainerBaseActor* DataInfoActor = Cast<AChaosVDDataContainerBaseActor>(Component->GetOwner());
+	if (!DataInfoActor)
+	{
+		return;
+	}
+
+	if (!DataInfoActor->IsVisible())
+	{
+		return;
+	}
+
+	const TSharedPtr<FChaosVDScene> CVDScene = DataInfoActor->GetScene().Pin();
+	if (!CVDScene)
+	{
+		return;
+	}
+	
+	FChaosVDCameraVisualizationDataContext VisualizationContext;
+	VisualizationContext.CVDScene = CVDScene;
+	VisualizationContext.SpaceTransform = DataInfoActor->GetSimulationTransform();
+	VisualizationContext.SolverDataSelectionObject = CVDScene->GetSolverDataSelectionObject().Pin();
+	VisualizationContext.DataComponent = DataComponent;
+
+	if (const UChaosVDCameraDataSettings* EditorSettings = FChaosVDSettingsManager::Get().GetSettingsObject<UChaosVDCameraDataSettings>())
+	{
+		VisualizationContext.VisualizationFlags = static_cast<uint32>(UChaosVDCameraDataSettings::GetDataVisualizationFlags());
+		VisualizationContext.DebugDrawSettings = EditorSettings;
+		VisualizationContext.DepthPriority = EditorSettings->DepthPriority;
+	}
+
+	if (!VisualizationContext.IsVisualizationFlagEnabled(EChaosVDCameraDataVisualizationFlags::EnableDraw))
+	{
+		return;
+	}
+
+	TConstArrayView<TSharedPtr<FChaosVDCameraDataWrapper>> CameraData = DataComponent->GetCameraData();
+
+	for (const TSharedPtr<FChaosVDCameraDataWrapper>& CameraDataWrapper : CameraData)
+	{
+		if (CameraDataWrapper)
+		{
+			DrawCameraTrace(View, PDI, VisualizationContext, CameraDataWrapper.ToSharedRef());
+		}
+	}
+}
+
+bool FChaosVDCameraDataComponentVisualizer::CanHandleClick(const HChaosVDComponentVisProxy& VisProxy)
+{
+	return VisProxy.DataSelectionHandle && VisProxy.DataSelectionHandle->IsA<FChaosVDCameraDataWrapper>();
+}
+
+void FChaosVDCameraDataComponentVisualizer::DrawCameraTrace(const FSceneView* View, FPrimitiveDrawInterface* PDI, const FChaosVDCameraVisualizationDataContext& VisualizationContext, const TSharedRef<FChaosVDCameraDataWrapper>& CameraData)
+{
+	const UChaosVDCameraDataSettings* Settings = Cast<const UChaosVDCameraDataSettings>(VisualizationContext.DebugDrawSettings);
+	if (!ensure(Settings))
+	{
+		return;
+	}
+
+	const FVector Center = CameraData->Position;
+	const FVector DirPos = Center + FRotator(CameraData->Rotation).RotateVector(FVector(42 * Settings->DirectionVectorScale, 0, 0));
+	
+
+	FChaosVDDebugDrawUtils::DrawLine(PDI, Center, DirPos, FColor::Red, FText::GetEmpty(), VisualizationContext.DepthPriority, 2.0f);
+
+	constexpr FVector BoxExtents(5.0f, UE::Math::TVectorConstInit{});
+
+	FTransform LocationTransform;
+	LocationTransform.SetLocation(Center);
+
+	FChaosVDDebugDrawUtils::DrawBox(PDI, BoxExtents, FColor::Red, LocationTransform, FText::GetEmpty(), VisualizationContext.DepthPriority, 2.0f);
+}
+
+#undef LOCTEXT_NAMESPACE
diff --git a/Engine/Plugins/ChaosVD/Source/ChaosVDBuiltInExtensions/Private/CameraTraces/Visualizers/ChaosVDCameraDataComponentVisualizer.h b/Engine/Plugins/ChaosVD/Source/ChaosVDBuiltInExtensions/Private/CameraTraces/Visualizers/ChaosVDCameraDataComponentVisualizer.h
new file mode 100644
--- /dev/null
+++ b/Engine/Plugins/ChaosVD/Source/ChaosVDBuiltInExtensions/Private/CameraTraces/Visualizers/ChaosVDCameraDataComponentVisualizer.h
@@ -0,0 +1,68 @@
+// Copyright Epic Games, Inc. All Rights Reserved.
+
+#pragma once
+#include "Visualizers/ChaosVDComponentVisualizerBase.h"
+#include "ChaosVDSolverDataSelection.h"
+#include "DataWrappers/ChaosVDCameraDataWrapper.h"
+#include "CameraTraces/Settings/ChaosVDCameraDataSettings.h"
+
+#include "ChaosVDCameraDataComponentVisualizer.generated.h"
+
+class UChaosVDCameraDataComponent;
+
+USTRUCT()
+struct FChaosVDCameraSelectionContext
+{
+	GENERATED_BODY()
+
+	const FChaosVDCameraDataWrapper* CameraData = nullptr;
+
+	bool operator==(const FChaosVDCameraSelectionContext& Other) const
+	{
+		return CameraData == Other.CameraData;
+	}
+};
+
+struct FChaosVDCameraSelectionHandle : public FChaosVDSolverDataSelectionHandle
+{
+	virtual bool IsSelected() override;
+	void CreateStructViewForDetailsPanelIfNeeded();
+
+	virtual TSharedPtr<FStructOnScope> GetCustomDataReadOnlyStructViewForDetails() override;
+
+private:
+	TSharedPtr<FChaosVDSelectionMultipleView> StructDataView;
+	TSharedPtr<FStructOnScope> StructDataViewStructOnScope;
+};
+
+/** Visualization context structure specific for camera data visualizations */
+struct FChaosVDCameraVisualizationDataContext : public FChaosVDVisualizationContext
+{
+	TSharedPtr<FChaosVDSolverDataSelectionHandle> DataSelectionHandle = MakeShared<FChaosVDSolverDataSelectionHandle>();
+
+	ESceneDepthPriorityGroup DepthPriority = SDPG_Foreground;
+
+	const UChaosVDCameraDataComponent* DataComponent = nullptr;
+
+	bool IsVisualizationFlagEnabled(EChaosVDCameraDataVisualizationFlags Flag) const
+	{
+		const EChaosVDCameraDataVisualizationFlags FlagsAsCameraDataVisFlags = static_cast<EChaosVDCameraDataVisualizationFlags>(VisualizationFlags);
+		return EnumHasAnyFlags(FlagsAsCameraDataVisFlags, Flag);
+	}
+};
+
+class FChaosVDCameraDataComponentVisualizer final : public FChaosVDComponentVisualizerBase
+{
+public:
+	FChaosVDCameraDataComponentVisualizer();
+
+	virtual void RegisterVisualizerMenus() override;
+
+	virtual void DrawVisualization(const UActorComponent* Component, const FSceneView* View, FPrimitiveDrawInterface* PDI) override;
+
+	virtual bool CanHandleClick(const HChaosVDComponentVisProxy& VisProxy) override;
+
+protected:
+
+	void DrawCameraTrace(const FSceneView* View, FPrimitiveDrawInterface* PDI, const FChaosVDCameraVisualizationDataContext& VisualizationContext, const TSharedRef<FChaosVDCameraDataWrapper>& CameraData);
+};
diff --git a/Engine/Plugins/ChaosVD/Source/ChaosVDBuiltInExtensions/Private/ChaosVDBuiltInExtensions.cpp b/Engine/Plugins/ChaosVD/Source/ChaosVDBuiltInExtensions/Private/ChaosVDBuiltInExtensions.cpp
--- a/Engine/Plugins/ChaosVD/Source/ChaosVDBuiltInExtensions/Private/ChaosVDBuiltInExtensions.cpp
+++ b/Engine/Plugins/ChaosVD/Source/ChaosVDBuiltInExtensions/Private/ChaosVDBuiltInExtensions.cpp
@@ -4,11 +4,13 @@
 
 #include "AccelerationStructures/ChaosVDAccelerationStructuresExtension.h"
 #include "GenericDebugDraw/ChaosVDGenericDebugDrawExtension.h"
+#include "CameraTraces/ChaosVDCameraTracesExtension.h"
 
 void FChaosVDBuiltInExtensionsModule::StartupModule()
 {
 	CreateAndRegisterExtensionInstance<FChaosVDGenericDebugDrawExtension>();
 	CreateAndRegisterExtensionInstance<FChaosVDAccelerationStructuresExtension>();
+	CreateAndRegisterExtensionInstance<FChaosVDCameraTracesExtension>();
 }
 
 void FChaosVDBuiltInExtensionsModule::ShutdownModule()
diff --git a/Engine/Source/Runtime/Engine/Private/PlayerCameraManager.cpp b/Engine/Source/Runtime/Engine/Private/PlayerCameraManager.cpp
--- a/Engine/Source/Runtime/Engine/Private/PlayerCameraManager.cpp
+++ b/Engine/Source/Runtime/Engine/Private/PlayerCameraManager.cpp
@@ -23,6 +23,7 @@
 #include "TimerManager.h"
 #include "Camera/CameraLensEffectInterface.h"
 #include "GameDelegates.h"
+#include "ChaosVisualDebugger/ChaosVisualDebuggerTrace.h"
 
 #include UE_INLINE_GENERATED_CPP_BY_NAME(PlayerCameraManager)
 
@@ -831,6 +832,14 @@ void APlayerCameraManager::UpdateCamera(float DeltaTime)
 			bShouldSendClientSideCameraUpdate = false;
 		}
 	}
+
+#if WITH_CHAOS_VISUAL_DEBUGGER
+	if (FChaosVisualDebuggerTrace::IsTracing())
+	{
+		const FMinimalViewInfo& NewInfo = GetCameraCacheView();
+		CVD_TRACE_CAMERA(this, GetWorld(), GetFName(), NewInfo.Location, FQuat(NewInfo.Rotation), NewInfo.FOV);
+	}
+#endif
 }
 
 bool APlayerCameraManager::AllowPhotographyMode() const
diff --git a/Engine/Source/Runtime/Experimental/Chaos/Private/ChaosVisualDebugger/ChaosVDOptionalDataChannel.cpp b/Engine/Source/Runtime/Experimental/Chaos/Private/ChaosVisualDebugger/ChaosVDOptionalDataChannel.cpp
--- a/Engine/Source/Runtime/Experimental/Chaos/Private/ChaosVisualDebugger/ChaosVDOptionalDataChannel.cpp
+++ b/Engine/Source/Runtime/Experimental/Chaos/Private/ChaosVisualDebugger/ChaosVDOptionalDataChannel.cpp
@@ -104,5 +104,6 @@ CVD_DEFINE_OPTIONAL_DATA_CHANNEL(JointNonLinearConstraints, EChaosVDDataChannelI
 CVD_DEFINE_OPTIONAL_DATA_CHANNEL(CharacterGroundConstraints, EChaosVDDataChannelInitializationFlags::CanChangeEnabledState);
 CVD_DEFINE_OPTIONAL_DATA_CHANNEL(AccelerationStructures, EChaosVDDataChannelInitializationFlags::CanChangeEnabledState);
 CVD_DEFINE_OPTIONAL_DATA_CHANNEL(ClusterParticlesChildData, EChaosVDDataChannelInitializationFlags::CanChangeEnabledState);
+CVD_DEFINE_OPTIONAL_DATA_CHANNEL(CameraPaths, EChaosVDDataChannelInitializationFlags::StartEnabled | EChaosVDDataChannelInitializationFlags::CanChangeEnabledState);
 
 #endif
diff --git a/Engine/Source/Runtime/Experimental/Chaos/Private/ChaosVisualDebugger/ChaosVisualDebuggerTrace.cpp b/Engine/Source/Runtime/Experimental/Chaos/Private/ChaosVisualDebugger/ChaosVisualDebuggerTrace.cpp
--- a/Engine/Source/Runtime/Experimental/Chaos/Private/ChaosVisualDebugger/ChaosVisualDebuggerTrace.cpp
+++ b/Engine/Source/Runtime/Experimental/Chaos/Private/ChaosVisualDebugger/ChaosVisualDebuggerTrace.cpp
@@ -26,9 +26,12 @@
 #include "DataWrappers/ChaosVDJointDataWrappers.h"
 #include "DataWrappers/ChaosVDParticleDataWrapper.h"
 #include "DataWrappers/ChaosVDQueryDataWrappers.h"
+#include "DataWrappers/ChaosVDCameraDataWrapper.h"
 #include "HAL/CriticalSection.h"
 #include "Misc/ScopeRWLock.h"
 #include "Templates/UniquePtr.h"
+#include "UObject/ObjectMacros.h"
+#include "Engine/Engine.h"
 
 static_assert(CHAOS_VISUAL_DEBUGGER_WITHOUT_TRACE || UE_TRACE_MINIMAL_ENABLED || UE_TRACE_ENABLED, "Chaos Visual Debugger support requires Trace or Trace minimal support. Please enable trace support with the Unreal Build Tool flag bEnableTrace, disable cvd support with bCompileChaosVisualDebuggerSupport or allows CVD support without trace by defining CHAOS_VISUAL_DEBUGGER_WITHOUT_TRACE to 1");
 
@@ -317,6 +320,30 @@ void FChaosVisualDebuggerTrace::TraceParticlesSoA(const Chaos::FPBDRigidsSOAs& P
 	TraceParticleClusterChildData(ParticlesSoA.GetDirtyParticlesView(), ClusteringData, *CVDContextData);
 }
 
+void FChaosVisualDebuggerTrace::TraceCamera(AActor* Object, UWorld* World, FName CameraName, const FVector& Position, const FQuat& Rotation, float FOV)
+{
+	using namespace Chaos::VisualDebugger::Utils;
+	if (!Object || !IsTracing() || !CVDDC_CameraPaths->IsChannelEnabled())
+	{
+		return;
+	}
+
+	FChaosVDCameraDataWrapper DataWrapper;
+
+	DataWrapper.Camera.CameraName = CameraName;
+	DataWrapper.Camera.ActorName = Object->GetFName();
+	DataWrapper.Camera.ActorAssetPath = FSoftObjectPath(Object->GetClass()).GetAssetPath();
+	DataWrapper.Camera.MapAssetPath = FTopLevelAssetPath(World);
+	DataWrapper.Position = Position;
+	DataWrapper.Rotation = Rotation;
+	DataWrapper.FOV = FOV;
+
+	FChaosVDScopedTLSBufferAccessor TLSDataBuffer;
+	Chaos::VisualDebugger::WriteDataToBuffer(TLSDataBuffer.BufferRef, DataWrapper);
+
+	TraceBinaryData(TLSDataBuffer.BufferRef, FChaosVDCameraDataWrapper::WrapperTypeName);
+}
+
 void FChaosVisualDebuggerTrace::SetupForFullCaptureIfNeeded(int32 SolverID, bool& bOutFullCaptureRequested)
 {
 	SessionState.DeltaRecordingStatesLock.ReadLock();
diff --git a/Engine/Source/Runtime/Experimental/Chaos/Public/ChaosVisualDebugger/ChaosVDOptionalDataChannel.h b/Engine/Source/Runtime/Experimental/Chaos/Public/ChaosVisualDebugger/ChaosVDOptionalDataChannel.h
--- a/Engine/Source/Runtime/Experimental/Chaos/Public/ChaosVisualDebugger/ChaosVDOptionalDataChannel.h
+++ b/Engine/Source/Runtime/Experimental/Chaos/Public/ChaosVisualDebugger/ChaosVDOptionalDataChannel.h
@@ -220,6 +220,8 @@ CVD_DECLARE_OPTIONAL_DATA_CHANNEL(AccelerationStructures);
 // If this channel is enabled, we will record all child particles of any active dirty cluster particle
 CVD_DECLARE_OPTIONAL_DATA_CHANNEL(ClusterParticlesChildData);
 
+CVD_DECLARE_OPTIONAL_DATA_CHANNEL(CameraPaths);
+
 #else  //WITH_CHAOS_VISUAL_DEBUGGER
 
 /** Declares an Optional CVD Data channel to be available globally. The data channel can be accessed by using CVDDC_TheNameOfTheChannelUsedWithThisMacro */
diff --git a/Engine/Source/Runtime/Experimental/Chaos/Public/ChaosVisualDebugger/ChaosVDTraceMacros.h b/Engine/Source/Runtime/Experimental/Chaos/Public/ChaosVisualDebugger/ChaosVDTraceMacros.h
--- a/Engine/Source/Runtime/Experimental/Chaos/Public/ChaosVisualDebugger/ChaosVDTraceMacros.h
+++ b/Engine/Source/Runtime/Experimental/Chaos/Public/ChaosVisualDebugger/ChaosVDTraceMacros.h
@@ -54,6 +54,11 @@
 		FChaosVisualDebuggerTrace::TraceParticlesSoA(ParticleSoA, ##__VA_ARGS__);
 	#endif
 
+	#ifndef CVD_TRACE_CAMERA
+		#define CVD_TRACE_CAMERA(Object, World, CameraName, Position, Rotation, FOV) \
+		FChaosVisualDebuggerTrace::TraceCamera(Object, World, CameraName, Position, Rotation, FOV);
+	#endif
+
 	#ifndef CVD_TRACE_SOLVER_START_FRAME
 		#define CVD_TRACE_SOLVER_START_FRAME(SolverType, SolverRef) \
 			FChaosVDContext StartContextData; \
@@ -234,6 +239,10 @@
 		#define CVD_TRACE_PARTICLES(ParticleHandles)
 	#endif
 
+	#ifndef CVD_TRACE_CAMERA
+		#define CVD_TRACE_CAMERA(Object, World, CameraName, Position, Rotation, FOV)
+	#endif
+
 	#ifndef CVD_TRACE_SOLVER_START_FRAME
 		#define CVD_TRACE_SOLVER_START_FRAME(SolverType, SolverRef)
 	#endif
diff --git a/Engine/Source/Runtime/Experimental/Chaos/Public/ChaosVisualDebugger/ChaosVisualDebuggerTrace.h b/Engine/Source/Runtime/Experimental/Chaos/Public/ChaosVisualDebugger/ChaosVisualDebuggerTrace.h
--- a/Engine/Source/Runtime/Experimental/Chaos/Public/ChaosVisualDebugger/ChaosVisualDebuggerTrace.h
+++ b/Engine/Source/Runtime/Experimental/Chaos/Public/ChaosVisualDebugger/ChaosVisualDebuggerTrace.h
@@ -160,6 +160,8 @@ struct FCollisionObjectQueryParams;
 struct FCollisionResponseParams;
 struct FCollisionQueryParams;
 enum ECollisionChannel : int;
+class AActor;
+class UWorld;
 
 namespace Chaos
 {
@@ -326,6 +328,10 @@ public:
 	 */
 	static CHAOS_API void TraceParticlesSoA(const Chaos::FPBDRigidsSOAs& ParticlesSoA, Chaos::FRigidClustering* ClusteringData = nullptr);
 
+
+	/** Traces data from cameras */
+	static CHAOS_API void TraceCamera(AActor* Object, UWorld* World, FName CameraName, const FVector& Position, const FQuat& Rotation, float FOV);
+
 	/** Traces the provided particle view in parallel */
 	template<typename ParticleType>
 	static void TraceParticlesView(const Chaos::TParticleView<ParticleType>& ParticlesView);
diff --git a/Engine/Source/Runtime/Experimental/ChaosVisualDebugger/Private/DataWrappers/ChaosVDCameraDataWrapper.cpp b/Engine/Source/Runtime/Experimental/ChaosVisualDebugger/Private/DataWrappers/ChaosVDCameraDataWrapper.cpp
new file mode 100644
--- /dev/null
+++ b/Engine/Source/Runtime/Experimental/ChaosVisualDebugger/Private/DataWrappers/ChaosVDCameraDataWrapper.cpp
@@ -0,0 +1,21 @@
+// Copyright Epic Games, Inc. All Rights Reserved.
+
+#include "DataWrappers/ChaosVDCameraDataWrapper.h"
+
+#include UE_INLINE_GENERATED_CPP_BY_NAME(ChaosVDCameraDataWrapper)
+
+FStringView FChaosVDCameraDataWrapper::WrapperTypeName = TEXT("FChaosVDCameraDataWrapper");
+
+bool FChaosVDCameraDataWrapper::Serialize(FArchive& Ar)
+{
+	Ar << Camera.CameraName;
+	Ar << Camera.ActorName;
+	Ar << Camera.ActorAssetPath;
+	Ar << Camera.MapAssetPath;
+	Ar << Position;
+	Ar << Rotation;
+	Ar << FOV;
+	
+	return !Ar.IsError();
+}
+
diff --git a/Engine/Source/Runtime/Experimental/ChaosVisualDebugger/Public/DataWrappers/ChaosVDCameraDataWrapper.h b/Engine/Source/Runtime/Experimental/ChaosVisualDebugger/Public/DataWrappers/ChaosVDCameraDataWrapper.h
new file mode 100644
--- /dev/null
+++ b/Engine/Source/Runtime/Experimental/ChaosVisualDebugger/Public/DataWrappers/ChaosVDCameraDataWrapper.h
@@ -0,0 +1,72 @@
+// Copyright Epic Games, Inc. All Rights Reserved.
+
+#pragma once
+#include "ChaosVDDataSerializationMacros.h"
+
+#include "ChaosVDCameraDataWrapper.generated.h"
+
+USTRUCT()
+struct FChaosVDCameraIdentifier
+{
+	GENERATED_BODY()
+
+	UPROPERTY()
+	FName CameraName;
+
+	UPROPERTY()
+	FName ActorName;
+
+	UPROPERTY()
+	FTopLevelAssetPath ActorAssetPath;
+
+	UPROPERTY()
+	FTopLevelAssetPath MapAssetPath;
+
+	bool operator==(const FChaosVDCameraIdentifier& CameraIdentifier) const
+	{
+		return CameraName == CameraIdentifier.CameraName && ActorName == CameraIdentifier.ActorName &&
+			MapAssetPath == CameraIdentifier.MapAssetPath && ActorAssetPath == CameraIdentifier.ActorAssetPath; 
+	}
+
+	friend uint32 GetTypeHash(const FChaosVDCameraIdentifier& Other)
+	{
+		uint32 Hash = 0;
+		Hash = HashCombine(Hash, GetTypeHash(Other.CameraName));
+		Hash = HashCombine(Hash, GetTypeHash(Other.ActorName));
+		Hash = HashCombine(Hash, GetTypeHash(Other.ActorAssetPath));
+		Hash = HashCombine(Hash, GetTypeHash(Other.MapAssetPath));
+
+		return Hash;
+	}
+};
+
+USTRUCT()
+struct FChaosVDCameraDataWrapper
+{
+	GENERATED_BODY()
+
+	CHAOSVDRUNTIME_API static FStringView WrapperTypeName;
+
+	UPROPERTY()
+	FChaosVDCameraIdentifier Camera;
+
+	UPROPERTY()
+	FVector Position;
+
+	UPROPERTY()
+	FQuat Rotation;
+
+	UPROPERTY()
+	float FOV;
+
+	CHAOSVDRUNTIME_API bool Serialize(FArchive& Ar);
+
+};
+
+USTRUCT()
+struct FChaosVDCameraDataContainer
+{
+	GENERATED_BODY()
+
+	TArray<TSharedPtr<FChaosVDCameraDataWrapper>> CameraData;
+};
