Fixing case of include.

#Synced-CL 48400404

diff --git a/Engine/Plugins/ChaosVD/Source/ChaosVDBuiltInExtensions/Private/CameraTraces/ChaosVDCameraTracesExtension.cpp b/Engine/Plugins/ChaosVD/Source/ChaosVDBuiltInExtensions/Private/CameraTraces/ChaosVDCameraTracesExtension.cpp
--- a/Engine/Plugins/ChaosVD/Source/ChaosVDBuiltInExtensions/Private/CameraTraces/ChaosVDCameraTracesExtension.cpp
+++ b/Engine/Plugins/ChaosVD/Source/ChaosVDBuiltInExtensions/Private/CameraTraces/ChaosVDCameraTracesExtension.cpp
@@ -17,7 +17,7 @@
 #include "Widgets/SWidget.h"
 #include "ChaosVDPlaybackController.h"
 #include "ChaosVDScene.h"
-#include "Components/ChaosVdCameraDataComponent.h"
+#include "Components/ChaosVDCameraDataComponent.h"
 #include "ChaosVDEngine.h"
 
 #define LOCTEXT_NAMESPACE "ChaosVisualDebugger"
