Fix header casing

#Synced-CL 48457324

diff --git a/Engine/Plugins/ChaosVD/Source/ChaosVDBuiltInExtensions/Private/CameraTraces/Visualizers/ChaosVDCameraDataComponentVisualizer.cpp b/Engine/Plugins/ChaosVD/Source/ChaosVDBuiltInExtensions/Private/CameraTraces/Visualizers/ChaosVDCameraDataComponentVisualizer.cpp
--- a/Engine/Plugins/ChaosVD/Source/ChaosVDBuiltInExtensions/Private/CameraTraces/Visualizers/ChaosVDCameraDataComponentVisualizer.cpp
+++ b/Engine/Plugins/ChaosVD/Source/ChaosVDBuiltInExtensions/Private/CameraTraces/Visualizers/ChaosVDCameraDataComponentVisualizer.cpp
@@ -2,7 +2,7 @@
 
 #include "ChaosVDCameraDataComponentVisualizer.h"
 
-#include "CameraTraces/Components/ChaosVdCameraDataComponent.h"
+#include "CameraTraces/Components/ChaosVDCameraDataComponent.h"
 #include "CameraTraces/Settings/ChaosVDCameraDataSettings.h"
 #include "Actors/ChaosVDDataContainerBaseActor.h"
 #include "ChaosVDScene.h"
