[CVD] extracted and refactored remote connection functionality in order to reuse it for RewindDebugger.
Added modules TracedBasedDebuggers and FNTracedBasedDebuggers to hold the shareable functionalities.
New modules/types are guarded by WITH_TRACE_BASED_DEBUGGERS
Macro WITH_CHAOS_VISUAL_DEBUGGER_EXTERNAL_MESSAGING is replaced by WITH_TRACE_BASED_DEBUGGERS_EXTERNAL_MESSAGING
Specific debuggers can add custom payload to the session info struct and to the session info message using FInstancedStruct.
e.g.,
The struct FChaosVDFullSessionInfoResponseData is stored as the DebuggerSpecificData for the message type UE::TraceBasedDebuggers::FFullSessionInfoResponseMessage.
The struct FChaosVDSessionData is stored as the DebuggerSpecificSessionData inside UE::TraceBasedDebuggers::FSessionInfo.

#Synced-CL 48614523

diff --git a/Engine/Plugins/ChaosVD/Source/ChaosVD/ChaosVD.Build.cs b/Engine/Plugins/ChaosVD/Source/ChaosVD/ChaosVD.Build.cs
--- a/Engine/Plugins/ChaosVD/Source/ChaosVD/ChaosVD.Build.cs
+++ b/Engine/Plugins/ChaosVD/Source/ChaosVD/ChaosVD.Build.cs
@@ -31,8 +31,9 @@ public class ChaosVD : ModuleRules
 				"Slate",
 				"SlateCore", 
 				"OutputLog",
-				"WorkspaceMenuStructure", 
+				"WorkspaceMenuStructure",
 				"TraceAnalysis", 
+				"TraceBasedDebuggers",
 				"TraceInsights",
 				"TraceLog", 
 				"GeometryFramework",
diff --git a/Engine/Plugins/ChaosVD/Source/ChaosVD/Private/ChaosVDEngine.cpp b/Engine/Plugins/ChaosVD/Source/ChaosVD/Private/ChaosVDEngine.cpp
--- a/Engine/Plugins/ChaosVD/Source/ChaosVD/Private/ChaosVDEngine.cpp
+++ b/Engine/Plugins/ChaosVD/Source/ChaosVD/Private/ChaosVDEngine.cpp
@@ -2,12 +2,13 @@
 
 #include "ChaosVDEngine.h"
 
+#include "Chaos/ChaosVDEngineEditorBridge.h"
 #include "ChaosVDModule.h"
 #include "ChaosVDPlaybackController.h"
 #include "ChaosVDScene.h"
 #include "ChaosVDSettingsManager.h"
-#include "Chaos/ChaosVDEngineEditorBridge.h"
 #include "ChaosVisualDebugger/ChaosVDOptionalDataChannel.h"
+#include "RemoteSessionsManager.h"
 #include "Settings/ChaosVDMiscSettings.h"
 #include "Trace/ChaosVDCombinedTraceFile.h"
 #include "Trace/ChaosVDTraceManager.h"
@@ -28,9 +29,9 @@ void FChaosVDEngine::Initialize()
 	PlaybackController = MakeShared<FChaosVDPlaybackController>(CurrentScene);
 	PlaybackController->Initialize();
 
-	if (TSharedPtr<FChaosVDRemoteSessionsManager> RemoteSessionManager = FChaosVDEngineEditorBridge::Get().GetRemoteSessionsManager())
+	if (TSharedPtr<UE::TraceBasedDebuggers::FRemoteSessionsManager> RemoteSessionManager = FChaosVDEngineEditorBridge::Get().GetSessionsManager())
 	{
-		LiveSessionStoppedDelegateHandle = RemoteSessionManager->OnSessionRecordingStopped().AddLambda([WeakThis = AsWeak()](TWeakPtr<FChaosVDSessionInfo> Session)
+		LiveSessionStoppedDelegateHandle = RemoteSessionManager->OnSessionRecordingStopped().AddLambda([WeakThis = AsWeak()](TWeakPtr<UE::TraceBasedDebuggers::FSessionInfo> Session)
 		{
 			if (const TSharedPtr<FChaosVDEngine> CVDEngine = WeakThis.Pin())
 			{
@@ -121,7 +122,7 @@ void FChaosVDEngine::DeInitialize()
 
 	CloseActiveTraceSessions();
 
-	if (TSharedPtr<FChaosVDRemoteSessionsManager> RemoteSessionManager = FChaosVDEngineEditorBridge::Get().GetRemoteSessionsManager())
+	if (TSharedPtr<UE::TraceBasedDebuggers::FRemoteSessionsManager> RemoteSessionManager = FChaosVDEngineEditorBridge::Get().GetSessionsManager())
 	{
 		RemoteSessionManager->OnSessionRecordingStopped().Remove(LiveSessionStoppedDelegateHandle);
 	}
diff --git a/Engine/Plugins/ChaosVD/Source/ChaosVD/Private/Settings/ChaosVDGeneralSettings.cpp b/Engine/Plugins/ChaosVD/Source/ChaosVD/Private/Settings/ChaosVDGeneralSettings.cpp
--- a/Engine/Plugins/ChaosVD/Source/ChaosVD/Private/Settings/ChaosVDGeneralSettings.cpp
+++ b/Engine/Plugins/ChaosVD/Source/ChaosVD/Private/Settings/ChaosVDGeneralSettings.cpp
@@ -6,7 +6,7 @@
 
 #include UE_INLINE_GENERATED_CPP_BY_NAME(ChaosVDGeneralSettings)
 
-EChaosVDTransportMode UChaosVDGeneralSettings::GetTransportModeForTarget(EBuildTargetType TargetType) const
+UE::TraceBasedDebuggers::ETraceTransportMode UChaosVDGeneralSettings::GetTransportModeForTargetType(EBuildTargetType TargetType) const
 {
 	switch (TargetType)
 	{
@@ -19,11 +19,11 @@ EChaosVDTransportMode UChaosVDGeneralSettings::GetTransportModeForTarget(EBuildT
 	case EBuildTargetType::Unknown:
 	case EBuildTargetType::Program:
 	default:
-		return EChaosVDTransportMode::Direct;
+		return UE::TraceBasedDebuggers::ETraceTransportMode::Direct;
 	}
 }
 
-void UChaosVDGeneralSettings::SetTransportModeForTarget(EBuildTargetType TargetType, EChaosVDTransportMode NewTransportMode)
+void UChaosVDGeneralSettings::SetTransportModeForTargetType(EBuildTargetType TargetType, UE::TraceBasedDebuggers::ETraceTransportMode NewTransportMode)
 {
 	switch (TargetType)
 	{
@@ -51,3 +51,15 @@ void UChaosVDGeneralSettings::SetTransportModeForTarget(EBuildTargetType TargetT
 		}
 	}
 }
+
+PRAGMA_DISABLE_DEPRECATION_WARNINGS
+EChaosVDTransportMode UChaosVDGeneralSettings::GetTransportModeForTarget(EBuildTargetType TargetType) const
+{
+	return static_cast<EChaosVDTransportMode>(GetTransportModeForTargetType(TargetType));
+}
+
+void UChaosVDGeneralSettings::SetTransportModeForTarget(EBuildTargetType TargetType, EChaosVDTransportMode NewTransportMode)
+{
+	SetTransportModeForTargetType(TargetType, static_cast<UE::TraceBasedDebuggers::ETraceTransportMode>(NewTransportMode));
+}
+PRAGMA_ENABLE_DEPRECATION_WARNINGS
diff --git a/Engine/Plugins/ChaosVD/Source/ChaosVD/Private/Trace/ChaosVDTraceManager.cpp b/Engine/Plugins/ChaosVD/Source/ChaosVD/Private/Trace/ChaosVDTraceManager.cpp
--- a/Engine/Plugins/ChaosVD/Source/ChaosVD/Private/Trace/ChaosVDTraceManager.cpp
+++ b/Engine/Plugins/ChaosVD/Source/ChaosVD/Private/Trace/ChaosVDTraceManager.cpp
@@ -2,18 +2,17 @@
 
 #include "Trace/ChaosVDTraceManager.h"
 
+#include "Chaos/ChaosVDEngineEditorBridge.h"
+#include "Chaos/Serialization/SerializedDataBuffer.h"
 #include "ChaosVDEngine.h"
 #include "ChaosVDModule.h"
 #include "ChaosVDRuntimeModule.h"
 #include "ChaosVDSettingsManager.h"
-#include "Chaos/ChaosVDEngineEditorBridge.h"
-#include "Chaos/ChaosVDTraceRelayTransport.h"
-#include "Chaos/Serialization/SerializedDataBuffer.h"
 #include "Features/IModularFeatures.h"
 #include "HAL/PlatformFileManager.h"
-#include "Modules/ModuleManager.h"
 #include "Misc/Paths.h"
-#include "ProfilingDebugging/TraceAuxiliary.h"
+#include "Modules/ModuleManager.h"
+#include "RemoteSessionsManager.h"
 #include "Settings/ChaosVDGeneralSettings.h"
 #include "Trace/ChaosVDTraceModule.h"
 #include "Trace/StoreClient.h"
@@ -107,8 +106,9 @@ namespace Chaos::VD
 
 		ElapsedTimeSinceLastUpdate = 0.0;
 
-		TSharedPtr<FChaosVDRemoteSessionsManager> RemoteSessionManager = FChaosVDEngineEditorBridge::Get().GetRemoteSessionsManager();
-		TSharedPtr<FChaosVDSessionInfo> SessionInfo = RemoteSessionManager ? RemoteSessionManager->GetSessionInfo(DataStreamInstance->GetOwningRemoteSessionID()).Pin() : nullptr;
+		using namespace UE::TraceBasedDebuggers;
+		TSharedPtr<FRemoteSessionsManager> RemoteSessionManager = FChaosVDEngineEditorBridge::Get().GetSessionsManager();
+		TSharedPtr<FSessionInfo> SessionInfo = RemoteSessionManager ? RemoteSessionManager->GetSessionInfo(DataStreamInstance->GetOwningRemoteSessionID()).Pin() : nullptr;
 		if (!SessionInfo)
 		{
 			return false;
@@ -154,10 +154,11 @@ namespace Chaos::VD
 		StatsUpdater(this)
 	{
 #if WITH_CHAOS_VISUAL_DEBUGGER
-		TSharedPtr<ITraceDataRelayTransport> RelayTransportInstance = FChaosVDEngineEditorBridge::Get().GetRelayTransportInstance();
+		using namespace UE::TraceBasedDebuggers;
+		TSharedPtr<IDataRelayTransport> RelayTransportInstance = FChaosVDEngineEditorBridge::Get().GetTraceRelayTransportInstance();
 		if (ensure(RelayTransportInstance))
 		{
-			RelayTransportInstance->RegisterRelayDataReceiverForSessionID(RemoteSessionID, ITraceDataRelayTransport::FProcessReceivedRelayDataDelegate::CreateRaw(this, &FRelayDataStream::EnqueueRelayedData));
+			RelayTransportInstance->RegisterRelayDataReceiverForSessionID(RemoteSessionID, IDataRelayTransport::FProcessReceivedRelayDataDelegate::CreateRaw(this, &FRelayDataStream::EnqueueRelayedData));
 		}
 #endif
 	}
@@ -171,10 +172,11 @@ namespace Chaos::VD
 		StatsUpdater(this)
 	{
 #if WITH_CHAOS_VISUAL_DEBUGGER
-		TSharedPtr<ITraceDataRelayTransport> RelayTransportInstance = FChaosVDEngineEditorBridge::Get().GetRelayTransportInstance();
+		using namespace UE::TraceBasedDebuggers;
+		TSharedPtr<IDataRelayTransport> RelayTransportInstance = FChaosVDEngineEditorBridge::Get().GetTraceRelayTransportInstance();
 		if (ensure(RelayTransportInstance))
 		{
-			RelayTransportInstance->RegisterRelayDataReceiverForSessionID(RemoteSessionID, ITraceDataRelayTransport::FProcessReceivedRelayDataDelegate::CreateRaw(this, &FRelayDataStream::EnqueueRelayedData));
+			RelayTransportInstance->RegisterRelayDataReceiverForSessionID(RemoteSessionID, IDataRelayTransport::FProcessReceivedRelayDataDelegate::CreateRaw(this, &FRelayDataStream::EnqueueRelayedData));
 		}
 #endif
 	}
@@ -185,11 +187,12 @@ namespace Chaos::VD
 		FPlatformProcess::ReturnSynchEventToPool(WaitForDataEvent);
 
 #if WITH_CHAOS_VISUAL_DEBUGGER
-		TSharedPtr<ITraceDataRelayTransport> RelayTransportInstance = FChaosVDEngineEditorBridge::Get().GetRelayTransportInstance();
+		using namespace UE::TraceBasedDebuggers;
+		TSharedPtr<IDataRelayTransport> RelayTransportInstance = FChaosVDEngineEditorBridge::Get().GetTraceRelayTransportInstance();
 		if (RelayTransportInstance)
 		{
 			RelayTransportInstance->UnregisterRelayDataReceiverForSessionID(RemoteSessionID);
-		}	
+		}
 #endif
 	}
 
@@ -476,7 +479,7 @@ UE::Trace::FStoreClient::FTraceData FChaosVDTraceManager::GetTraceDataStreamFrom
 	FStoreClient* StoreClient = FStoreClient::Connect(InSessionHost.GetData());
 	if (!StoreClient)
 	{
-		return UE::Trace::FStoreClient::FTraceData();
+		return FStoreClient::FTraceData();
 	}
 
 	return StoreClient->ReadTrace(SessionID);
diff --git a/Engine/Plugins/ChaosVD/Source/ChaosVD/Private/Widgets/ChaosVDRecordingControls.cpp b/Engine/Plugins/ChaosVD/Source/ChaosVD/Private/Widgets/ChaosVDRecordingControls.cpp
--- a/Engine/Plugins/ChaosVD/Source/ChaosVD/Private/Widgets/ChaosVDRecordingControls.cpp
+++ b/Engine/Plugins/ChaosVD/Source/ChaosVD/Private/Widgets/ChaosVDRecordingControls.cpp
@@ -3,32 +3,33 @@
 #include "Widgets/ChaosVDRecordingControls.h"
 
 #include "AsyncCompilationHelpers.h"
+#include "Chaos/ChaosVDEngineEditorBridge.h"
+#include "Chaos/ChaosVDRemoteSessionsManager.h"
 #include "ChaosVDEngine.h"
 #include "ChaosVDModule.h"
-#include "ChaosVDRecordingDetails.h"
 #include "ChaosVDScene.h"
 #include "ChaosVDSettingsManager.h"
 #include "ChaosVDStyle.h"
 #include "ChaosVisualDebugger/ChaosVisualDebuggerTrace.h"
 #include "Editor.h"
+#include "Framework/MultiBox/MultiBoxBuilder.h"
+#include "Input/Reply.h"
+#include "Misc/MessageDialog.h"
+#include "Misc/ScopedSlowTask.h"
+#include "RemoteSessionsManager.h"
 #include "SChaosVDNameListPicker.h"
 #include "SEnumCombo.h"
+#include "Settings/ChaosVDGeneralSettings.h"
 #include "SocketSubsystem.h"
-#include "Input/Reply.h"
-#include "Misc/MessageDialog.h"
 #include "StatusBarSubsystem.h"
 #include "ToolMenu.h"
 #include "ToolMenus.h"
-#include "Chaos/ChaosVDEngineEditorBridge.h"
-#include "Chaos/ChaosVDTraceRelayTransport.h"
-#include "Framework/MultiBox/MultiBoxBuilder.h"
-#include "Misc/ScopedSlowTask.h"
-#include "Settings/ChaosVDGeneralSettings.h"
+#include "TraceDataRelayTransport.h"
 #include "Widgets/Images/SImage.h"
 #include "Widgets/Input/SButton.h"
-#include "Widgets/SChaosVDMainTab.h"
 #include "Widgets/Input/SComboButton.h"
 #include "Widgets/Notifications/SNotificationList.h"
+#include "Widgets/SChaosVDMainTab.h"
 #include "Widgets/Text/STextBlock.h"
 
 #include UE_INLINE_GENERATED_CPP_BY_NAME(ChaosVDRecordingControls)
@@ -45,7 +46,7 @@ FChaosVDRecordingControls::FChaosVDRecordingControls(const TSharedRef<SChaosVDMa
 	RecordingAnimation = FCurveSequence();
 	RecordingAnimation.AddCurve(0.f, 1.5f, ECurveEaseFunction::Linear);
 
-	CurrentSelectedSessionId = FChaosVDRemoteSessionsManager::LocalEditorSessionID;
+	CurrentSelectedSessionId = UE::TraceBasedDebuggers::FRemoteSessionsManager::LocalEditorSessionID;
 }
 
 FChaosVDRecordingControls::~FChaosVDRecordingControls()
@@ -60,7 +61,7 @@ void FChaosVDRecordingControls::Initialize()
 		return;
 	}
 
-	if (TSharedPtr<FChaosVDRemoteSessionsManager> RemoteSessionManager = FChaosVDEngineEditorBridge::Get().GetRemoteSessionsManager())
+	if (const TSharedPtr<UE::TraceBasedDebuggers::FRemoteSessionsManager> RemoteSessionManager = FChaosVDEngineEditorBridge::Get().GetSessionsManager())
 	{
 		RemoteSessionManager->OnSessionRecordingStarted().AddSP(this, &FChaosVDRecordingControls::HandleRecordingStart);
 		RemoteSessionManager->OnSessionRecordingStopped().AddSP(this, &FChaosVDRecordingControls::HandleRecordingStop);
@@ -78,7 +79,7 @@ void FChaosVDRecordingControls::Deinitialize()
 		return;
 	}
 
-	if (TSharedPtr<FChaosVDRemoteSessionsManager> RemoteSessionManager = FChaosVDEngineEditorBridge::Get().GetRemoteSessionsManager())
+	if (const TSharedPtr<UE::TraceBasedDebuggers::FRemoteSessionsManager> RemoteSessionManager = FChaosVDEngineEditorBridge::Get().GetSessionsManager())
 	{
 		RemoteSessionManager->OnSessionRecordingStarted().RemoveAll(this);
 		RemoteSessionManager->OnSessionRecordingStopped().RemoveAll(this);
@@ -146,7 +147,7 @@ TSharedRef<SWidget> FChaosVDRecordingControls::GenerateToggleRecordingStateButto
 
 FText FChaosVDRecordingControls::GetCurrentSelectedSessionName() const
 {
-	if (TSharedPtr<FChaosVDSessionInfo> CurrentSessionPtr = GetCurrentSessionInfo())
+	if (const TSharedPtr<UE::TraceBasedDebuggers::FSessionInfo> CurrentSessionPtr = GetCurrentSessionInfo())
 	{
 		return FText::AsCultureInvariant(CurrentSessionPtr->SessionName);
 	}
@@ -176,22 +177,22 @@ TSharedRef<SWidget> FChaosVDRecordingControls::GenerateTargetSessionSelector()
 
 TSharedRef<SWidget> FChaosVDRecordingControls::GenerateTargetSessionDropdown()
 {
-	using namespace Chaos::VisualDebugger;
+	using namespace UE::TraceBasedDebuggers;
 
 	FMenuBuilder MenuBuilder(true, nullptr);
 
 	MenuBuilder.BeginSection("CVDRecordingWidgetTargets", LOCTEXT("CVDRecordingTargetsMenu", "Available Targets"));
 	{
-		if (TSharedPtr<FChaosVDRemoteSessionsManager> RemoteSessionManager = FChaosVDEngineEditorBridge::Get().GetRemoteSessionsManager())
+		if (const TSharedPtr<FRemoteSessionsManager> RemoteSessionManager = FChaosVDEngineEditorBridge::Get().GetSessionsManager())
 		{
-			RemoteSessionManager->EnumerateActiveSessions([this, &MenuBuilder](const TSharedRef<FChaosVDSessionInfo>& InSessionInfoRef)
+			RemoteSessionManager->EnumerateActiveSessions([this, &MenuBuilder](const TSharedRef<FSessionInfo>& InSessionInfoRef)
 			{
-				if (EnumHasAnyFlags(InSessionInfoRef->GetSessionTypeAttributes(), EChaosVDRemoteSessionAttributes::IsMultiSessionWrapper))
+				if (EnumHasAnyFlags(InSessionInfoRef->GetSessionTypeAttributes(), ERemoteSessionAttributes::IsMultiSessionWrapper))
 				{
 					return true;
 				}
-	
-				FText SessionNameAsText = FText::AsCultureInvariant(InSessionInfoRef->SessionName);
+
+				const FText SessionNameAsText = FText::AsCultureInvariant(InSessionInfoRef->SessionName);
 				MenuBuilder.AddMenuEntry(
 				SessionNameAsText,
 				FText::Format(LOCTEXT("SingleTargetItemTooltip", "Select {0} session as current target"), SessionNameAsText),
@@ -201,48 +202,48 @@ TSharedRef<SWidget> FChaosVDRecordingControls::GenerateTargetSessionDropdown()
 			});
 		}
 	}
-	
+
 	MenuBuilder.EndSection();
 
 	MenuBuilder.BeginSection("CVDRecordingWidgetTargetsMulti", LOCTEXT("CVDRecordingMultiTargetsMenu", "Multi Target"));
 
-	FText AllRemoteTargetsLabel = LOCTEXT("AllRemoteOption", "All Remote");
+	const FText AllRemoteTargetsLabel = LOCTEXT("AllRemoteOption", "All Remote");
 	MenuBuilder.AddMenuEntry(
 			AllRemoteTargetsLabel,
 			LOCTEXT("MultiRemoteTargetTooltip", "Select this to act on all remote targets"),
-			GetIconForSession(FChaosVDRemoteSessionsManager::AllRemoteSessionsWrapperGUID),
-			FUIAction(FExecuteAction::CreateSP(this, &FChaosVDRecordingControls::SelectTargetSession, FChaosVDRemoteSessionsManager::AllRemoteSessionsWrapperGUID),
-			FCanExecuteAction::CreateSP(this, &FChaosVDRecordingControls::CanSelectMultiSessionTarget, FChaosVDRemoteSessionsManager::AllRemoteSessionsWrapperGUID),
+			GetIconForSession(FRemoteSessionsManager::AllRemoteSessionsWrapperGUID),
+			FUIAction(FExecuteAction::CreateSP(this, &FChaosVDRecordingControls::SelectTargetSession, FRemoteSessionsManager::AllRemoteSessionsWrapperGUID),
+			FCanExecuteAction::CreateSP(this, &FChaosVDRecordingControls::CanSelectMultiSessionTarget, FRemoteSessionsManager::AllRemoteSessionsWrapperGUID),
 			EUIActionRepeatMode::RepeatDisabled));
 
-	FText AllRemoteServersTargetsLabel = LOCTEXT("AllRemoteServersOption", "All Remote Servers");
+	const FText AllRemoteServersTargetsLabel = LOCTEXT("AllRemoteServersOption", "All Remote Servers");
 	MenuBuilder.AddMenuEntry(
 			AllRemoteServersTargetsLabel,
 			LOCTEXT("MultiRemoteServerTargetTooltip", "Select this to act on all remote server targets"),
-			GetIconForSession(FChaosVDRemoteSessionsManager::AllRemoteServersWrapperGUID),
-			FUIAction(FExecuteAction::CreateSP(this, &FChaosVDRecordingControls::SelectTargetSession, FChaosVDRemoteSessionsManager::AllRemoteServersWrapperGUID),
-			FCanExecuteAction::CreateSP(this, &FChaosVDRecordingControls::CanSelectMultiSessionTarget, FChaosVDRemoteSessionsManager::AllRemoteServersWrapperGUID),
+			GetIconForSession(FRemoteSessionsManager::AllRemoteServersWrapperGUID),
+			FUIAction(FExecuteAction::CreateSP(this, &FChaosVDRecordingControls::SelectTargetSession, FRemoteSessionsManager::AllRemoteServersWrapperGUID),
+			FCanExecuteAction::CreateSP(this, &FChaosVDRecordingControls::CanSelectMultiSessionTarget, FRemoteSessionsManager::AllRemoteServersWrapperGUID),
 			EUIActionRepeatMode::RepeatDisabled));
 
-	FText AllRemoteClientsTargetsLabel = LOCTEXT("AllRemoteClientsOption", "All Remote Clients");
+	const FText AllRemoteClientsTargetsLabel = LOCTEXT("AllRemoteClientsOption", "All Remote Clients");
 	MenuBuilder.AddMenuEntry(
 			AllRemoteClientsTargetsLabel,
 			LOCTEXT("MultiRemoteClientTargetTooltip", "Select this to act on all remote client targets"),
-			GetIconForSession(FChaosVDRemoteSessionsManager::AllRemoteClientsWrapperGUID),
-			FUIAction(FExecuteAction::CreateSP(this, &FChaosVDRecordingControls::SelectTargetSession, FChaosVDRemoteSessionsManager::AllRemoteClientsWrapperGUID),
-			FCanExecuteAction::CreateSP(this, &FChaosVDRecordingControls::CanSelectMultiSessionTarget, FChaosVDRemoteSessionsManager::AllRemoteClientsWrapperGUID),
+			GetIconForSession(FRemoteSessionsManager::AllRemoteClientsWrapperGUID),
+			FUIAction(FExecuteAction::CreateSP(this, &FChaosVDRecordingControls::SelectTargetSession, FRemoteSessionsManager::AllRemoteClientsWrapperGUID),
+			FCanExecuteAction::CreateSP(this, &FChaosVDRecordingControls::CanSelectMultiSessionTarget, FRemoteSessionsManager::AllRemoteClientsWrapperGUID),
 			EUIActionRepeatMode::RepeatDisabled));
 
-	FText AllTargets = LOCTEXT("AllTargetsOption", "All");
+	const FText AllTargets = LOCTEXT("AllTargetsOption", "All");
 	MenuBuilder.AddMenuEntry(
 		AllTargets,
 		LOCTEXT("MultiAllTargetTooltip", "Select this to act on all targets, both Local and Remote"),
-		GetIconForSession(FChaosVDRemoteSessionsManager::AllSessionsWrapperGUID),
-		FUIAction(FExecuteAction::CreateSP(this, &FChaosVDRecordingControls::SelectTargetSession, FChaosVDRemoteSessionsManager::AllSessionsWrapperGUID), EUIActionRepeatMode::RepeatDisabled));
+		GetIconForSession(FRemoteSessionsManager::AllSessionsWrapperGUID),
+		FUIAction(FExecuteAction::CreateSP(this, &FChaosVDRecordingControls::SelectTargetSession, FRemoteSessionsManager::AllSessionsWrapperGUID), EUIActionRepeatMode::RepeatDisabled));
 
 	MenuBuilder.AddMenuSeparator();
 
-	FText CustomTargets = LOCTEXT("CustomTargetsOption", "Custom Selection");
+	const FText CustomTargets = LOCTEXT("CustomTargetsOption", "Custom Selection");
 	MenuBuilder.AddSubMenu(CustomTargets,
 		LOCTEXT("MultiCustomTargetTooltip", "Select this to act on the specific targets you selected"),
 		FNewMenuDelegate::CreateSP(this, &FChaosVDRecordingControls::GenerateCustomTargetsMenu)
@@ -330,10 +331,11 @@ TSharedRef<SWidget> FChaosVDRecordingControls::GenerateAdvancedRecordOptionsMenu
 
 EVisibility FChaosVDRecordingControls::GetRecordingTimeTextBlockVisibility() const
 {
-	TSharedPtr<FChaosVDSessionInfo> SessionInfo = GetCurrentSessionInfo();
-	bool bIsVisible = SessionInfo && !EnumHasAnyFlags(SessionInfo->GetSessionTypeAttributes(), EChaosVDRemoteSessionAttributes::IsMultiSessionWrapper) && IsRecording();
+	using namespace UE::TraceBasedDebuggers;
+	const TSharedPtr<FSessionInfo> SessionInfo = GetCurrentSessionInfo();
+	const bool bIsVisible = SessionInfo && !EnumHasAnyFlags(SessionInfo->GetSessionTypeAttributes(), ERemoteSessionAttributes::IsMultiSessionWrapper) && IsRecording();
 
-	return  bIsVisible ? EVisibility::Visible : EVisibility::Collapsed;
+	return bIsVisible ? EVisibility::Visible : EVisibility::Collapsed;
 }
 
 TSharedRef<SWidget> FChaosVDRecordingControls::GenerateDataChannelsButton()
@@ -380,15 +382,15 @@ TSharedRef<SWidget> FChaosVDRecordingControls::GenerateLoadingModeSelector()
 
 TSharedRef<SWidget> FChaosVDRecordingControls::GenerateDataChannelsMenu()
 {
-	using namespace Chaos::VisualDebugger;
+	using namespace UE::TraceBasedDebuggers;
 
 	FMenuBuilder MenuBuilder(true, nullptr);
 
 	MenuBuilder.BeginSection("CVDRecordingWidget", LOCTEXT("CVDRecordingMenuChannels", "Data Channels"));
 	{
-		if (TSharedPtr<FChaosVDSessionInfo> SessionInfo = GetCurrentSessionInfo())
+		if (const FChaosVDSessionData* SessionData = GetChaosDataForCurrentSession())
 		{
-			for (const TPair<FString, FChaosVDDataChannelState>& DataChannelStateWithName : SessionInfo->DataChannelsStatesByName)
+			for (const TPair<FString, FChaosVDDataChannelState>& DataChannelStateWithName : SessionData->DataChannelsStatesByName)
 			{
 				FText ChannelNamesAsText = FText::AsCultureInvariant(DataChannelStateWithName.Key);
 				MenuBuilder.AddMenuEntry(
@@ -408,23 +410,31 @@ TSharedRef<SWidget> FChaosVDRecordingControls::GenerateDataChannelsMenu()
 
 void FChaosVDRecordingControls::ToggleChannelEnabledState(FString ChannelName)
 {
-	TSharedPtr<FChaosVDRemoteSessionsManager> RemoteSessionManager = FChaosVDEngineEditorBridge::Get().GetRemoteSessionsManager();
-	TSharedPtr<FChaosVDSessionInfo> SessionInfo = GetCurrentSessionInfo();
+	using namespace UE::TraceBasedDebuggers;
+	const TSharedPtr<FRemoteSessionsManager> RemoteSessionManager = FChaosVDEngineEditorBridge::Get().GetSessionsManager();
+	const TSharedPtr<FSessionInfo> SessionInfo = GetCurrentSessionInfo();
 	if (RemoteSessionManager && SessionInfo)
 	{
-		if (FChaosVDDataChannelState* ChannelState = SessionInfo->DataChannelsStatesByName.Find(ChannelName))
+		if (FChaosVDSessionData* ChaosData = SessionInfo->GetDebuggerData<FChaosVDSessionData>())
 		{
-			ChannelState->bWaitingUpdatedState = true;
-			RemoteSessionManager->SendDataChannelStateChangeCommand(SessionInfo->Address, {ChannelState->ChannelName, !ChannelState->bIsEnabled });
+			if (FChaosVDDataChannelState* ChannelState = ChaosData->DataChannelsStatesByName.Find(ChannelName))
+			{
+				ChannelState->bWaitingUpdatedState = true;
+				RemoteSessionManager->SendCommand(SessionInfo->Address, FChaosVDChannelStateChangeCommandMessage
+					{
+						FChaosVDDataChannelState{.ChannelName = ChannelState->ChannelName, .bIsEnabled = !ChannelState->bIsEnabled}
+					});
+			}
 		}
 	}
 }
 
 bool FChaosVDRecordingControls::IsChannelEnabled(FString ChannelName)
 {
-	if (TSharedPtr<FChaosVDSessionInfo> SessionInfo = GetCurrentSessionInfo())
+	using namespace UE::TraceBasedDebuggers;
+	if (const FChaosVDSessionData* ChaosData = GetChaosDataForCurrentSession())
 	{
-		if (FChaosVDDataChannelState* ChannelState = SessionInfo->DataChannelsStatesByName.Find(ChannelName))
+		if (const FChaosVDDataChannelState* ChannelState = ChaosData->DataChannelsStatesByName.Find(ChannelName))
 		{
 			return ChannelState->bIsEnabled;
 		}
@@ -435,9 +445,10 @@ bool FChaosVDRecordingControls::IsChannelEnabled(FString ChannelName)
 
 bool FChaosVDRecordingControls::CanChangeChannelEnabledState(FString ChannelName)
 {
-	if (TSharedPtr<FChaosVDSessionInfo> SessionInfo = GetCurrentSessionInfo())
+	using namespace UE::TraceBasedDebuggers;
+	if (const FChaosVDSessionData* ChaosData = GetChaosDataForCurrentSession())
 	{
-		if (FChaosVDDataChannelState* ChannelState = SessionInfo->DataChannelsStatesByName.Find(ChannelName))
+		if (const FChaosVDDataChannelState* ChannelState = ChaosData->DataChannelsStatesByName.Find(ChannelName))
 		{
 			return ChannelState->bCanChangeChannelState && !ChannelState->bWaitingUpdatedState;
 		}
@@ -453,16 +464,17 @@ void FChaosVDRecordingControls::SelectTargetSession(FGuid SessionId)
 
 void FChaosVDRecordingControls::GenerateCustomTargetsMenu(FMenuBuilder& MenuBuilder)
 {
-	if (TSharedPtr<FChaosVDRemoteSessionsManager> RemoteSessionManager = FChaosVDEngineEditorBridge::Get().GetRemoteSessionsManager())
+	using namespace UE::TraceBasedDebuggers;
+	if (const TSharedPtr<FRemoteSessionsManager> RemoteSessionManager = FChaosVDEngineEditorBridge::Get().GetSessionsManager())
 	{
-		RemoteSessionManager->EnumerateActiveSessions([this, &MenuBuilder](const TSharedRef<FChaosVDSessionInfo>& InSessionInfoRef)
+		RemoteSessionManager->EnumerateActiveSessions([this, &MenuBuilder](const TSharedRef<FSessionInfo>& InSessionInfoRef)
 		{
-			if (EnumHasAnyFlags(InSessionInfoRef->GetSessionTypeAttributes(), EChaosVDRemoteSessionAttributes::IsMultiSessionWrapper))
+			if (EnumHasAnyFlags(InSessionInfoRef->GetSessionTypeAttributes(), ERemoteSessionAttributes::IsMultiSessionWrapper))
 			{
 				return true;
 			}
-	
-			FText SessionNameAsText = FText::AsCultureInvariant(InSessionInfoRef->SessionName);
+
+			const FText SessionNameAsText = FText::AsCultureInvariant(InSessionInfoRef->SessionName);
 			MenuBuilder.AddMenuEntry(
 			SessionNameAsText,
 			FText::Format(LOCTEXT("MultiTargetItemTooltip", "Select {0} session as one of the current targets"), SessionNameAsText),
@@ -480,9 +492,10 @@ void FChaosVDRecordingControls::GenerateCustomTargetsMenu(FMenuBuilder& MenuBuil
 
 bool FChaosVDRecordingControls::IsSessionPartOfCustomTargetSelection(FGuid SessionGuid)
 {
-	if (TSharedPtr<FChaosVDRemoteSessionsManager> RemoteSessionManager = FChaosVDEngineEditorBridge::Get().GetRemoteSessionsManager())
+	using namespace UE::TraceBasedDebuggers;
+	if (const TSharedPtr<FRemoteSessionsManager> RemoteSessionManager = FChaosVDEngineEditorBridge::Get().GetSessionsManager())
 	{
-		if (TSharedPtr<FChaosVDMultiSessionInfo> CustomSessionTarget = StaticCastSharedPtr<FChaosVDMultiSessionInfo>(RemoteSessionManager->GetSessionInfo(FChaosVDRemoteSessionsManager::CustomSessionsWrapperGUID).Pin()))
+		if (TSharedPtr<FMultiSessionInfo> CustomSessionTarget = StaticCastSharedPtr<FMultiSessionInfo>(RemoteSessionManager->GetSessionInfo(FRemoteSessionsManager::CustomSessionsWrapperGUID).Pin()))
 		{
 			return CustomSessionTarget->InnerSessionsByInstanceID.Find(SessionGuid) != nullptr;
 		}
@@ -493,9 +506,10 @@ bool FChaosVDRecordingControls::IsSessionPartOfCustomTargetSelection(FGuid Sessi
 
 void FChaosVDRecordingControls::ToggleSessionSelectionInCustomTarget(FGuid SessionGuid)
 {
-	if (TSharedPtr<FChaosVDRemoteSessionsManager> RemoteSessionManager = FChaosVDEngineEditorBridge::Get().GetRemoteSessionsManager())
+	using namespace UE::TraceBasedDebuggers;
+	if (const TSharedPtr<FRemoteSessionsManager> RemoteSessionManager = FChaosVDEngineEditorBridge::Get().GetSessionsManager())
 	{
-		if (TSharedPtr<FChaosVDMultiSessionInfo> CustomSessionTarget = StaticCastSharedPtr<FChaosVDMultiSessionInfo>(RemoteSessionManager->GetSessionInfo(FChaosVDRemoteSessionsManager::CustomSessionsWrapperGUID).Pin()))
+		if (TSharedPtr<FMultiSessionInfo> CustomSessionTarget = StaticCastSharedPtr<FMultiSessionInfo>(RemoteSessionManager->GetSessionInfo(FRemoteSessionsManager::CustomSessionsWrapperGUID).Pin()))
 		{
 			if (CustomSessionTarget->InnerSessionsByInstanceID.Find(SessionGuid) != nullptr)
 			{
@@ -506,18 +520,19 @@ void FChaosVDRecordingControls::ToggleSessionSelectionInCustomTarget(FGuid Sessi
 				CustomSessionTarget->InnerSessionsByInstanceID.Add(SessionGuid, RemoteSessionManager->GetSessionInfo(SessionGuid));
 			}
 
-			SelectTargetSession(FChaosVDRemoteSessionsManager::CustomSessionsWrapperGUID);
+			SelectTargetSession(FRemoteSessionsManager::CustomSessionsWrapperGUID);
 		}
 	}
 }
 
 bool FChaosVDRecordingControls::CanSelectInCustomTarget(FGuid SessionGuid) const
 {
-	if (TSharedPtr<FChaosVDRemoteSessionsManager> RemoteSessionManager = FChaosVDEngineEditorBridge::Get().GetRemoteSessionsManager())
+	using namespace UE::TraceBasedDebuggers;
+	if (const TSharedPtr<FRemoteSessionsManager> RemoteSessionManager = FChaosVDEngineEditorBridge::Get().GetSessionsManager())
 	{
-		if (TSharedPtr<FChaosVDSessionInfo> CustomSessionTarget = RemoteSessionManager->GetSessionInfo(SessionGuid).Pin())
+		if (const TSharedPtr<FSessionInfo> CustomSessionTarget = RemoteSessionManager->GetSessionInfo(SessionGuid).Pin())
 		{
-			return CustomSessionTarget->ReadyState == EChaosVDRemoteSessionReadyState::Ready;
+			return CustomSessionTarget->ReadyState == ERemoteSessionReadyState::Ready;
 		}
 	}
 
@@ -526,9 +541,10 @@ bool FChaosVDRecordingControls::CanSelectInCustomTarget(FGuid SessionGuid) const
 
 bool FChaosVDRecordingControls::CanSelectMultiSessionTarget(FGuid SessionGuid) const
 {
-	if (TSharedPtr<FChaosVDRemoteSessionsManager> RemoteSessionManager = FChaosVDEngineEditorBridge::Get().GetRemoteSessionsManager())
+	using namespace UE::TraceBasedDebuggers;
+	if (const TSharedPtr<FRemoteSessionsManager> RemoteSessionManager = FChaosVDEngineEditorBridge::Get().GetSessionsManager())
 	{
-		if (TSharedPtr<FChaosVDSessionInfo> CustomSessionTarget = RemoteSessionManager->GetSessionInfo(SessionGuid).Pin())
+		if (const TSharedPtr<FSessionInfo> CustomSessionTarget = RemoteSessionManager->GetSessionInfo(SessionGuid).Pin())
 		{
 			return CanSelectMultiSessionTarget(CustomSessionTarget.ToSharedRef());
 		}
@@ -537,11 +553,12 @@ bool FChaosVDRecordingControls::CanSelectMultiSessionTarget(FGuid SessionGuid) c
 	return false;
 }
 
-bool FChaosVDRecordingControls::CanSelectMultiSessionTarget(const TSharedRef<FChaosVDSessionInfo>& SessionInfoRef) const
+bool FChaosVDRecordingControls::CanSelectMultiSessionTarget(const TSharedRef<UE::TraceBasedDebuggers::FSessionInfo>& SessionInfoRef) const
 {
-	if (EnumHasAnyFlags(SessionInfoRef->GetSessionTypeAttributes(), EChaosVDRemoteSessionAttributes::IsMultiSessionWrapper))
+	using namespace UE::TraceBasedDebuggers;
+	if (EnumHasAnyFlags(SessionInfoRef->GetSessionTypeAttributes(), ERemoteSessionAttributes::IsMultiSessionWrapper))
 	{
-		TSharedRef<FChaosVDMultiSessionInfo> AsMultiSessionInfo = StaticCastSharedRef<FChaosVDMultiSessionInfo>(SessionInfoRef);
+		const TSharedRef<FMultiSessionInfo> AsMultiSessionInfo = StaticCastSharedRef<FMultiSessionInfo>(SessionInfoRef);
 		return !AsMultiSessionInfo->InnerSessionsByInstanceID.IsEmpty();
 	}
 
@@ -561,7 +578,7 @@ void FChaosVDRecordingControls::ToggleAutoStartPlayback()
 
 bool FChaosVDRecordingControls::IsAutoStartPlaybackEnabled() const
 {
-	UChaosVDGeneralSettings* Settings = FChaosVDSettingsManager::Get().GetSettingsObject<UChaosVDGeneralSettings>();
+	const UChaosVDGeneralSettings* Settings = FChaosVDSettingsManager::Get().GetSettingsObject<UChaosVDGeneralSettings>();
 	return Settings ? Settings->bAutoStartLiveSessionsPlayback : false;
 }
 
@@ -578,49 +595,59 @@ void FChaosVDRecordingControls::ToggleSaveTracesToDisk()
 
 bool FChaosVDRecordingControls::IsSaveTracesToDiskEnabled()
 {
-	UChaosVDGeneralSettings* Settings = FChaosVDSettingsManager::Get().GetSettingsObject<UChaosVDGeneralSettings>();
+	const UChaosVDGeneralSettings* Settings = FChaosVDSettingsManager::Get().GetSettingsObject<UChaosVDGeneralSettings>();
 	return Settings ? Settings->bSaveRecordingsToDisk : false;
 }
 
 FSlateIcon FChaosVDRecordingControls::GetIconForSession(FGuid SessionId)
 {
-	if (TSharedPtr<FChaosVDSessionInfo> SessionInfo = GetSessionInfo(SessionId))
+	if (const TSharedPtr<UE::TraceBasedDebuggers::FSessionInfo> SessionInfo = GetSessionInfo(SessionId))
 	{
-		return SessionInfo->IsRecording() ? FSlateIcon(FChaosVDStyle::GetStyleSetName(),FName("RecordIcon"), FName("RecordIcon")) : FSlateIcon();
+		return SessionInfo->IsRecording(Chaos::VD::DebuggerGuid) ? FSlateIcon(FChaosVDStyle::GetStyleSetName(),FName("RecordIcon"), FName("RecordIcon")) : FSlateIcon();
 	}
 
 	return FSlateIcon();
 }
 
-TSharedPtr<FChaosVDSessionInfo> FChaosVDRecordingControls::GetCurrentSessionInfo() const
+TSharedPtr<UE::TraceBasedDebuggers::FSessionInfo> FChaosVDRecordingControls::GetCurrentSessionInfo() const
 {
 	return GetSessionInfo(CurrentSelectedSessionId);
 }
 
-TSharedPtr<FChaosVDSessionInfo> FChaosVDRecordingControls::GetSessionInfo(FGuid Id) const
+FChaosVDSessionData* FChaosVDRecordingControls::GetChaosDataForCurrentSession() const
 {
-	TSharedPtr<FChaosVDRemoteSessionsManager> RemoteSessionManager = FChaosVDEngineEditorBridge::Get().GetRemoteSessionsManager();
-	TSharedPtr<FChaosVDSessionInfo> SessionInfo = RemoteSessionManager ? RemoteSessionManager->GetSessionInfo(Id).Pin() : nullptr;
+	if (const TSharedPtr<UE::TraceBasedDebuggers::FSessionInfo> SessionInfo = GetCurrentSessionInfo())
+	{
+		return SessionInfo->GetDebuggerData<FChaosVDSessionData>();
+	}
+	return nullptr;
+}
+
+TSharedPtr<UE::TraceBasedDebuggers::FSessionInfo> FChaosVDRecordingControls::GetSessionInfo(FGuid Id) const
+{
+	using namespace UE::TraceBasedDebuggers;
+	const TSharedPtr<FRemoteSessionsManager> RemoteSessionManager = FChaosVDEngineEditorBridge::Get().GetSessionsManager();
+	TSharedPtr<FSessionInfo> SessionInfo = RemoteSessionManager ? RemoteSessionManager->GetSessionInfo(Id).Pin() : nullptr;
 
 	return SessionInfo;
 }
 
 bool FChaosVDRecordingControls::HasDataChannelsSupport() const
 {
-	if (TSharedPtr<FChaosVDSessionInfo> SessionInfo = GetCurrentSessionInfo())
+	if (const FChaosVDSessionData* SessionData = GetChaosDataForCurrentSession())
 	{
-		return !SessionInfo->DataChannelsStatesByName.IsEmpty();
+		return !SessionData->DataChannelsStatesByName.IsEmpty();
 	}
-
 	return false;
 }
 
 bool FChaosVDRecordingControls::CanChangeLoadingMode() const
 {
-	if (TSharedPtr<FChaosVDSessionInfo> CurrentSession = GetCurrentSessionInfo())
+	using namespace UE::TraceBasedDebuggers;
+	if (const TSharedPtr<FSessionInfo> CurrentSession = GetCurrentSessionInfo())
 	{
 		// In multi session mode targets, the loading mode is controlled automatically
-		if (EnumHasAnyFlags(CurrentSession->GetSessionTypeAttributes(), EChaosVDRemoteSessionAttributes::IsMultiSessionWrapper))
+		if (EnumHasAnyFlags(CurrentSession->GetSessionTypeAttributes(), ERemoteSessionAttributes::IsMultiSessionWrapper))
 		{
 			return false;
 		}
@@ -640,18 +667,18 @@ const FSlateBrush* FChaosVDRecordingControls::GetRecordOrStopButton() const
 	return bRecordingButtonHovered && IsRecording() ? FChaosVDStyle::Get().GetBrush("StopIcon") : RecordIconBrush;
 }
 
-void FChaosVDRecordingControls::HandleRecordingStop(TWeakPtr<FChaosVDSessionInfo> SessionInfo)
+void FChaosVDRecordingControls::HandleRecordingStop(const TSharedPtr<UE::TraceBasedDebuggers::FSessionInfo>&)
 {
 	if (UStatusBarSubsystem* StatusBarSubsystem = GEditor ? GEditor->GetEditorSubsystem<UStatusBarSubsystem>() : nullptr)
 	{
 		StatusBarSubsystem->PopStatusBarMessage(StatusBarID, RecordingMessageHandle);
 
-		const FText LiveSessionEnded = LOCTEXT("RecordingSessionEndedMessage"," Recording session has ended");
+		const FText LiveSessionEnded = LOCTEXT("RecordingSessionEndedMessage", "Recording session has ended");
 		LiveSessionEndedMessageHandle = StatusBarSubsystem->PushStatusBarMessage(StatusBarID, LiveSessionEnded);
 	}
 }
 
-void FChaosVDRecordingControls::HandleRecordingStart(TWeakPtr<FChaosVDSessionInfo> SessionInfo)
+void FChaosVDRecordingControls::HandleRecordingStart(const TSharedPtr<UE::TraceBasedDebuggers::FSessionInfo>& InSessionInfo)
 {
 	UStatusBarSubsystem* StatusBarSubsystem = GEditor ? GEditor->GetEditorSubsystem<UStatusBarSubsystem>() : nullptr;
 	if (!StatusBarSubsystem)
@@ -712,11 +739,12 @@ bool FChaosVDRecordingControls::FAsyncConnectionAttemptTask::CanExecute() const
 	return State == EState::NotStarted || State == EState::InProgress;
 }
 
-void FChaosVDRecordingControls::CancelInFlightConnectionAttempt(FGuid SessionID)
+void FChaosVDRecordingControls::CancelInFlightConnectionAttempt(const TSharedPtr<UE::TraceBasedDebuggers::FSessionInfo>& SessionInfo)
 {
-	if (TSharedPtr<FAsyncConnectionAttemptTask>* ExistingRequestInFlightPtrPtr = InFlightAsyncConnectionRequests.Find(SessionID))
+	const FGuid SessionID = SessionInfo ? SessionInfo->InstanceId : FGuid{};
+	if (const TSharedPtr<FAsyncConnectionAttemptTask>* ExistingRequestInFlightPtrPtr = InFlightAsyncConnectionRequests.Find(SessionID))
 	{
-		if (TSharedPtr<FAsyncConnectionAttemptTask> ExistingRequestInFlightPtr = *ExistingRequestInFlightPtrPtr)
+		if (const TSharedPtr<FAsyncConnectionAttemptTask> ExistingRequestInFlightPtr = *ExistingRequestInFlightPtrPtr)
 		{
 			ExistingRequestInFlightPtr->Cancel();
 		}
@@ -733,20 +761,21 @@ void FChaosVDRecordingControls::CancelInFlightConnectionAttempt(FGuid SessionID)
 
 void FChaosVDRecordingControls::HandleRemoteSessionsUpdated()
 {
-	TSharedPtr<FChaosVDSessionInfo> CurrentSessionPtr = GetCurrentSessionInfo();
+	using namespace UE::TraceBasedDebuggers;
+	const TSharedPtr<FSessionInfo> CurrentSessionPtr = GetCurrentSessionInfo();
 	if (!CurrentSessionPtr)
 	{
 		// if the session that is currently selected is not valid, try to find a valid session and select that as default target session
-		if (TSharedPtr<FChaosVDRemoteSessionsManager> RemoteSessionManager = FChaosVDEngineEditorBridge::Get().GetRemoteSessionsManager())
+		if (const TSharedPtr<FRemoteSessionsManager> RemoteSessionManager = FChaosVDEngineEditorBridge::Get().GetSessionsManager())
 		{
-			if (TSharedPtr<FChaosVDSessionInfo> FirstValidSession = RemoteSessionManager->GetFirstAvailableSingleSessionInfo().Pin())
+			if (const TSharedPtr<FSessionInfo> FirstValidSession = RemoteSessionManager->GetFirstAvailableSingleSessionInfo().Pin())
 			{
 				SelectTargetSession(FirstValidSession->InstanceId);
 				return;
 			}
 		}
 
-		SelectTargetSession(FChaosVDRemoteSessionsManager::InvalidSessionGUID);
+		SelectTargetSession(FRemoteSessionsManager::InvalidSessionGUID);
 	}
 }
 
@@ -889,18 +918,19 @@ void FChaosVDRecordingControls::FAsyncConnectionAttemptTask::Execute()
 	}
 }
 
-FChaosVDRecordingControls::FTraceStoreConnectionSettings FChaosVDRecordingControls::GetTraceStoreConnectionSettings(const TSharedRef<FChaosVDSessionInfo>& InSessionInfoRef)
+FChaosVDRecordingControls::FTraceStoreConnectionSettings FChaosVDRecordingControls::GetTraceStoreConnectionSettings(const TSharedRef<UE::TraceBasedDebuggers::FSessionInfo>& InSessionInfoRef)
 {
+	using namespace UE::TraceBasedDebuggers;
 	FTraceStoreConnectionSettings ConnectionSettings;
 
-	if (!ensure(!EnumHasAnyFlags(InSessionInfoRef->GetSessionTypeAttributes(), EChaosVDRemoteSessionAttributes::IsMultiSessionWrapper)))
+	if (!ensure(!EnumHasAnyFlags(InSessionInfoRef->GetSessionTypeAttributes(), ERemoteSessionAttributes::IsMultiSessionWrapper)))
 	{
 		UE_LOG(LogChaosVDEditor, Error, TEXT("[%s] Cannot not be called with a multi session wrapper."), ANSI_TO_TCHAR(__FUNCTION__));
 
 		return ConnectionSettings;
 	}
 
-	const FChaosVDTraceDetails& RecordingSessionDetails = InSessionInfoRef->GetConnectionDetails();
+	const FTraceConnectionDetails& RecordingSessionDetails = InSessionInfoRef->GetTraceConnectionDetails();
 
 	if (const UE::Trace::FStoreClient::FSessionInfo* TraceSessionInfo = FChaosVDTraceManager::GetTraceSessionInfo(RecordingSessionDetails.TraceTarget, RecordingSessionDetails.TraceGuid))
 	{
@@ -912,7 +942,7 @@ FChaosVDRecordingControls::FTraceStoreConnectionSettings FChaosVDRecordingContro
 	return ConnectionSettings;
 }
 
-void FChaosVDRecordingControls::ToggleMultiSessionSessionRecordingState(const TSharedRef<FChaosVDMultiSessionInfo>& InSessionInfoRef)
+void FChaosVDRecordingControls::ToggleMultiSessionSessionRecordingState(const TSharedRef<UE::TraceBasedDebuggers::FMultiSessionInfo>& InSessionInfoRef)
 {
 	const TSharedPtr<SChaosVDMainTab> MainTabSharedPtr = MainTabWeakPtr.Pin();
 	if (!ensure(MainTabSharedPtr))
@@ -928,21 +958,22 @@ void FChaosVDRecordingControls::ToggleMultiSessionSessionRecordingState(const TS
 		MainTabSharedPtr->GetChaosVDEngineInstance()->CloseActiveTraceSessions();
 	}
 
-	InSessionInfoRef->EnumerateInnerSessions([this, bNewRecordingState](const TSharedRef<FChaosVDSessionInfo>& InInnerSessionRef)
+	InSessionInfoRef->EnumerateInnerSessions([this, bNewRecordingState](const TSharedRef<UE::TraceBasedDebuggers::FSessionInfo>& InInnerSessionRef)
 	{
 		SetSessionRecordingState(bNewRecordingState, InInnerSessionRef);
 		return true;
 	});
 }
 
-void FChaosVDRecordingControls::ToggleSingleSessionRecordingState(const TSharedRef<FChaosVDSessionInfo>& SessionInfoRef)
+void FChaosVDRecordingControls::ToggleSingleSessionRecordingState(const TSharedRef<UE::TraceBasedDebuggers::FSessionInfo>& SessionInfoRef)
 {
-	SetSessionRecordingState(!SessionInfoRef->IsRecording(), SessionInfoRef);
+	SetSessionRecordingState(!SessionInfoRef->IsRecording(Chaos::VD::DebuggerGuid), SessionInfoRef);
 }
 
-void FChaosVDRecordingControls::SetSessionRecordingState(bool bIsRecording, const TSharedRef<FChaosVDSessionInfo>& SessionInfoRef)
+void FChaosVDRecordingControls::SetSessionRecordingState(bool bIsRecording, const TSharedRef<UE::TraceBasedDebuggers::FSessionInfo>& SessionInfoRef)
 {
-	TSharedPtr<FChaosVDRemoteSessionsManager> RemoteSessionManager = FChaosVDEngineEditorBridge::Get().GetRemoteSessionsManager();
+	using namespace UE::TraceBasedDebuggers;
+	const TSharedPtr<FRemoteSessionsManager> RemoteSessionManager = FChaosVDEngineEditorBridge::Get().GetSessionsManager();
 	if (!RemoteSessionManager)
 	{
 		UE_LOG(LogChaosVDEditor, Error, TEXT("[%hs] Session Manager is not available"), __func__);
@@ -962,12 +993,12 @@ void FChaosVDRecordingControls::SetSessionRecordingState(bool bIsRecording, cons
 		TArray<FString, TInlineAllocator<1>> RecordingArgs;
 
 		int32 RemainingRetries = 4;
-		EChaosVDTransportMode DataTransportOverrideMode = EChaosVDTransportMode::Invalid;
+		ETraceTransportMode DataTransportOverrideMode = ETraceTransportMode::Invalid;
 
-		if (UChaosVDGeneralSettings* Settings = FChaosVDSettingsManager::Get().GetSettingsObject<UChaosVDGeneralSettings>())
+		if (const UChaosVDGeneralSettings* Settings = FChaosVDSettingsManager::Get().GetSettingsObject<UChaosVDGeneralSettings>())
 		{
 			RemainingRetries = Settings->MaxConnectionRetries;
-			DataTransportOverrideMode = Settings->GetTransportModeForTarget(SessionInfoRef->BuildTargetType);
+			DataTransportOverrideMode = Settings->GetTransportModeForTargetType(SessionInfoRef->BuildTargetType);
 		}
 		else
 		{
@@ -976,29 +1007,29 @@ void FChaosVDRecordingControls::SetSessionRecordingState(bool bIsRecording, cons
 
 		// Set the session in a busy state so the auto-expire policy for it is relaxed. Starting a recording might cause the remote instance to stall for a few seconds
 		// if the map loaded is too complex
-		SessionInfoRef->ReadyState = EChaosVDRemoteSessionReadyState::Busy;
+		SessionInfoRef->ReadyState = ERemoteSessionReadyState::Busy;
 
 		auto RecordingAttemptFailedCallback = [WeakMainTabPtr = StaticCastWeakPtr<SChaosVDMainTab>(MainTabSharedPtr->AsWeak()), SessionGUID = SessionInfoRef->InstanceId]()
 		{
-			TSharedPtr<FChaosVDRemoteSessionsManager> RemoteSessionManagerPtr = FChaosVDEngineEditorBridge::Get().GetRemoteSessionsManager();
-			TSharedPtr<FChaosVDSessionInfo> SessionInfoPtr = RemoteSessionManagerPtr ? RemoteSessionManagerPtr->GetSessionInfo(SessionGUID).Pin() : nullptr;
+			const TSharedPtr<FRemoteSessionsManager> RemoteSessionManagerPtr = FChaosVDEngineEditorBridge::Get().GetSessionsManager();
+			const TSharedPtr<FSessionInfo> SessionInfoPtr = RemoteSessionManagerPtr ? RemoteSessionManagerPtr->GetSessionInfo(SessionGUID).Pin() : nullptr;
 
 			if (!SessionInfoPtr)
 			{
 				return;
 			}
 
-			SessionInfoPtr->ReadyState = EChaosVDRemoteSessionReadyState::Ready;
+			SessionInfoPtr->ReadyState = ERemoteSessionReadyState::Ready;
 
-			if (TSharedPtr<SChaosVDMainTab> MainTab = WeakMainTabPtr.Pin())
+			if (const TSharedPtr<SChaosVDMainTab> MainTab = WeakMainTabPtr.Pin())
 			{
 				MainTab->GetChaosVDEngineInstance()->CloseSessionByRemoteSessionID(SessionGUID);
 			}
 		};
 
 		FChaosVDStartRecordingCommandMessage RecordingParams;
-		RecordingParams.RecordingMode = EChaosVDRecordingMode::Live;
-		RecordingParams.TransportMode = DataTransportOverrideMode != EChaosVDTransportMode::Invalid ? DataTransportOverrideMode : EChaosVDTransportMode::Direct;
+		RecordingParams.RecordingMode = ERecordingMode::Live;
+		RecordingParams.TransportMode = DataTransportOverrideMode != ETraceTransportMode::Invalid ? DataTransportOverrideMode : ETraceTransportMode::Direct;
 
 		check(GLog);
 		bool bOutCanBindAll = false;
@@ -1025,7 +1056,7 @@ void FChaosVDRecordingControls::SetSessionRecordingState(bool bIsRecording, cons
 		}
 #endif
 
-		if (RecordingParams.TransportMode == EChaosVDTransportMode::Direct)
+		if (RecordingParams.TransportMode == ETraceTransportMode::Direct)
 		{
 			// Direct Tracing mode requires opening a trace session in the editor side first before starting the recording
 			// so targets have something to connect to (Usually targets connect to a running trace server)
@@ -1041,23 +1072,23 @@ void FChaosVDRecordingControls::SetSessionRecordingState(bool bIsRecording, cons
 
 			RecordingParams.Target.Appendf(TEXT(":%u"), SessionDesc.SessionPort);
 
-			RemoteSessionManager->SendStartRecordingCommand(SessionInfoRef->Address, RecordingParams);
+			RemoteSessionManager->SendCommand(SessionInfoRef->Address, RecordingParams);
 
 			ExecuteAsyncConnectionAttemptTaskWithRetry(SessionInfoRef->InstanceId, RemainingRetries, [WeakThis = AsWeak(), SessionInstanceId = SessionInfoRef->InstanceId]()
 			{
-				TSharedPtr<FChaosVDRemoteSessionsManager> RemoteSessionManagerPtr = FChaosVDEngineEditorBridge::Get().GetRemoteSessionsManager();
-				TSharedPtr<FChaosVDSessionInfo> SessionInfoPtr = RemoteSessionManagerPtr ? RemoteSessionManagerPtr->GetSessionInfo(SessionInstanceId).Pin() : nullptr;
+				const TSharedPtr<FRemoteSessionsManager> RemoteSessionManagerPtr = FChaosVDEngineEditorBridge::Get().GetSessionsManager();
+				const TSharedPtr<FSessionInfo> SessionInfoPtr = RemoteSessionManagerPtr ? RemoteSessionManagerPtr->GetSessionInfo(SessionInstanceId).Pin() : nullptr;
 				if (!SessionInfoPtr)
 				{
 					return false;
 				}
 
 				// In Direct trace mode we if we get connection details from the remote instance, we can assume the connection succeeded.
-				return SessionInfoPtr->GetConnectionDetails().IsValid();
-	
+				return SessionInfoPtr->GetTraceConnectionDetails().IsValid();
+
 			}, RecordingAttemptFailedCallback);
 		}
-		else if (RecordingParams.TransportMode == EChaosVDTransportMode::Relay)
+		else if (RecordingParams.TransportMode == ETraceTransportMode::Relay)
 		{
 			// Relay Tracing mode requires opening a trace session in the editor side first before starting the recording
 			// so targets have something to connect to (Usually targets connect to a running trace server)
@@ -1068,21 +1099,21 @@ void FChaosVDRecordingControls::SetSessionRecordingState(bool bIsRecording, cons
 				return;
 			}
 
-			RemoteSessionManager->SendStartRecordingCommand(SessionInfoRef->Address, RecordingParams);
+			RemoteSessionManager->SendCommand(SessionInfoRef->Address, RecordingParams);
 
 #if WITH_CHAOS_VISUAL_DEBUGGER
 			// Once the start recording command is issued, we can try to connect to the created session (if everything went well).
 			// If it didn't go well, this will take care of update the UI to notify the user
 			ExecuteAsyncConnectionAttemptTaskWithRetry(SessionInfoRef->InstanceId, RemainingRetries, [WeakThis = AsWeak(), SessionInstanceId = SessionInfoRef->InstanceId]()
 			{
-				TSharedPtr<FChaosVDRemoteSessionsManager> RemoteSessionManagerPtr = FChaosVDEngineEditorBridge::Get().GetRemoteSessionsManager();
-				TSharedPtr<FChaosVDSessionInfo> SessionInfoPtr = RemoteSessionManagerPtr ? RemoteSessionManagerPtr->GetSessionInfo(SessionInstanceId).Pin() : nullptr;
+				const TSharedPtr<FRemoteSessionsManager> RemoteSessionManagerPtr = FChaosVDEngineEditorBridge::Get().GetSessionsManager();
+				const TSharedPtr<FSessionInfo> SessionInfoPtr = RemoteSessionManagerPtr ? RemoteSessionManagerPtr->GetSessionInfo(SessionInstanceId).Pin() : nullptr;
 				if (!SessionInfoPtr)
 				{
 					return false;
 				}
 
-				const FChaosVDTraceDetails& RecordingSessionDetails = SessionInfoPtr->GetConnectionDetails();
+				const FTraceConnectionDetails& RecordingSessionDetails = SessionInfoPtr->GetTraceConnectionDetails();
 
 				// We didn't receive the connection details yet. Trying again later...
 				if (!RecordingSessionDetails.IsValid())
@@ -1091,28 +1122,28 @@ void FChaosVDRecordingControls::SetSessionRecordingState(bool bIsRecording, cons
 				}
 
 				// Check if there is a connection attempt in progress and its results
-				const Chaos::VD::EConnectionAttemptResult ConnectionAttemptStatus = FChaosVDEngineEditorBridge::Get().GetRelayTransportInstance()->GetConnectionAttemptResult(SessionInstanceId);
+				const EConnectionAttemptResult ConnectionAttemptStatus = FChaosVDEngineEditorBridge::Get().GetTraceRelayTransportInstance()->GetConnectionAttemptResult(SessionInstanceId);
 				switch (ConnectionAttemptStatus)
 				{
-				case Chaos::VD::InProgress:
+				case EConnectionAttemptResult::InProgress:
 					{
 						// A connection attempt started, but we don't have the results yet. Trying again later...
 						return false;
 					}
-				case Chaos::VD::NotStarted:
+				case EConnectionAttemptResult::NotStarted:
 					{
-						Chaos::VD::FRelayConnectionInfo ConnectionInfo;
+						FRelayConnectionInfo ConnectionInfo;
 						ConnectionInfo.Port = RecordingSessionDetails.Port;
 						ConnectionInfo.Address = RecordingSessionDetails.TraceTarget;
 						ConnectionInfo.CertificateAuthority = RecordingSessionDetails.CertAuth;
-						Chaos::VD::EConnectionAttemptResult Result = FChaosVDEngineEditorBridge::Get().GetRelayTransportInstance()->ConnectToRelay(SessionInstanceId, ConnectionInfo);
+						EConnectionAttemptResult Result = FChaosVDEngineEditorBridge::Get().GetTraceRelayTransportInstance()->ConnectToRelay(SessionInstanceId, ConnectionInfo);
 
-						return Result == Chaos::VD::EConnectionAttemptResult::Success;
+						return Result == EConnectionAttemptResult::Success;
 					}
-				case Chaos::VD::Success:
+				case EConnectionAttemptResult::Success:
 					// The current connection attempt succeeded, we can stop the async retry loop
 					return true;
-				case Chaos::VD::Failed:
+				case EConnectionAttemptResult::Failed:
 				default:
 					// There was a connection attempt, but it failed. Abandoning the retry attempt as they will not succeed
 					return false;
@@ -1124,22 +1155,23 @@ void FChaosVDRecordingControls::SetSessionRecordingState(bool bIsRecording, cons
 	}
 	else
 	{
-		CancelInFlightConnectionAttempt(SessionInfoRef->InstanceId);
-		RemoteSessionManager->SendStopRecordingCommand(SessionInfoRef->Address);
+		CancelInFlightConnectionAttempt(SessionInfoRef.ToSharedPtr());
+		RemoteSessionManager->SendCommand(SessionInfoRef->Address, FChaosVDStopRecordingCommandMessage{});
 	}
 }
 
 FReply FChaosVDRecordingControls::ToggleRecordingState()
 {
-	TSharedPtr<FChaosVDSessionInfo> SessionInfoPtr = GetCurrentSessionInfo();
+	using namespace UE::TraceBasedDebuggers;
+	TSharedPtr<FSessionInfo> SessionInfoPtr = GetCurrentSessionInfo();
 	if (!SessionInfoPtr)
 	{
 		return FReply::Handled();
 	}
 
-	if (EnumHasAnyFlags(SessionInfoPtr->GetSessionTypeAttributes(), EChaosVDRemoteSessionAttributes::IsMultiSessionWrapper))
+	if (EnumHasAnyFlags(SessionInfoPtr->GetSessionTypeAttributes(), ERemoteSessionAttributes::IsMultiSessionWrapper))
 	{
-		ToggleMultiSessionSessionRecordingState(StaticCastSharedPtr<FChaosVDMultiSessionInfo>(SessionInfoPtr).ToSharedRef());
+		ToggleMultiSessionSessionRecordingState(StaticCastSharedPtr<FMultiSessionInfo>(SessionInfoPtr).ToSharedRef());
 	}
 	else
 	{
@@ -1151,26 +1183,35 @@ FReply FChaosVDRecordingControls::ToggleRecordingState()
 
 bool FChaosVDRecordingControls::IsRecordingToggleButtonEnabled() const
 {
-	if (CurrentSelectedSessionId == FChaosVDRemoteSessionsManager::InvalidSessionGUID)
+	using namespace UE::TraceBasedDebuggers;
+	if (CurrentSelectedSessionId == FRemoteSessionsManager::InvalidSessionGUID)
 	{
 		return false;
 	}
 
-	TSharedPtr<FChaosVDSessionInfo> SessionInfo = GetCurrentSessionInfo();
+	TSharedPtr<FSessionInfo> SessionInfo = GetCurrentSessionInfo();
 	if (!SessionInfo)
 	{
 		return false;
 	}
 
+	// Do not allow recording if another trace-based debuggers is already recording in the current session
+	if (SessionInfo->LastKnownRecordingState.DebuggerId.IsValid()
+		&& SessionInfo->LastKnownRecordingState.DebuggerId != Chaos::VD::DebuggerGuid
+		&& SessionInfo->LastKnownRecordingState.bIsRecording)
+	{
+		return false;
+	}
+
 	if (!IsRecording())
 	{
 		return true;
 	}
 
 	// If we are recording, don't show the stop button for the mode that is disabled
-	if (EnumHasAnyFlags(SessionInfo->GetSessionTypeAttributes(), EChaosVDRemoteSessionAttributes::IsMultiSessionWrapper))
+	if (EnumHasAnyFlags(SessionInfo->GetSessionTypeAttributes(), ERemoteSessionAttributes::IsMultiSessionWrapper))
 	{
-		TSharedPtr<FChaosVDMultiSessionInfo> AsMultiSessionInfo = StaticCastSharedPtr<FChaosVDMultiSessionInfo>(SessionInfo);
+		TSharedPtr<FMultiSessionInfo> AsMultiSessionInfo = StaticCastSharedPtr<FMultiSessionInfo>(SessionInfo);
 
 		return CanSelectMultiSessionTarget(SessionInfo.ToSharedRef());
 	}
@@ -1260,13 +1301,15 @@ void FChaosVDRecordingControls::AddToMenu(FName ExistingMenu)
 
 bool FChaosVDRecordingControls::IsRecording() const
 {
-	TSharedPtr<FChaosVDSessionInfo> SessionInfo = GetCurrentSessionInfo();
-	return SessionInfo ? SessionInfo->IsRecording() : false;
+	using namespace UE::TraceBasedDebuggers;
+	const TSharedPtr<FSessionInfo> SessionInfo = GetCurrentSessionInfo();
+	return SessionInfo ? SessionInfo->IsRecording(Chaos::VD::DebuggerGuid) : false;
 }
 
 FText FChaosVDRecordingControls::GetRecordingTimeText() const
 {
-	if (TSharedPtr<FChaosVDSessionInfo> SessionInfo = GetCurrentSessionInfo())
+	using namespace UE::TraceBasedDebuggers;
+	if (const TSharedPtr<FSessionInfo> SessionInfo = GetCurrentSessionInfo())
 	{
 		FNumberFormattingOptions FormatOptions;
 		FormatOptions.MinimumFractionalDigits = 2;
diff --git a/Engine/Plugins/ChaosVD/Source/ChaosVD/Private/Widgets/ChaosVDRecordingControls.h b/Engine/Plugins/ChaosVD/Source/ChaosVD/Private/Widgets/ChaosVDRecordingControls.h
--- a/Engine/Plugins/ChaosVD/Source/ChaosVD/Private/Widgets/ChaosVDRecordingControls.h
+++ b/Engine/Plugins/ChaosVD/Source/ChaosVD/Private/Widgets/ChaosVDRecordingControls.h
@@ -3,13 +3,14 @@
 #pragma once
 
 #include "StatusBarSubsystem.h"
-#include "Chaos/ChaosVDRemoteSessionsManager.h"
 #include "ChaosVDEngine.h"
+#include "RemoteSessionsManager.h"
 #include "Templates/SharedPointer.h"
 #include "Widgets/SCompoundWidget.h"
 
 #include "ChaosVDRecordingControls.generated.h"
 
+struct FChaosVDSessionData;
 class FMenuBuilder;
 
 namespace Chaos::VisualDebugger
@@ -72,17 +73,20 @@ protected:
 	bool CanSelectInCustomTarget(FGuid SessionGuid) const;
 
 	bool CanSelectMultiSessionTarget(FGuid SessionGuid) const;
-	bool CanSelectMultiSessionTarget(const TSharedRef<FChaosVDSessionInfo>& SessionInfoRef) const;
+	bool CanSelectMultiSessionTarget(const TSharedRef<UE::TraceBasedDebuggers::FSessionInfo>& SessionInfoRef) const;
 
 	void ToggleAutoStartPlayback();
 	bool IsAutoStartPlaybackEnabled() const;
 	void ToggleSaveTracesToDisk();
-	bool IsSaveTracesToDiskEnabled();	
+	bool IsSaveTracesToDiskEnabled();
 
 	FSlateIcon GetIconForSession(FGuid SessionId);
 
-	TSharedPtr<FChaosVDSessionInfo> GetCurrentSessionInfo() const;
-	TSharedPtr<FChaosVDSessionInfo> GetSessionInfo(FGuid Id) const;
+	TSharedPtr<UE::TraceBasedDebuggers::FSessionInfo> GetCurrentSessionInfo() const;
+	TSharedPtr<UE::TraceBasedDebuggers::FSessionInfo> GetSessionInfo(FGuid Id) const;
+
+	FChaosVDSessionData* GetChaosDataForCurrentSession() const;
+
 	void HandleRemoteSessionsUpdated();
 
 	bool HasDataChannelsSupport() const;
@@ -91,15 +95,15 @@ protected:
 	
 	const FSlateBrush* GetRecordOrStopButton() const;
 	
-	void HandleRecordingStop(TWeakPtr<FChaosVDSessionInfo> SessionInfo);
-	void HandleRecordingStart(TWeakPtr<FChaosVDSessionInfo> SessionInfo);
+	void HandleRecordingStop(const TSharedPtr<UE::TraceBasedDebuggers::FSessionInfo>& SessionInfo);
+	void HandleRecordingStart(const TSharedPtr<UE::TraceBasedDebuggers::FSessionInfo>& SessionInfo);
 
 	void ExecuteAsyncConnectionAttemptTaskWithRetry(FGuid SessionID, int32 RemainingRetries, const TFunction<bool()>& InRecordingStartAttemptCallback, const TFunction<void()>& InRecordingFailedCallback);
 
 	FReply ToggleRecordingState();
-	void ToggleMultiSessionSessionRecordingState(const TSharedRef<FChaosVDMultiSessionInfo>& InSessionInfoRef);
-	void ToggleSingleSessionRecordingState(const TSharedRef<FChaosVDSessionInfo>& SessionInfoRef);
-	void SetSessionRecordingState(bool bIsRecording, const TSharedRef<FChaosVDSessionInfo>& SessionInfoRef);
+	void ToggleMultiSessionSessionRecordingState(const TSharedRef<UE::TraceBasedDebuggers::FMultiSessionInfo>& InSessionInfoRef);
+	void ToggleSingleSessionRecordingState(const TSharedRef<UE::TraceBasedDebuggers::FSessionInfo>& SessionInfoRef);
+	void SetSessionRecordingState(bool bIsRecording, const TSharedRef<UE::TraceBasedDebuggers::FSessionInfo>& SessionInfoRef);
 
 	bool IsRecordingToggleButtonEnabled() const;
 	EVisibility IsRecordingToggleButtonVisible() const;
@@ -118,7 +122,7 @@ protected:
 		FString SessionAddress;
 	};
 
-	FTraceStoreConnectionSettings GetTraceStoreConnectionSettings(const TSharedRef<FChaosVDSessionInfo>& InSessionInfoRef);
+	FTraceStoreConnectionSettings GetTraceStoreConnectionSettings(const TSharedRef<UE::TraceBasedDebuggers::FSessionInfo>& InSessionInfoRef);
 
 	FText GetCurrentSelectedSessionName() const;
 
@@ -142,7 +146,7 @@ protected:
 
 	EChaosVDLoadRecordedDataMode CurrentLoadingMode = EChaosVDLoadRecordedDataMode::SingleSource;
 
-	void CancelInFlightConnectionAttempt(FGuid SessionID);
+	void CancelInFlightConnectionAttempt(const TSharedPtr<UE::TraceBasedDebuggers::FSessionInfo>& SessionInfo);
 
 	enum class ERecordingStartRequestState
 	{
diff --git a/Engine/Plugins/ChaosVD/Source/ChaosVD/Private/Widgets/SChaosVDMainTab.cpp b/Engine/Plugins/ChaosVD/Source/ChaosVD/Private/Widgets/SChaosVDMainTab.cpp
--- a/Engine/Plugins/ChaosVD/Source/ChaosVD/Private/Widgets/SChaosVDMainTab.cpp
+++ b/Engine/Plugins/ChaosVD/Source/ChaosVD/Private/Widgets/SChaosVDMainTab.cpp
@@ -218,7 +218,7 @@ void SChaosVDMainTab::Construct(const FArguments& InArgs, TSharedPtr<FChaosVDEng
 
 	SetUpDisableCPUThrottlingDelegate();
 
-	if (TSharedPtr<FChaosVDRemoteSessionsManager> RemoteSessionManager = FChaosVDEngineEditorBridge::Get().GetRemoteSessionsManager())
+	if (TSharedPtr<UE::TraceBasedDebuggers::FRemoteSessionsManager> RemoteSessionManager = FChaosVDEngineEditorBridge::Get().GetSessionsManager())
 	{
 		RemoteSessionManager->StartSessionDiscovery();
 	}
@@ -226,7 +226,7 @@ void SChaosVDMainTab::Construct(const FArguments& InArgs, TSharedPtr<FChaosVDEng
 
 SChaosVDMainTab::~SChaosVDMainTab()
 {
-	if (TSharedPtr<FChaosVDRemoteSessionsManager> RemoteSessionManager = FChaosVDEngineEditorBridge::Get().GetRemoteSessionsManager())
+	if (TSharedPtr<UE::TraceBasedDebuggers::FRemoteSessionsManager> RemoteSessionManager = FChaosVDEngineEditorBridge::Get().GetSessionsManager())
 	{
 		RemoteSessionManager->StopSessionDiscovery();
 	}
diff --git a/Engine/Plugins/ChaosVD/Source/ChaosVD/Private/Widgets/SConnectionStatusBar.cpp b/Engine/Plugins/ChaosVD/Source/ChaosVD/Private/Widgets/SConnectionStatusBar.cpp
--- a/Engine/Plugins/ChaosVD/Source/ChaosVD/Private/Widgets/SConnectionStatusBar.cpp
+++ b/Engine/Plugins/ChaosVD/Source/ChaosVD/Private/Widgets/SConnectionStatusBar.cpp
@@ -3,9 +3,11 @@
 #include "SConnectionStatusBar.h"
 
 #include "Chaos/ChaosVDEngineEditorBridge.h"
+#include "ChaosVDRecordingDetails.h"
 #include "Components/HorizontalBox.h"
-#include "Widgets/SOverlay.h"
+#include "RemoteSessionsManager.h"
 #include "Widgets/Images/SImage.h"
+#include "Widgets/SOverlay.h"
 #include "Widgets/Text/STextBlock.h"
 
 #define LOCTEXT_NAMESPACE "ChaosVisualDebugger"
@@ -93,16 +95,16 @@ FText SConnectionStatusBar::GetTitleText() const
 		return LOCTEXT("ConnectionReadyStatusText", "Ready");
 	case EState::ReceivingData:
 		{
-			FText BytesPerSecondText =  FText::AsMemory(CurrentReceivedDataBytesPerSecond, &SizeFormattingOptions,nullptr, IEC);
-			FText PendingBytesText =  FText::AsMemory(CurrentPendingBytes, &SizeFormattingOptions,nullptr, IEC);
-			FText TotalBytesText =  FText::AsMemory(CurrentTotalReceivedData, &SizeFormattingOptions,nullptr, IEC);
+			FText BytesPerSecondText = FText::AsMemory(CurrentReceivedDataBytesPerSecond, &SizeFormattingOptions, nullptr, IEC);
+			FText PendingBytesText = FText::AsMemory(CurrentPendingBytes, &SizeFormattingOptions, nullptr, IEC);
+			FText TotalBytesText = FText::AsMemory(CurrentTotalReceivedData, &SizeFormattingOptions, nullptr, IEC);
 			return FText::FormatOrdered(LOCTEXT("ConnectionDownloadedStatusText", "Receiving {0}/s | {1} Pending | Total {2}"), BytesPerSecondText, PendingBytesText, TotalBytesText);
 		}
 	case EState::WaitingForData:
 	default:
 		{
-			FText PendingBytesText =  FText::AsMemory(CurrentPendingBytes, &SizeFormattingOptions,nullptr, IEC);
-			FText TotalBytesText =  FText::AsMemory(CurrentTotalReceivedData, &SizeFormattingOptions,nullptr, IEC);
+			FText PendingBytesText = FText::AsMemory(CurrentPendingBytes, &SizeFormattingOptions, nullptr, IEC);
+			FText TotalBytesText = FText::AsMemory(CurrentTotalReceivedData, &SizeFormattingOptions, nullptr, IEC);
 			return FText::FormatOrdered(LOCTEXT("ConnectionWaitingStatusText", "Waiting for Data | {0} Pending | Total {1}"), PendingBytesText, TotalBytesText);
 		}
 	}
@@ -145,6 +147,8 @@ void SConnectionStatusBar::Tick(const FGeometry& AllottedGeometry, const double
 
 void SConnectionStatusBar::UpdateConnectionStatus(double DeltaTime)
 {
+	using namespace UE::TraceBasedDebuggers;
+
 	ElapsedTimeSinceLastUpdate += DeltaTime;
 
 	constexpr double UpdateIntervalSeconds = 1.0;
@@ -155,7 +159,7 @@ void SConnectionStatusBar::UpdateConnectionStatus(double DeltaTime)
 
 	ElapsedTimeSinceLastUpdate = 0;
 
-	TSharedPtr<FChaosVDRemoteSessionsManager> RemoteSessionManager = FChaosVDEngineEditorBridge::Get().GetRemoteSessionsManager();
+	TSharedPtr<FRemoteSessionsManager> RemoteSessionManager = FChaosVDEngineEditorBridge::Get().GetSessionsManager();
 	if (!RemoteSessionManager)
 	{
 		State = EState::Ready;
@@ -165,16 +169,16 @@ void SConnectionStatusBar::UpdateConnectionStatus(double DeltaTime)
 	bool bIsAnyRecording = false;
 	int64 PendingBytesNum = 0;
 	uint64 BytesReadSinceLastMeasurement = 0;
-	RemoteSessionManager->EnumerateActiveSessions([this, &BytesReadSinceLastMeasurement, &bIsAnyRecording, &PendingBytesNum](const TSharedRef<FChaosVDSessionInfo>& InSessionInfoRef)
+	RemoteSessionManager->EnumerateActiveSessions([this, &BytesReadSinceLastMeasurement, &bIsAnyRecording, &PendingBytesNum](const TSharedRef<FSessionInfo>& InSessionInfoRef)
 	{
-		if (!EnumHasAnyFlags(InSessionInfoRef->GetSessionTypeAttributes(), EChaosVDRemoteSessionAttributes::IsMultiSessionWrapper))
+		if (!EnumHasAnyFlags(InSessionInfoRef->GetSessionTypeAttributes(), ERemoteSessionAttributes::IsMultiSessionWrapper))
 		{
-			if (InSessionInfoRef->IsRecording())
+			if (InSessionInfoRef->IsRecording(DebuggerGuid))
 			{
 				bIsAnyRecording = true;
 				PendingBytesNum += InSessionInfoRef->LastKnownRecordingState.BufferedDataBytesSize;
 				BytesReadSinceLastMeasurement += InSessionInfoRef->GetReceivedBytesPerSecond();
-			}			
+			}
 		}
 
 		return true;
diff --git a/Engine/Plugins/ChaosVD/Source/ChaosVD/Public/Settings/ChaosVDGeneralSettings.h b/Engine/Plugins/ChaosVD/Source/ChaosVD/Public/Settings/ChaosVDGeneralSettings.h
--- a/Engine/Plugins/ChaosVD/Source/ChaosVD/Public/Settings/ChaosVDGeneralSettings.h
+++ b/Engine/Plugins/ChaosVD/Source/ChaosVD/Public/Settings/ChaosVDGeneralSettings.h
@@ -19,14 +19,22 @@ public:
 	/** Returns the configured transport mode for the provided target type
 	 * @param TargetType Build target type to evaluate
 	 */
-	CHAOSVD_API EChaosVDTransportMode GetTransportModeForTarget(EBuildTargetType TargetType) const;
+	CHAOSVD_API UE::TraceBasedDebuggers::ETraceTransportMode GetTransportModeForTargetType(EBuildTargetType TargetType) const;
 
 	/** Sets a specific transport mode for the provided target type
 	 * @param TargetType Build target type to evaluate
 	 * @param NewTransportMode New transport mode to use for the provided target
 	 */
+	CHAOSVD_API void SetTransportModeForTargetType(EBuildTargetType TargetType, UE::TraceBasedDebuggers::ETraceTransportMode NewTransportMode);
+
+PRAGMA_DISABLE_DEPRECATION_WARNINGS
+	UE_DEPRECATED(5.8, "Use GetTransportModeForTargetType")
+	CHAOSVD_API EChaosVDTransportMode GetTransportModeForTarget(EBuildTargetType TargetType) const;
+
+	UE_DEPRECATED(5.8, "Use SetTransportModeForTargetType")
 	CHAOSVD_API void SetTransportModeForTarget(EBuildTargetType TargetType, EChaosVDTransportMode NewTransportMode);
-	
+PRAGMA_ENABLE_DEPRECATION_WARNINGS
+
 	/** If true, CVD will only load frames that have solver data in them - Only takes effect before loading a file */
 	UPROPERTY(Config, EditAnywhere, Category = "Session Data Loading")
 	bool bTrimEmptyFrames = true;
@@ -71,9 +79,9 @@ public:
 private:
 	/** [DEBUG Setting] Default Data Transport for Game Servers - Used to change the transport mode for the trace data. This is not saved by design */
 	UPROPERTY(EditAnywhere, Category = Connection)
-	EChaosVDTransportMode GameServerDataTransportMode = EChaosVDTransportMode::Direct;
+	UE::TraceBasedDebuggers::ETraceTransportMode GameServerDataTransportMode = UE::TraceBasedDebuggers::ETraceTransportMode::Direct;
 
 	/** [DEBUG Setting] Default Data Transport for Game Clients - Used to change the transport mode for the trace data. This is not saved by design */
 	UPROPERTY(EditAnywhere, Category = Connection)
-	EChaosVDTransportMode GameClientDataTransportMode = EChaosVDTransportMode::Direct;
+	UE::TraceBasedDebuggers::ETraceTransportMode GameClientDataTransportMode = UE::TraceBasedDebuggers::ETraceTransportMode::Direct;
 };
diff --git a/Engine/Source/Programs/UnrealBuildTool/Configuration/Rules/ModuleRules.cs b/Engine/Source/Programs/UnrealBuildTool/Configuration/Rules/ModuleRules.cs
--- a/Engine/Source/Programs/UnrealBuildTool/Configuration/Rules/ModuleRules.cs
+++ b/Engine/Source/Programs/UnrealBuildTool/Configuration/Rules/ModuleRules.cs
@@ -1680,6 +1680,48 @@ namespace UnrealBuildTool
 			}
 		}
 
+		/// <summary>
+		/// Returns true if any trace-based debugger is supported by the provided target
+		/// </summary>
+		protected bool AreTraceBasedDebuggersSupported(ReadOnlyTargetRules target)
+		{
+			return IsRewindDebuggerSupported(target) || IsChaosVisualDebuggerSupported(target);
+		}
+
+		/// <summary>
+		/// Setup this module for trace-based debuggers
+		/// </summary>
+		public void SetupModuleTraceBasedDebuggersSupport(ReadOnlyTargetRules target)
+		{
+			SetupModuleRewindDebuggerSupport(target);
+			SetupModuleChaosVisualDebuggerSupport(target);
+
+			if (AreTraceBasedDebuggersSupported(target))
+			{
+				PublicDefinitions.Add("WITH_TRACE_BASED_DEBUGGERS=1");
+			}
+			else
+			{
+				PublicDefinitions.Add("WITH_TRACE_BASED_DEBUGGERS=0");
+			}
+		}
+
+		/// <summary>
+		/// Returns true if the RewindDebugger is supported by the provided target
+		/// </summary>
+		protected bool IsRewindDebuggerSupported(ReadOnlyTargetRules target)
+		{
+			return target.bCompileRewindDebuggerSupport;
+		}
+
+		/// <summary>
+		/// Setup this module for Rewind Debugger support (Required for recording debug data that will be visualized in the Rewind Debugger tool)
+		/// </summary>
+		public void SetupModuleRewindDebuggerSupport(ReadOnlyTargetRules target)
+		{
+			//TODO:
+		}
+
 		/// <summary>
 		/// Setup this module for Iris support (based on the settings in UEBuildConfiguration)
 		/// </summary>
diff --git a/Engine/Source/Programs/UnrealBuildTool/Configuration/Rules/ReadOnlyTargetRules.cs b/Engine/Source/Programs/UnrealBuildTool/Configuration/Rules/ReadOnlyTargetRules.cs
--- a/Engine/Source/Programs/UnrealBuildTool/Configuration/Rules/ReadOnlyTargetRules.cs
+++ b/Engine/Source/Programs/UnrealBuildTool/Configuration/Rules/ReadOnlyTargetRules.cs
@@ -141,6 +141,8 @@ namespace UnrealBuildTool
 
 		public bool bUseChaosMemoryTracking => Inner.bUseChaosMemoryTracking;
 
+		public bool bCompileRewindDebuggerSupport => Inner.bCompileRewindDebuggerSupport;
+
 		public bool bCompileChaosVisualDebuggerSupport => Inner.bCompileChaosVisualDebuggerSupport;
 	
 		public bool bCompileAsioSSLSupport => Inner.bCompileAsioSSLSupport;
diff --git a/Engine/Source/Programs/UnrealBuildTool/Configuration/Rules/TargetRules.cs b/Engine/Source/Programs/UnrealBuildTool/Configuration/Rules/TargetRules.cs
--- a/Engine/Source/Programs/UnrealBuildTool/Configuration/Rules/TargetRules.cs
+++ b/Engine/Source/Programs/UnrealBuildTool/Configuration/Rules/TargetRules.cs
@@ -1023,6 +1023,28 @@ namespace UnrealBuildTool
 			set => bEnableChaosVisualDebuggerSupport = value;
 		}
 		private bool? bEnableChaosVisualDebuggerSupport;
+
+		/// <summary>
+		/// Whether to compile in RewindDebugger (RWD) support features to record traces for debug primitives and logs
+		/// </summary>
+		[XmlConfigFile(Name = "bCompileRewindDebuggerSupport")]
+		[CommandLine("-bCompileRewindDebuggerSupport")]
+		[RequiresUniqueBuildEnvironment]
+		public bool bCompileRewindDebuggerSupport
+		{
+			get
+			{
+				if (bEnableRewindDebuggerSupport.HasValue)
+				{
+					return bEnableRewindDebuggerSupport.Value;
+				}
+
+				// By default, if we didn't explicitly ask for RWD support, we should only enable it in non-shipping builds and only if trace is enabled
+				return Configuration != UnrealTargetConfiguration.Shipping && bEnableTrace;
+			}
+			set => bEnableRewindDebuggerSupport = value;
+		}
+		private bool? bEnableRewindDebuggerSupport;
 		
 		/// <summary>
 		/// Whether to compile in ssl support in the third party asio module
diff --git a/Engine/Source/Runtime/Engine/Private/Collision/ChaosVDSQTraceHelper.cpp b/Engine/Source/Runtime/Engine/Private/Collision/ChaosVDSQTraceHelper.cpp
--- a/Engine/Source/Runtime/Engine/Private/Collision/ChaosVDSQTraceHelper.cpp
+++ b/Engine/Source/Runtime/Engine/Private/Collision/ChaosVDSQTraceHelper.cpp
@@ -4,6 +4,7 @@
 
 #if WITH_CHAOS_VISUAL_DEBUGGER
 
+#include "ChaosVDRuntimeModule.h"
 #include "Components/PrimitiveComponent.h"
 #include "Engine/World.h"
 
diff --git a/Engine/Source/Runtime/Experimental/Chaos/Chaos.Build.cs b/Engine/Source/Runtime/Experimental/Chaos/Chaos.Build.cs
--- a/Engine/Source/Runtime/Experimental/Chaos/Chaos.Build.cs
+++ b/Engine/Source/Runtime/Experimental/Chaos/Chaos.Build.cs
@@ -22,6 +22,7 @@ namespace UnrealBuildTool.Rules
                 "Voronoi",
 				"GeometryCore",
 				"ChaosVDRuntime",
+				"TraceBasedDebuggers",
 				"NNE",
 				}
 			);
diff --git a/Engine/Source/Runtime/Experimental/Chaos/Private/ChaosVisualDebugger/ChaosVDContextProvider.cpp b/Engine/Source/Runtime/Experimental/Chaos/Private/ChaosVisualDebugger/ChaosVDContextProvider.cpp
--- a/Engine/Source/Runtime/Experimental/Chaos/Private/ChaosVisualDebugger/ChaosVDContextProvider.cpp
+++ b/Engine/Source/Runtime/Experimental/Chaos/Private/ChaosVisualDebugger/ChaosVDContextProvider.cpp
@@ -2,12 +2,22 @@
 
 #include "ChaosVisualDebugger/ChaosVDContextProvider.h"
 
+#include "ChaosVDRuntimeModule.h"
 #include "ChaosVisualDebugger/ChaosVisualDebuggerTrace.h"
+#include "TraceBasedDebuggerRuntime.h"
 
 #if WITH_CHAOS_VISUAL_DEBUGGER
 
 UE_DEFINE_THREAD_SINGLETON_TLS(FChaosVDThreadContext, CHAOS_API)
 
+FChaosVDThreadContext::~FChaosVDThreadContext()
+{
+	if (RecordingStoppedHandle.IsValid())
+	{
+		FChaosVDRuntimeModule::RemoveRecordingStopCallback(RecordingStoppedHandle);
+	}
+}
+
 bool FChaosVDThreadContext::GetCurrentContext(FChaosVDContext& OutContext)
 {
 	if (LocalContextStack.Num() > 0)
@@ -55,4 +65,18 @@ void FChaosVDThreadContext::PopContext()
 		LocalContextStack.Pop();
 	}
 }
+
+TArray<uint8>& FChaosVDThreadContext::GetTLSDataBufferRef()
+{
+	checkf(BufferAccessCounter == 0, TEXT("The CVD buffer is already in use!"));
+
+	if (!RecordingStoppedHandle.IsValid())
+	{
+		// If we access the buffer, make sure we subscribe to the recording stop callback
+		// So we free any memory we keep allocated
+		RecordingStoppedHandle = FChaosVDRuntimeModule::RegisterRecordingStopCallback(UE::TraceBasedDebuggers::FRecordingStateChangedDelegate::FDelegate::CreateLambda([this](){ CVDDataBuffer.Empty(); }));
+	}
+		
+	return CVDDataBuffer;
+}
 #endif //WITH_CHAOS_VISUAL_DEBUGGER
diff --git a/Engine/Source/Runtime/Experimental/Chaos/Private/ChaosVisualDebugger/ChaosVisualDebuggerTrace.cpp b/Engine/Source/Runtime/Experimental/Chaos/Private/ChaosVisualDebugger/ChaosVisualDebuggerTrace.cpp
--- a/Engine/Source/Runtime/Experimental/Chaos/Private/ChaosVisualDebugger/ChaosVisualDebuggerTrace.cpp
+++ b/Engine/Source/Runtime/Experimental/Chaos/Private/ChaosVisualDebugger/ChaosVisualDebuggerTrace.cpp
@@ -4,8 +4,6 @@
 
 #if WITH_CHAOS_VISUAL_DEBUGGER
 
-#include "Chaos/PBDRigidClustering.h"
-
 #include "Chaos/Character/CharacterGroundConstraintContainer.h"
 #include "Chaos/Framework/PhysicsSolverBase.h"
 #include "Chaos/ImplicitObject.h"
@@ -13,7 +11,9 @@
 #include "Chaos/ParticleHandle.h"
 #include "Chaos/PBDCollisionConstraints.h"
 #include "Chaos/PBDJointConstraints.h"
+#include "Chaos/PBDRigidClustering.h"
 #include "Chaos/PBDRigidsSOAs.h"
+#include "ChaosVDRuntimeModule.h"
 #include "ChaosVisualDebugger/ChaosVDDataWrapperUtils.h"
 #include "ChaosVisualDebugger/ChaosVDMemWriterReader.h"
 #include "ChaosVisualDebugger/ChaosVDSerializedNameTable.h"
@@ -27,13 +27,17 @@
 #include "DataWrappers/ChaosVDParticleDataWrapper.h"
 #include "DataWrappers/ChaosVDQueryDataWrappers.h"
 #include "DataWrappers/ChaosVDCameraDataWrapper.h"
+#include "Engine/Engine.h"
 #include "HAL/CriticalSection.h"
 #include "Misc/ScopeRWLock.h"
 #include "Templates/UniquePtr.h"
+#include "TraceBasedDebuggerRuntime.h"
 #include "UObject/ObjectMacros.h"
-#include "Engine/Engine.h"
 
-static_assert(CHAOS_VISUAL_DEBUGGER_WITHOUT_TRACE || UE_TRACE_MINIMAL_ENABLED || UE_TRACE_ENABLED, "Chaos Visual Debugger support requires Trace or Trace minimal support. Please enable trace support with the Unreal Build Tool flag bEnableTrace, disable cvd support with bCompileChaosVisualDebuggerSupport or allows CVD support without trace by defining CHAOS_VISUAL_DEBUGGER_WITHOUT_TRACE to 1");
+static_assert(TRACE_BASED_DEBUGGERS_WITHOUT_TRACE || UE_TRACE_MINIMAL_ENABLED || UE_TRACE_ENABLED,
+	"Chaos Visual Debugger support requires Trace or Trace minimal support."
+	" Please enable trace support with the Unreal Build Tool flag bEnableTrace, disable cvd support with bCompileChaosVisualDebuggerSupport"
+	" or allows CVD support without trace by defining TRACE_BASED_DEBUGGERS_WITHOUT_TRACE to 1");
 
 UE_TRACE_MINIMAL_EVENT_DEFINE(ChaosVDLogger, ChaosVDSolverFrameStart)
 
@@ -1360,17 +1364,18 @@ bool FChaosVisualDebuggerTrace::IsTracing()
 void FChaosVisualDebuggerTrace::RegisterEventHandlers()
 {
 	{
+		using namespace UE::TraceBasedDebuggers;
 		UE::TWriteScopeLock WriteLock(SessionState.DeltaRecordingStatesLock);
 
 		
 		if (!SessionState.RecordingStartedDelegateHandle.IsValid())
 		{
-			SessionState.RecordingStartedDelegateHandle = FChaosVDRuntimeModule::RegisterRecordingStartedCallback(FChaosVDRecordingStateChangedDelegate::FDelegate::CreateStatic(&FChaosVisualDebuggerTrace::HandleRecordingStart));
+			SessionState.RecordingStartedDelegateHandle = FChaosVDRuntimeModule::RegisterRecordingStartedCallback(FRecordingStateChangedDelegate::FDelegate::CreateStatic(&FChaosVisualDebuggerTrace::HandleRecordingStart));
 		}
 
 		if (!SessionState.RecordingStoppedDelegateHandle.IsValid())
 		{
-			SessionState.RecordingStoppedDelegateHandle = FChaosVDRuntimeModule::RegisterRecordingStopCallback(FChaosVDRecordingStateChangedDelegate::FDelegate::CreateStatic(&FChaosVisualDebuggerTrace::HandleRecordingStop));
+			SessionState.RecordingStoppedDelegateHandle = FChaosVDRuntimeModule::RegisterRecordingStopCallback(FRecordingStateChangedDelegate::FDelegate::CreateStatic(&FChaosVisualDebuggerTrace::HandleRecordingStop));
 		}
 
 		if (!SessionState.RecordingFullCaptureRequestedHandle.IsValid())
diff --git a/Engine/Source/Runtime/Experimental/Chaos/Public/Chaos/PBDRigidsEvolution.h b/Engine/Source/Runtime/Experimental/Chaos/Public/Chaos/PBDRigidsEvolution.h
--- a/Engine/Source/Runtime/Experimental/Chaos/Public/Chaos/PBDRigidsEvolution.h
+++ b/Engine/Source/Runtime/Experimental/Chaos/Public/Chaos/PBDRigidsEvolution.h
@@ -21,7 +21,11 @@
 #include "Chaos/PendingSpatialData.h"
 #include "ProfilingDebugging/CsvProfiler.h"
 #include "RewindData.h"
+
+#if UE_ENABLE_INCLUDE_ORDER_DEPRECATED_IN_5_8
 #include "ChaosVisualDebugger/ChaosVDContextProvider.h"
+#endif // UE_ENABLE_INCLUDE_ORDER_DEPRECATED_IN_5_8
+
 
 // Enable support for Collision Test Mode (reset particle positions and constraints every tick for debugging)
 #define CHAOS_EVOLUTION_COLLISION_TESTMODE (!UE_BUILD_TEST && !UE_BUILD_SHIPPING)
diff --git a/Engine/Source/Runtime/Experimental/Chaos/Public/ChaosVisualDebugger/ChaosVDContextProvider.h b/Engine/Source/Runtime/Experimental/Chaos/Public/ChaosVisualDebugger/ChaosVDContextProvider.h
--- a/Engine/Source/Runtime/Experimental/Chaos/Public/ChaosVisualDebugger/ChaosVDContextProvider.h
+++ b/Engine/Source/Runtime/Experimental/Chaos/Public/ChaosVisualDebugger/ChaosVDContextProvider.h
@@ -10,7 +10,9 @@
 
 #if WITH_CHAOS_VISUAL_DEBUGGER
 
+#if UE_ENABLE_INCLUDE_ORDER_DEPRECATED_IN_5_8
 #include "ChaosVDRuntimeModule.h"
+#endif // UE_ENABLE_INCLUDE_ORDER_DEPRECATED_IN_5_8
 
 enum class EChaosVDContextType : int32
 {
@@ -63,13 +65,7 @@ class FChaosVDThreadContext : public TThreadSingleton<FChaosVDThreadContext>
 {
 public:
 
-	virtual ~FChaosVDThreadContext() override
-	{
-		if (RecordingStoppedHandle.IsValid())
-		{
-			FChaosVDRuntimeModule::RemoveRecordingStopCallback(RecordingStoppedHandle);
-		}
-	}
+	CHAOS_API virtual ~FChaosVDThreadContext() override;
 
 	/** Copies the Current CVD context data into the provided struct
 	 * @return true if the copy was successful
@@ -98,19 +94,7 @@ public:
 protected:
 	TArray<FChaosVDContext, TInlineAllocator<16>> LocalContextStack;
 
-	TArray<uint8>& GetTLSDataBufferRef()
-	{
-		checkf(BufferAccessCounter == 0, TEXT("The CVD buffer is already in use!"));
-
-		if (!RecordingStoppedHandle.IsValid())
-		{
-			// If we access the buffer, make sure we subscribe to the recording stop callback
-			// So we free any memory we keep allocated
-			RecordingStoppedHandle = FChaosVDRuntimeModule::RegisterRecordingStopCallback(FChaosVDRecordingStateChangedDelegate::FDelegate::CreateLambda([this](){ CVDDataBuffer.Empty(); }));
-		}
-		
-		return CVDDataBuffer;
-	}
+	CHAOS_API TArray<uint8>& GetTLSDataBufferRef();
 
 	FDelegateHandle RecordingStoppedHandle;
 
diff --git a/Engine/Source/Runtime/Experimental/Chaos/Public/ChaosVisualDebugger/ChaosVisualDebuggerTrace.h b/Engine/Source/Runtime/Experimental/Chaos/Public/ChaosVisualDebugger/ChaosVisualDebuggerTrace.h
--- a/Engine/Source/Runtime/Experimental/Chaos/Public/ChaosVisualDebugger/ChaosVisualDebuggerTrace.h
+++ b/Engine/Source/Runtime/Experimental/Chaos/Public/ChaosVisualDebugger/ChaosVisualDebuggerTrace.h
@@ -1,23 +1,31 @@
 // Copyright Epic Games, Inc. All Rights Reserved.
 
 #pragma once
-#include "Chaos/Core.h"
-#include "Chaos/ParticleHandleFwd.h"
+
+#if UE_ENABLE_INCLUDE_ORDER_DEPRECATED_IN_5_8
 #include "Chaos/Serializable.h"
-#include "Containers/Set.h"
+#include "ChaosVDContextProvider.h"
+#include "HAL/ThreadSafeBool.h"
+
+#if WITH_CHAOS_VISUAL_DEBUGGER
+#include "ChaosVDRuntimeModule.h"
+#endif
+
+#endif // UE_ENABLE_INCLUDE_ORDER_DEPRECATED_IN_5_8
 
 #include "Chaos/ChaosArchive.h"
-#include "ChaosVDContextProvider.h"
-#include "ChaosVDOptionalDataChannel.h"
+#include "Chaos/Core.h"
+#include "Chaos/ParticleHandleFwd.h"
 #include "Chaos/ParticleIterator.h"
-#include "HAL/ThreadSafeBool.h"
+#include "ChaosVDOptionalDataChannel.h"
+#include "Containers/Set.h"
 
 #ifndef CHAOS_VISUAL_DEBUGGER_ENABLED
 	#define CHAOS_VISUAL_DEBUGGER_ENABLED WITH_CHAOS_VISUAL_DEBUGGER
 #endif
 
-#ifndef CHAOS_VISUAL_DEBUGGER_WITHOUT_TRACE
-	#define CHAOS_VISUAL_DEBUGGER_WITHOUT_TRACE 0
+#ifndef TRACE_BASED_DEBUGGERS_WITHOUT_TRACE
+	#define TRACE_BASED_DEBUGGERS_WITHOUT_TRACE 0
 #endif
 
 #include "Chaos/Framework/PhysicsSolverBase.h"
@@ -27,12 +35,13 @@
 
 #include "Chaos/ChaosDebugNameDefines.h"
 #include "ChaosVDMemWriterReader.h"
-#include "ChaosVDRuntimeModule.h"
 #include "Containers/StripedMap.h"
 #include "DataWrappers/ChaosVDImplicitObjectDataWrapper.h"
 #include "Trace/Trace.h"
 #include "Trace/Trace.inl"
 
+enum class EChaosVDFullCaptureFlags;
+
 namespace Chaos
 {
 class FRigidClustering;
diff --git a/Engine/Source/Runtime/Experimental/ChaosSolverEngine/ChaosSolverEngine.Build.cs b/Engine/Source/Runtime/Experimental/ChaosSolverEngine/ChaosSolverEngine.Build.cs
--- a/Engine/Source/Runtime/Experimental/ChaosSolverEngine/ChaosSolverEngine.Build.cs
+++ b/Engine/Source/Runtime/Experimental/ChaosSolverEngine/ChaosSolverEngine.Build.cs
@@ -4,21 +4,22 @@ namespace UnrealBuildTool.Rules
 {
 	public class ChaosSolverEngine : ModuleRules
 	{
-        public ChaosSolverEngine(ReadOnlyTargetRules Target) : base(Target)
+		public ChaosSolverEngine(ReadOnlyTargetRules Target) : base(Target)
 		{
 			PublicDependencyModuleNames.AddRange(
 				new string[]
 				{
 					"Core",
 					"CoreUObject",
-                    "Engine",
-                    "RenderCore",
-                    "RHI",
-					"DeveloperSettings", 
+					"Engine",
+					"RenderCore",
+					"RHI",
+					"DeveloperSettings",
 					"DataflowCore",
 					"DataflowEngine",
 					"DataflowSimulation",
 					"ChaosVDRuntime", // If CVD is compiled out, this only provides access to its structs data types
+					"TraceBasedDebuggers"
 				}
 				);
 
@@ -28,7 +29,6 @@ namespace UnrealBuildTool.Rules
 					"Messaging"
 				}
 				);
-			
 
 			SetupModulePhysicsSupport(Target);
 			SetupModuleChaosVisualDebuggerSupport(Target);
diff --git a/Engine/Source/Runtime/Experimental/ChaosSolverEngine/Private/Chaos/ChaosVDEngineEditorBridge.cpp b/Engine/Source/Runtime/Experimental/ChaosSolverEngine/Private/Chaos/ChaosVDEngineEditorBridge.cpp
--- a/Engine/Source/Runtime/Experimental/ChaosSolverEngine/Private/Chaos/ChaosVDEngineEditorBridge.cpp
+++ b/Engine/Source/Runtime/Experimental/ChaosSolverEngine/Private/Chaos/ChaosVDEngineEditorBridge.cpp
@@ -2,12 +2,10 @@
 
 #include "Chaos/ChaosVDEngineEditorBridge.h"
 
-#include "GameFramework/PlayerController.h"
-
 #if WITH_CHAOS_VISUAL_DEBUGGER
 
+#include "Chaos/ChaosVDRemoteSessionsManager.h"
 #include "ChaosVDRuntimeModule.h"
-#include "Chaos/ChaosVDTraceRelayTransport.h"
 #include "ChaosVisualDebugger/ChaosVDTraceMacros.h"
 #include "ChaosVisualDebugger/ChaosVisualDebuggerTrace.h"
 #include "Components/PrimitiveComponent.h"
@@ -19,6 +17,8 @@
 #include "PhysicsProxy/SingleParticlePhysicsProxy.h"
 #include "Runtime/Experimental/Chaos/Private/Chaos/PhysicsObjectInternal.h"
 #include "Serialization/MemoryWriter.h"
+#include "TraceBasedDebuggerRuntime.h"
+#include "TraceDataRelayTransport.h"
 #include "UObject/Package.h"
 
 
@@ -51,7 +51,7 @@ bool FChaosVDEngineEditorBridge::AddOnScreenRecordingMessage(float DummyDeltaTim
 
 	// Add a long duration value, we will remove the message manually when the recording stops
 	constexpr float MessageDurationSeconds = 3600.0f;
-	GEngine->AddOnScreenDebugMessage(CVDRecordingMessageKey, MessageDurationSeconds, FColor::Red,ChaosVDRecordingStartedMessage.ToString());
+	GEngine->AddOnScreenDebugMessage(CVDRecordingMessageKey, MessageDurationSeconds, FColor::Red, ChaosVDRecordingStartedMessage.ToString());
 
 	return bContinueLooping;
 }
@@ -78,15 +78,14 @@ void FChaosVDEngineEditorBridge::RemoveOnScreenRecordingMessage()
 
 void FChaosVDEngineEditorBridge::HandleCVDRecordingStarted()
 {
-	AddOnScreenRecordingMessage();
-
-	BroadcastSessionStatus(FApp::GetDeltaTime());
+	// Call base class to handle external communication
+	// and will call back with OnRecordingStartedInternal.
+	HandleRecordingStarted();
+}
 
-	if (!RecordingStatusUpdateHandle.IsValid())
-	{
-		constexpr float UpdateInterval = 0.5f;
-		RecordingStatusUpdateHandle = FTSTicker::GetCoreTicker().AddTicker(FTickerDelegate::CreateRaw(this, &FChaosVDEngineEditorBridge::BroadcastSessionStatus), UpdateInterval);
-	}
+void FChaosVDEngineEditorBridge::OnRecordingStartedInternal()
+{
+	AddOnScreenRecordingMessage();
 }
 
 void FChaosVDEngineEditorBridge::HandleCVDPostRecordingStarted()
@@ -95,22 +94,37 @@ void FChaosVDEngineEditorBridge::HandleCVDPostRecordingStarted()
 }
 
 void FChaosVDEngineEditorBridge::HandleCVDRecordingStopped()
+{
+	// Call base class to handle external communication
+	// and will call back with OnRecordingStoppedInternal.
+	HandleRecordingStopped();
+}
+
+void FChaosVDEngineEditorBridge::OnRecordingStoppedInternal()
 {
 	RemoveOnScreenRecordingMessage();
+}
 
-	FTSTicker::RemoveTicker(RecordingStatusUpdateHandle);
-	RecordingStatusUpdateHandle.Reset();
-	BroadcastSessionStatus(FApp::GetDeltaTime());
+void FChaosVDEngineEditorBridge::HandleTraceConnectionDetailsUpdated() const
+{
+	UE::TraceBasedDebuggers::FTraceConnectionDetailsMessage ConnectionDetailsMessage;
+	ConnectionDetailsMessage.InstanceId = FApp::GetInstanceId();
+	ConnectionDetailsMessage.TraceDetails = FChaosVDRuntimeModule::Get().GetCurrentTraceConnectionDetails();
+
+	GetSessionsManager()->PublishMessage(ConnectionDetailsMessage);
 }
 
 void FChaosVDEngineEditorBridge::HandleCVDRecordingStartFailed(const FText& InFailureReason) const
 {
 #if !WITH_EDITOR
-	// In non-editor builds we don't have an error pop-up, therefore we want to show the error message on screen
-	FText ErrorMessage = FText::FormatOrdered(NSLOCTEXT("ChaosVisualDebugger","StartRecordingFailedOnScreenMessage", "Failed to start CVD recording. {0}"), InFailureReason);
+	if (GEngine)
+	{
+		// In non-editor builds we don't have an error pop-up, therefore we want to show the error message on screen
+		FText ErrorMessage = FText::FormatOrdered(NSLOCTEXT("ChaosVisualDebugger", "StartRecordingFailedOnScreenMessage", "Failed to start CVD recording. {0}"), InFailureReason);
 
-	constexpr float MessageDurationSeconds = 4.0f;
-	GEngine->AddOnScreenDebugMessage(CVDRecordingMessageKey, MessageDurationSeconds, FColor::Red, ErrorMessage.ToString());
+		constexpr float MessageDurationSeconds = 4.0f;
+		GEngine->AddOnScreenDebugMessage(CVDRecordingMessageKey, MessageDurationSeconds, FColor::Red, ErrorMessage.ToString());
+	}
 #endif
 }
 
@@ -125,15 +139,15 @@ void FChaosVDEngineEditorBridge::HandlePIEStarted(UGameInstance* GameInstance)
 
 void FChaosVDEngineEditorBridge::HandleDataChannelChanged(TWeakPtr<FChaosVDDataDataChannel> ChannelWeakPtr)
 {
-	if (TSharedPtr<FChaosVDDataDataChannel> DataChannelPtr = ChannelWeakPtr.Pin())
+	if (const TSharedPtr<FChaosVDDataDataChannel> DataChannelPtr = ChannelWeakPtr.Pin())
 	{
 		FChaosVDChannelStateChangeResponseMessage NewChannelState;
 		NewChannelState.InstanceID = FApp::GetInstanceId();
 		NewChannelState.NewState.bIsEnabled = DataChannelPtr->IsChannelEnabled();
 		NewChannelState.NewState.ChannelName = DataChannelPtr->GetId().ToString();
 		NewChannelState.NewState.bCanChangeChannelState = DataChannelPtr->CanChangeEnabledState();
-		
-		RemoteSessionsManager->PublishDataChannelStateChangeUpdate(NewChannelState);
+
+		GetSessionsManager()->PublishMessage(NewChannelState);
 	}
 }
 
@@ -162,25 +176,25 @@ void FChaosVDEngineEditorBridge::SerializeCollisionChannelsNames()
 	CVD_TRACE_BINARY_DATA(CollisionChannelsDataBuffer, FChaosVDCollisionChannelsInfoContainer::WrapperTypeName);
 }
 
+FChaosVDEngineEditorBridge::FChaosVDEngineEditorBridge(): FEngineEditorBridge(LogChaos)
+{
+	RegisterRemoteSessionsHandler(MakeShared<FChaosVDRemoteSessionsHandler>());
+}
+
 FChaosVDEngineEditorBridge& FChaosVDEngineEditorBridge::Get()
 {
 	static FChaosVDEngineEditorBridge CVDEngineEditorBridge;
 	return CVDEngineEditorBridge;
 }
 
-bool FChaosVDEngineEditorBridge::BroadcastSessionStatus(float DeltaTime)
+void FChaosVDEngineEditorBridge::BuildRecordingStatusInternal(UE::TraceBasedDebuggers::FRecordingStatusMessage& OutStatusMessage) const
 {
-	FChaosVDRuntimeModule& RuntimeModule = FChaosVDRuntimeModule::Get();
+	const FChaosVDRuntimeModule& RuntimeModule = FChaosVDRuntimeModule::Get();
 
-	FChaosVDRecordingStatusMessage StatusMessage;
-	StatusMessage.InstanceId = FApp::GetInstanceId();
-	StatusMessage.bIsRecording = RuntimeModule.IsRecording();
-	StatusMessage.ElapsedTime = RuntimeModule.GetAccumulatedRecordingTime();
-	StatusMessage.BufferedDataBytesSize = RuntimeModule.GetBufferedDataBytesSize();
-
-	RemoteSessionsManager->PublishRecordingStatusUpdate(StatusMessage);
-
-	return true;
+	OutStatusMessage.DebuggerId = RuntimeModule.GetDebuggerId();
+	OutStatusMessage.bIsRecording = RuntimeModule.IsRecording();
+	OutStatusMessage.ElapsedTime = RuntimeModule.GetAccumulatedRecordingTime();
+	OutStatusMessage.BufferedDataBytesSize = RuntimeModule.GetBufferedDataBytesSize();
 }
 
 FChaosVDParticleMetadata FChaosVDEngineEditorBridge::GenerateParticleMetadata(const IPhysicsProxyBase* ParticleProxy, const Chaos::FGeometryParticleHandle* ParticleHandle)
@@ -261,38 +275,15 @@ FChaosVDParticleMetadata FChaosVDEngineEditorBridge::GenerateParticleMetadata(co
 	return Metadata;
 }
 
-void FChaosVDEngineEditorBridge::SetExternalTraceRelayInstance(const TSharedPtr<Chaos::VD::ITraceDataRelayTransport>& InExternalTraceRelayInstance)
-{
-	if (!ensure(!RelayTraceDataTransportInstance))
-	{
-		UE_LOG(LogChaos, Error, TEXT("An external relay instance was already set. This cannot be changed afterwards. Ignoring set request..."));
-		return;
-	}
-
-	RelayTraceDataTransportInstance = InExternalTraceRelayInstance;
-}
-
-void FChaosVDEngineEditorBridge::HandleTraceConnectionDetailsUpdated()
-{
-	FChaosVDRuntimeModule& RuntimeModule = FChaosVDRuntimeModule::Get();
-
-	FChaosVDTraceConnectionDetailsMessage ConnectionDetailsMessage;
-	ConnectionDetailsMessage.InstanceId = FApp::GetInstanceId();
-	ConnectionDetailsMessage.TraceDetails = RuntimeModule.GetCurrentTraceSessionDetails();
-
-	RemoteSessionsManager->PublishTraceConnectionDetailsUpdate(ConnectionDetailsMessage);
-}
-
-void FChaosVDEngineEditorBridge::Initialize()
+void FChaosVDEngineEditorBridge::OnInitializeInternal()
 {
+	using namespace UE::TraceBasedDebuggers;
 	FChaosVDRuntimeModule& CVDRuntimeModule = FChaosVDRuntimeModule::Get();
 
-	RemoteSessionsManager->Initialize();
-
-	RecordingStartedHandle = CVDRuntimeModule.RegisterRecordingStartedCallback(FChaosVDRecordingStateChangedDelegate::FDelegate::CreateRaw(this, &FChaosVDEngineEditorBridge::HandleCVDRecordingStarted));
-	PostRecordingStartedHandle = CVDRuntimeModule.RegisterPostRecordingStartedCallback(FChaosVDRecordingStateChangedDelegate::FDelegate::CreateRaw(this, &FChaosVDEngineEditorBridge::HandleCVDPostRecordingStarted));
-	RecordingStoppedHandle = CVDRuntimeModule.RegisterRecordingStopCallback(FChaosVDRecordingStateChangedDelegate::FDelegate::CreateRaw(this, &FChaosVDEngineEditorBridge::HandleCVDRecordingStopped));
-	RecordingStartFailedHandle = CVDRuntimeModule.RegisterRecordingStartFailedCallback(FChaosVDRecordingStartFailedDelegate::FDelegate::CreateRaw(this, &FChaosVDEngineEditorBridge::HandleCVDRecordingStartFailed));
+	RecordingStartedHandle = CVDRuntimeModule.RegisterRecordingStartedCallback(FRecordingStateChangedDelegate::FDelegate::CreateRaw(this, &FChaosVDEngineEditorBridge::HandleCVDRecordingStarted));
+	PostRecordingStartedHandle = CVDRuntimeModule.RegisterPostRecordingStartedCallback(FRecordingStateChangedDelegate::FDelegate::CreateRaw(this, &FChaosVDEngineEditorBridge::HandleCVDPostRecordingStarted));
+	RecordingStoppedHandle = CVDRuntimeModule.RegisterRecordingStopCallback(FRecordingStateChangedDelegate::FDelegate::CreateRaw(this, &FChaosVDEngineEditorBridge::HandleCVDRecordingStopped));
+	RecordingStartFailedHandle = CVDRuntimeModule.RegisterRecordingStartFailedCallback(FRecordingStartFailedDelegate::FDelegate::CreateRaw(this, &FChaosVDEngineEditorBridge::HandleCVDRecordingStartFailed));
 
 	CVDRuntimeModule.OnTraceConnectionDetailsUpdated().AddRaw(this, &FChaosVDEngineEditorBridge::HandleTraceConnectionDetailsUpdated);
 
@@ -304,39 +295,26 @@ void FChaosVDEngineEditorBridge::Initialize()
 	PIEStartedHandle = FWorldDelegates::OnPIEStarted.AddRaw(this, &FChaosVDEngineEditorBridge::HandlePIEStarted);
 #endif
 
-	// If we were already recording, we need to make sure we run the initialization step to setup the session broadcast ticker
+	// If we were already recording, we need to make sure we run the initialization step to set up the session broadcast ticker
 	// and the collision channel serialization
 	if (CVDRuntimeModule.IsRecording())
 	{
 		HandleCVDRecordingStarted();
 		HandleCVDPostRecordingStarted();
 	}
-
-	bIsInitialized = true;
-
-	EngineEditorBridgeInitializedDelegate.Broadcast();
 }
 
-void FChaosVDEngineEditorBridge::TearDown()
+void FChaosVDEngineEditorBridge::OnTearDownInternal()
 {
-	if (RelayTraceDataTransportInstance)
-	{
-		RelayTraceDataTransportInstance->Shutdown();
-	}
-
 #if WITH_EDITOR
 	FWorldDelegates::OnPIEStarted.Remove(PIEStartedHandle);
 	PIEStartedHandle.Reset();
 #endif
 	FTSTicker::RemoveTicker(DeferredShowMessageOnScreenHandle);
 	DeferredShowMessageOnScreenHandle.Reset();
-	
-	FTSTicker::RemoveTicker(RecordingStatusUpdateHandle);
-	RecordingStatusUpdateHandle.Reset();
 
 	if (FModuleManager::Get().IsModuleLoaded(TEXT("ChaosVDRuntime")))
 	{
-
 		FChaosVDRuntimeModule& CVDRuntimeModule = FChaosVDRuntimeModule::Get();
 		CVDRuntimeModule.RemoveRecordingStartedCallback(RecordingStartedHandle);
 		RecordingStartedHandle.Reset();
@@ -349,16 +327,12 @@ void FChaosVDEngineEditorBridge::TearDown()
 
 		CVDRuntimeModule.OnTraceConnectionDetailsUpdated().RemoveAll(this);
 
-		// Make sure of removing the message from the screen in case the recording didn't quite stopped yet
+		// Make sure of removing the message from the screen in case the recording didn't quite stop yet
 		if (CVDRuntimeModule.IsRecording())
 		{
 			HandleCVDRecordingStopped();
 		}
-
-		RemoteSessionsManager->Shutdown();
 	}
-
-	bIsInitialized = false;
 }
 #else
 
@@ -368,4 +342,4 @@ FChaosVDEngineEditorBridge& FChaosVDEngineEditorBridge::Get()
 	return CVDEngineEditorBridge;
 }
 
-#endif
+#endif // WITH_CHAOS_VISUAL_DEBUGGER
diff --git a/Engine/Source/Runtime/Experimental/ChaosSolverEngine/Private/Chaos/ChaosVDRemoteSessionsManager.cpp b/Engine/Source/Runtime/Experimental/ChaosSolverEngine/Private/Chaos/ChaosVDRemoteSessionsManager.cpp
--- a/Engine/Source/Runtime/Experimental/ChaosSolverEngine/Private/Chaos/ChaosVDRemoteSessionsManager.cpp
+++ b/Engine/Source/Runtime/Experimental/ChaosSolverEngine/Private/Chaos/ChaosVDRemoteSessionsManager.cpp
@@ -3,747 +3,186 @@
 #include "Chaos/ChaosVDRemoteSessionsManager.h"
 
 #include "ChaosVDRuntimeModule.h"
-#include "IMessageBus.h"
-#include "IMessagingModule.h"
-#include "MessageEndpoint.h"
 #include "MessageEndpointBuilder.h"
 #include "ChaosVisualDebugger/ChaosVDOptionalDataChannel.h"
 #include "ChaosVisualDebugger/ChaosVisualDebuggerTrace.h"
-#include "Misc/App.h"
 
-#include UE_INLINE_GENERATED_CPP_BY_NAME(ChaosVDRemoteSessionsManager)
-
-const FGuid FChaosVDRemoteSessionsManager::AllRemoteSessionsWrapperGUID = FGuid::NewGuid();
-const FGuid FChaosVDRemoteSessionsManager::AllRemoteServersWrapperGUID = FGuid::NewGuid();
-const FGuid FChaosVDRemoteSessionsManager::AllRemoteClientsWrapperGUID = FGuid::NewGuid();
-const FGuid FChaosVDRemoteSessionsManager::AllSessionsWrapperGUID = FGuid::NewGuid();
-const FGuid FChaosVDRemoteSessionsManager::CustomSessionsWrapperGUID = FGuid::NewGuid();
-const FGuid FChaosVDRemoteSessionsManager::InvalidSessionGUID = FGuid();
-const FGuid FChaosVDRemoteSessionsManager::LocalSessionID = FApp::GetInstanceId();
-
-const FString FChaosVDRemoteSessionsManager::LocalEditorSessionName = TEXT("Local Editor");
-const FGuid FChaosVDRemoteSessionsManager::LocalEditorSessionID = IsController() ? FApp::GetInstanceId() : InvalidSessionGUID;
-const FName FChaosVDRemoteSessionsManager::MessageBusEndPointName = FName("CVDSessionManagerEndPoint");
-const FString FChaosVDRemoteSessionsManager::AllRemoteSessionsTargetName = TEXT("All Remote");
-const FString FChaosVDRemoteSessionsManager::AllRemoteServersTargetName = TEXT("All Remote Servers");
-const FString FChaosVDRemoteSessionsManager::AllRemoteClientsTargetName = TEXT("All Remote Clients");
-const FString FChaosVDRemoteSessionsManager::AllSessionsTargetName = TEXT("All Sessions");
-const FString FChaosVDRemoteSessionsManager::CustomSessionsTargetName = TEXT("Custom Selection");
-
-DEFINE_LOG_CATEGORY(LogChaosVDRemoteSession)
-
-const FChaosVDTraceDetails& FChaosVDSessionInfo::GetConnectionDetails()
-{
-	return LastKnownConnectionDetails;
-}
-
-bool FChaosVDSessionInfo::IsRecording() const
-{
-	return LastKnownRecordingState.bIsRecording;
-}
-
-uint64 FChaosVDSessionInfo::GetBufferedBytesNum() const
-{
-	return LastKnownRecordingState.BufferedDataBytesSize;
-}
-
-EChaosVDRecordingMode FChaosVDSessionInfo::GetRecordingMode() const
-{
-	PRAGMA_DISABLE_DEPRECATION_WARNINGS
-	return LastKnownConnectionDetails.Mode;
-	PRAGMA_ENABLE_DEPRECATION_WARNINGS
-}
-
-EChaosVDRecordingMode FChaosVDSessionInfo::GetLastRequestedRecordingMode() const
-{
-	PRAGMA_DISABLE_DEPRECATION_WARNINGS
-	return EChaosVDRecordingMode::Invalid;
-	PRAGMA_ENABLE_DEPRECATION_WARNINGS
-}
-
-void FChaosVDSessionInfo::SetLastRequestedRecordingMode(EChaosVDRecordingMode NewRecordingMode)
-{
-	PRAGMA_DISABLE_DEPRECATION_WARNINGS
-	LastRequestedRecordingMode = NewRecordingMode;
-	PRAGMA_ENABLE_DEPRECATION_WARNINGS
-}
-
-bool FChaosVDSessionInfo::IsConnected() const
-{
-	return false;
-}
-
-void FChaosVDSessionInfo::SetReceivedBytesPerSecond(uint64 InNewBytesPerSecond)
-{
-	ReceivedBytesPerSecond = InNewBytesPerSecond;
-}
-
-uint64 FChaosVDSessionInfo::GetReceivedBytesPerSecond() const
-{
-	return ReceivedBytesPerSecond;
-}
-
-bool FChaosVDMultiSessionInfo::IsRecording() const
+void FChaosVDRemoteSessionsHandler::OnCreatingMessageEndpoint(
+	const TSharedPtr<UE::TraceBasedDebuggers::FRemoteSessionsManager>& SessionsManager
+	, FMessageEndpointBuilder& InEndPointBuilder)
 {
-	bool bIsRecording = false;
-	EnumerateInnerSessions([&bIsRecording](const TSharedRef<FChaosVDSessionInfo>& InSessionRef)
-	{
-		if (InSessionRef->IsRecording())
-		{
-			bIsRecording = true;
-			return false;
-		}
-		return true;
-	});
-
-	return bIsRecording;
-}
-
-uint64 FChaosVDMultiSessionInfo::GetBufferedBytesNum() const
-{
-	uint64 BufferedBytesNum = 0;
-	EnumerateInnerSessions([&BufferedBytesNum](const TSharedRef<FChaosVDSessionInfo>& InSessionRef)
-	{
-		BufferedBytesNum += InSessionRef->GetBufferedBytesNum();
-		return true;
-	});
-
-	return BufferedBytesNum;
-}
-
-EChaosVDRecordingMode FChaosVDMultiSessionInfo::GetRecordingMode() const
-{
-	PRAGMA_DISABLE_DEPRECATION_WARNINGS
-	EChaosVDRecordingMode FirstValidInstanceRecordingMode = EChaosVDRecordingMode::Invalid;
-	EnumerateInnerSessions([&FirstValidInstanceRecordingMode](const TSharedRef<FChaosVDSessionInfo>& InSessionRef)
-	{
-		EChaosVDRecordingMode RecordingMode = InSessionRef->GetRecordingMode();
-		if (RecordingMode == EChaosVDRecordingMode::Invalid)
-		{
-			// In multi-session, all recordings will have the same recording mode, but not of them might report connected state at the same time
-			// Therefore we need to continue searching until one of the session has a valid state before giving up.
-			return true;
-		}
-
-		FirstValidInstanceRecordingMode = RecordingMode;
-		return false;
-	});
-
-	return FirstValidInstanceRecordingMode;
-	PRAGMA_ENABLE_DEPRECATION_WARNINGS
-}
-
-void FChaosVDRemoteSessionsManager::RegisterBuiltInMessageTypes()
-{
-	SupportedMessageTypes.Emplace(FChaosVDSessionPong::StaticStruct());
-	SupportedMessageTypes.Emplace(FChaosVDRecordingStatusMessage::StaticStruct());
-	SupportedMessageTypes.Emplace(FChaosVDSessionPing::StaticStruct());
-	SupportedMessageTypes.Emplace(FChaosVDStartRecordingCommandMessage::StaticStruct());
-	SupportedMessageTypes.Emplace(FChaosVDStopRecordingCommandMessage::StaticStruct());
-	SupportedMessageTypes.Emplace(FChaosVDChannelStateChangeCommandMessage::StaticStruct());
-	SupportedMessageTypes.Emplace(FChaosVDFullSessionInfoRequestMessage::StaticStruct());
-	SupportedMessageTypes.Emplace(FChaosVDFullSessionInfoResponseMessage::StaticStruct());
-	SupportedMessageTypes.Emplace(FChaosVDChannelStateChangeResponseMessage::StaticStruct());
-	SupportedMessageTypes.Emplace(FChaosVDTraceConnectionDetailsMessage::StaticStruct());
-}
-
-FChaosVDRemoteSessionsManager::FChaosVDRemoteSessionsManager()
-{
-	RegisterBuiltInMessageTypes();
-}
-
-void FChaosVDRemoteSessionsManager::Initialize(const TSharedPtr<IMessageBus>& InMessageBus)
-{
-	ensureMsgf(!InMessageBus, TEXT("Initialize(MessageBusPtr) is deprecated!. The provided message buss will be ignored"));
-
-	Initialize();
-}
-
-constexpr bool FChaosVDRemoteSessionsManager::IsController()
-{
-#if WITH_EDITOR
-	return true;
-#else
-	return false;
-#endif
-}
-
-TWeakPtr<FChaosVDSessionInfo> FChaosVDRemoteSessionsManager::GetSessionInfo(FGuid Id)
-{
-	if (TSharedPtr<FChaosVDSessionInfo>* FoundSessionPtrPtr = ActiveSessionsByInstanceId.Find(Id))
-	{
-		return *FoundSessionPtrPtr;
-	}
-
-	return nullptr;
-}
-
-TWeakPtr<FChaosVDSessionInfo> FChaosVDRemoteSessionsManager::GetFirstAvailableSingleSessionInfo()
-{
-	for (const TPair<FGuid, TSharedPtr<FChaosVDSessionInfo>>& SessionWithID : ActiveSessionsByInstanceId)
-	{
-		if (SessionWithID.Value && !EnumHasAnyFlags(SessionWithID.Value->GetSessionTypeAttributes(), EChaosVDRemoteSessionAttributes::IsMultiSessionWrapper))
-		{
-			return SessionWithID.Value;
-		}
-	}
-
-	return nullptr;
-}
-
-TSharedPtr<FChaosVDSessionInfo> FChaosVDRemoteSessionsManager::CreatedWrapperSessionInfo(FGuid InstanceId, const FString& SessionName)
-{
-	TSharedPtr<FChaosVDMultiSessionInfo> NewSessionInfo = MakeShared<FChaosVDMultiSessionInfo>();
-	
-	NewSessionInfo->InstanceId = InstanceId;
-	NewSessionInfo->SessionName = SessionName;
-
-	return NewSessionInfo;
-}
-
-TSharedPtr<FMessageEndpoint> FChaosVDRemoteSessionsManager::CreateEndPoint(const TSharedRef<IMessageBus>& InMessageBus)
-{
-	FMessageEndpointBuilder EndPointBuilder = FMessageEndpoint::Builder(MessageBusEndPointName, InMessageBus)
-	                                          .Handling<FChaosVDSessionPing>(this, &FChaosVDRemoteSessionsManager::HandleSessionPingMessage)
-	                                          .Handling<FChaosVDStartRecordingCommandMessage>(this, &FChaosVDRemoteSessionsManager::HandleRecordingStartCommandMessage)
-	                                          .Handling<FChaosVDStopRecordingCommandMessage>(this, &FChaosVDRemoteSessionsManager::HandleRecordingStopCommandMessage)
-	                                          .Handling<FChaosVDChannelStateChangeCommandMessage>(this, &FChaosVDRemoteSessionsManager::HandleChangeDataChannelStateCommandMessage)
-	                                          .Handling<FChaosVDFullSessionInfoRequestMessage>(this, &FChaosVDRemoteSessionsManager::HandleFullSessionStateRequestMessage);
-
-	if (IsController())
-	{
-		EndPointBuilder.Handling<FChaosVDSessionPong>(this, &FChaosVDRemoteSessionsManager::HandleSessionPongMessage)
-						.Handling<FChaosVDRecordingStatusMessage>(this, &FChaosVDRemoteSessionsManager::HandleRecordingStatusUpdateMessage)
-						.Handling<FChaosVDFullSessionInfoResponseMessage>(this, &FChaosVDRemoteSessionsManager::HandleFullSessionStateResponseMessage)
-						.Handling<FChaosVDChannelStateChangeResponseMessage>(this, &FChaosVDRemoteSessionsManager::HandleChangeDataChannelStateResponseMessage)
-						.Handling<FChaosVDTraceConnectionDetailsMessage>(this, &FChaosVDRemoteSessionsManager::HandleConnectionDetailsUpdateMessage);
-	}
-
-	return EndPointBuilder.Build();
-}
-
-void FChaosVDRemoteSessionsManager::ReInitializeMessagingSystem(const TSharedPtr<IMessageBus>& InMessageBus)
-{
-	ShutdownMessagingSystem();
-	InitializeMessagingSystem(InMessageBus);
-}
-
-void FChaosVDRemoteSessionsManager::RegisterExternalSupportedMessageType(const UScriptStruct* ScriptStruct)
-{
-	if (ensure(ScriptStruct))
-	{
-		SupportedMessageTypes.Emplace(ScriptStruct);
-	}
-}
-
-void FChaosVDRemoteSessionsManager::InitializeMessagingSystem(const TSharedPtr<IMessageBus>& InMessageBus)
-{
-	if (!InMessageBus)
-	{
-		return;
-	}
-
-	MessageBusPtr = InMessageBus;
-	MessageEndpoint = CreateEndPoint(InMessageBus.ToSharedRef());
-	
-	if (!MessageEndpoint)
-	{
-		return;
-	}
-
-	if (IsController())
-	{
-		MessageEndpoint->Subscribe<FChaosVDSessionPong>();
-		MessageEndpoint->Subscribe<FChaosVDRecordingStatusMessage>();
-		MessageEndpoint->Subscribe<FChaosVDFullSessionInfoResponseMessage>();
-		MessageEndpoint->Subscribe<FChaosVDChannelStateChangeResponseMessage>();
-		MessageEndpoint->Subscribe<FChaosVDTraceConnectionDetailsMessage>();
-	}
-
-	MessageEndpoint->Subscribe<FChaosVDSessionPing>();
-	MessageEndpoint->Subscribe<FChaosVDStartRecordingCommandMessage>();
-	MessageEndpoint->Subscribe<FChaosVDStopRecordingCommandMessage>();
-	MessageEndpoint->Subscribe<FChaosVDChannelStateChangeCommandMessage>();
-	MessageEndpoint->Subscribe<FChaosVDFullSessionInfoRequestMessage>();
-
-	if (bInitialized)
-	{
-		MessageEndpoint->Enable();
-	}
-	else
-	{
-		MessageEndpoint->Disable();
-	}
-
-	MessagingInitializedDelegate.Broadcast(InMessageBus, MessageEndpoint);
-}
+	InEndPointBuilder
+		.Handling<FChaosVDChannelStateChangeCommandMessage>(
+			[](const FChaosVDChannelStateChangeCommandMessage& InMessage, const TSharedRef<IMessageContext>&)
+			{
+				UE_AUTORTFM_ONCOMMIT(InMessage)
+				{
+					using namespace Chaos::VisualDebugger;
+					TWeakPtr<FChaosVDOptionalDataChannel> ChannelInstance = FChaosVDDataChannelsManager::Get().GetChannelById(FName(InMessage.NewState.ChannelName));
+					if (TSharedPtr<FChaosVDOptionalDataChannel> LockedChannelInstance = ChannelInstance.Pin())
+					{
+						LockedChannelInstance->SetChannelEnabled(InMessage.NewState.bIsEnabled);
+					}
+				};
+			})
+		.Handling<FChaosVDStartRecordingCommandMessage>(
+			[](const FChaosVDStartRecordingCommandMessage& InMessage, const TSharedRef<IMessageContext>& InContext)
+			{
+				UE_AUTORTFM_ONCOMMIT(InMessage, InContext)
+				{
+					FChaosVisualDebuggerTrace::OverrideDefaultEnabledDataChannels(InMessage.DataChannelsEnabledOverrideList);
+					FChaosVDRuntimeModule::Get().StartRecording(InMessage);
+				};
+			})
+		.Handling<FChaosVDStopRecordingCommandMessage>(
+			[](const FChaosVDStopRecordingCommandMessage&, const TSharedRef<IMessageContext>&)
+			{
+				UE_AUTORTFM_ONCOMMIT()
+				{
+					FChaosVDRuntimeModule::Get().StopRecording();
+				};
+			});
 
-void FChaosVDRemoteSessionsManager::ShutdownMessagingSystem()
-{
-	if (MessageEndpoint)
+	using namespace UE::TraceBasedDebuggers;
+	if (SessionsManager && SessionsManager->IsController())
 	{
-		MessageEndpoint->Unsubscribe();
+		InEndPointBuilder
+			.Handling<FChaosVDChannelStateChangeResponseMessage>(
+				[WeakSessionsManager = SessionsManager.ToWeakPtr()](const FChaosVDChannelStateChangeResponseMessage& InMessage, const TSharedRef<IMessageContext>&)
+				{
+					if (const TSharedPtr<FRemoteSessionsManager> SessionManager = WeakSessionsManager.Pin())
+					{
+						if (const TSharedPtr<FSessionInfo> SessionInfoPtr = SessionManager->GetSessionInfo(InMessage.InstanceID).Pin())
+						{
+							if (FChaosVDSessionData* ChaosData = SessionInfoPtr->GetDebuggerData<FChaosVDSessionData>())
+							{
+								if (FChaosVDDataChannelState* FoundChannelState = ChaosData->DataChannelsStatesByName.Find(InMessage.NewState.ChannelName))
+								{
+									*FoundChannelState = InMessage.NewState;
+								}
+							}
+						}
+					}
+				});
 	}
-
-	MessageBusPtr.Reset();
-	MessageEndpoint = nullptr;
 }
 
-void FChaosVDRemoteSessionsManager::Initialize()
+void FChaosVDRemoteSessionsHandler::OnSubscribingMessageTypes(
+	const TSharedPtr<UE::TraceBasedDebuggers::FRemoteSessionsManager>& InSessionsManager
+	, FMessageEndpoint& InMessageEndpoint)
 {
-#if !defined(WITH_CHAOS_VISUAL_DEBUGGER_EXTERNAL_MESSAGING) || !WITH_CHAOS_VISUAL_DEBUGGER_EXTERNAL_MESSAGING
-	InitializeMessagingSystem(IMessagingModule::Get().GetDefaultBus());
-#endif
-
-	ActiveSessionsByInstanceId.Add(AllRemoteSessionsWrapperGUID, CreatedWrapperSessionInfo(AllRemoteSessionsWrapperGUID, AllRemoteSessionsTargetName));
-	ActiveSessionsByInstanceId.Add(AllRemoteServersWrapperGUID, CreatedWrapperSessionInfo(AllRemoteServersWrapperGUID, AllRemoteServersTargetName));
-	ActiveSessionsByInstanceId.Add(AllRemoteClientsWrapperGUID, CreatedWrapperSessionInfo(AllRemoteClientsWrapperGUID, AllRemoteClientsTargetName));
-	ActiveSessionsByInstanceId.Add(AllSessionsWrapperGUID, CreatedWrapperSessionInfo(AllSessionsWrapperGUID, AllSessionsTargetName));
-	ActiveSessionsByInstanceId.Add(CustomSessionsWrapperGUID, CreatedWrapperSessionInfo(CustomSessionsWrapperGUID, CustomSessionsTargetName));
-
-	if (MessageEndpoint)
+	if (InSessionsManager)
 	{
-		MessageEndpoint->Enable();
-	}
-
-	bInitialized = true;
-}
-
-void FChaosVDRemoteSessionsManager::Shutdown()
-{
-	bInitialized = false;
-	ShutdownMessagingSystem();
-}
+		InSessionsManager->RegisterExternalSupportedMessageType<FChaosVDStartRecordingCommandMessage>();
+		InSessionsManager->RegisterExternalSupportedMessageType<FChaosVDStopRecordingCommandMessage>();
+		InSessionsManager->RegisterExternalSupportedMessageType<FChaosVDChannelStateChangeCommandMessage>();
+		InSessionsManager->RegisterExternalSupportedMessageType<FChaosVDChannelStateChangeResponseMessage>();
 
-void FChaosVDRemoteSessionsManager::EnumerateMessageTypes(const FVisitorFunction& InVisitor)
-{
-	for (const UScriptStruct* MessageType : SupportedMessageTypes)
-	{
-		if (MessageType)
+		if (InSessionsManager->IsController())
 		{
-			InVisitor(MessageType);
+			InMessageEndpoint.Subscribe<FChaosVDChannelStateChangeResponseMessage>();
 		}
 	}
-}
-
-void FChaosVDRemoteSessionsManager::StartSessionDiscovery()
-{
-	if (TickHandle.IsValid())
-	{
-		UE_LOG(LogChaosVDRemoteSession, Warning, TEXT("[%hs] Session discovery already started"), __func__);
-		return;
-	}
-	
-	constexpr float TickInterval = 1.0f;
-	TickHandle = FTSTicker::GetCoreTicker().AddTicker(FTickerDelegate::CreateRaw(this, &FChaosVDRemoteSessionsManager::Tick), TickInterval);
-}
 
-void FChaosVDRemoteSessionsManager::StopSessionDiscovery()
-{
-	if (TickHandle.IsValid())
-	{
-		RemoveExpiredSessions(ERemoveSessionOptions::ForceRemoveAll);
-		FTSTicker::RemoveTicker(TickHandle);
-		TickHandle = FTSTicker::FDelegateHandle();
-	}
+	InMessageEndpoint.Subscribe<FChaosVDStartRecordingCommandMessage>();
+	InMessageEndpoint.Subscribe<FChaosVDStopRecordingCommandMessage>();
+	InMessageEndpoint.Subscribe<FChaosVDChannelStateChangeCommandMessage>();
 }
 
-void FChaosVDRemoteSessionsManager::PublishRecordingStatusUpdate(const FChaosVDRecordingStatusMessage& InUpdateMessage)
+void FChaosVDRemoteSessionsHandler::OnBuildingFullSessionInfoResponseMessage(
+	UE::TraceBasedDebuggers::FFullSessionInfoResponseMessage& InMessage)
 {
-	if (ensure(MessageEndpoint))
-	{
-		MessageEndpoint->Publish(FMessageEndpoint::MakeMessage<FChaosVDRecordingStatusMessage>(InUpdateMessage), EMessageScope::Network);
-	}
-}
+	using namespace Chaos::VisualDebugger;
 
-void FChaosVDRemoteSessionsManager::PublishTraceConnectionDetailsUpdate(const FChaosVDTraceConnectionDetailsMessage& InUpdateMessage)
-{
-	if (ensure(MessageEndpoint))
-	{
-		MessageEndpoint->Publish(FMessageEndpoint::MakeMessage<FChaosVDTraceConnectionDetailsMessage>(InUpdateMessage), EMessageScope::Network);
-	}
-}
+	InMessage.bIsRecording = FChaosVDRuntimeModule::Get().IsRecording();
+	InMessage.DebuggerId = Chaos::VD::DebuggerGuid;
 
-void FChaosVDRemoteSessionsManager::PublishDataChannelStateChangeUpdate(const FChaosVDChannelStateChangeResponseMessage& InNewStateData)
-{
-	if (ensure(MessageEndpoint))
-	{
-		MessageEndpoint->Publish(FMessageEndpoint::MakeMessage<FChaosVDChannelStateChangeResponseMessage>(InNewStateData), EMessageScope::Network);
-	}
-}
-
-void FChaosVDRemoteSessionsManager::SendStartRecordingCommand(const FMessageAddress& InDestinationAddress, const FChaosVDStartRecordingCommandMessage& RecordingStartCommandParams)
-{
-	if (!MessageEndpoint)
-	{
-		UE_LOG(LogChaosVDRemoteSession, Error, TEXT("[%hs] Failed to send command | Invalid endpoint."), __func__);
-		return;
-	}
-
-	MessageEndpoint->Send(
-	FMessageEndpoint::MakeMessage<FChaosVDStartRecordingCommandMessage>(RecordingStartCommandParams),
-	FChaosVDStartRecordingCommandMessage::StaticStruct(),
-	EMessageFlags::Reliable,
-	nullptr,
-	TArrayBuilder<FMessageAddress>().Add(InDestinationAddress),
-	FTimespan::Zero(),
-	FDateTime::MaxValue());
-}
-
-void FChaosVDRemoteSessionsManager::SendStopRecordingCommand(const FMessageAddress& InDestinationAddress)
-{
-	if (!MessageEndpoint)
-	{
-		UE_LOG(LogChaosVDRemoteSession, Error, TEXT("[%s] Failed to send command | Invalid endpoint."), ANSI_TO_TCHAR(__FUNCTION__));
-		return;
-	}
-
-	MessageEndpoint->Send(
-	FMessageEndpoint::MakeMessage<FChaosVDStopRecordingCommandMessage>(),
-	FChaosVDStopRecordingCommandMessage::StaticStruct(),
-	EMessageFlags::Reliable,
-	nullptr,
-	TArrayBuilder<FMessageAddress>().Add(InDestinationAddress),
-	FTimespan::Zero(),
-	FDateTime::MaxValue());
-}
-
-void FChaosVDRemoteSessionsManager::SendDataChannelStateChangeCommand(const FMessageAddress& InDestinationAddress,const FChaosVDChannelStateChangeCommandMessage& InNewStateData)
-{
-	if (!MessageEndpoint)
-	{
-		UE_LOG(LogChaosVDRemoteSession, Error, TEXT("[%s] Failed to send command | Invalid endpoint."), ANSI_TO_TCHAR(__FUNCTION__));
-		return;
-	}
-
-	MessageEndpoint->Send(
-	FMessageEndpoint::MakeMessage<FChaosVDChannelStateChangeCommandMessage>(InNewStateData),
-	FChaosVDChannelStateChangeCommandMessage::StaticStruct(),
-	EMessageFlags::Reliable,
-	nullptr,
-	TArrayBuilder<FMessageAddress>().Add(InDestinationAddress),
-	FTimespan::Zero(),
-	FDateTime::MaxValue());
-}
-
-void FChaosVDRemoteSessionsManager::SendFullSessionStateRequestCommand(const FMessageAddress& InDestinationAddress)
-{
-	if (!MessageEndpoint)
-	{
-		UE_LOG(LogChaosVDRemoteSession, Error, TEXT("[%s] Failed to send command | Invalid endpoint."), ANSI_TO_TCHAR(__FUNCTION__));
-		return;
-	}
-
-	MessageEndpoint->Send(
-	FMessageEndpoint::MakeMessage<FChaosVDFullSessionInfoRequestMessage>(),
-	FChaosVDFullSessionInfoRequestMessage::StaticStruct(),
-	EMessageFlags::Reliable,
-	nullptr,
-	TArrayBuilder<FMessageAddress>().Add(InDestinationAddress),
-	FTimespan::Zero(),
-	FDateTime::MaxValue());
-}
-
-bool FChaosVDRemoteSessionsManager::Tick(float DeltaTime)
-{
-	if (IsController())
-	{
-		SendPing();
-		RemoveExpiredSessions();
-	}
-
-	return true;
-}
-
-void FChaosVDRemoteSessionsManager::SendPing()
-{
-	if (ensure(MessageEndpoint))
-	{
-		if (FChaosVDSessionPing* SessionPingData = FMessageEndpoint::MakeMessage<FChaosVDSessionPing>())
+	FChaosVDFullSessionInfoResponseData MessageData;
+	FChaosVDDataChannelsManager::Get().EnumerateDataChannels([&MessageData](const TWeakPtr<FChaosVDOptionalDataChannel>& Channel)
 		{
-			SessionPingData->ControllerInstanceId = FApp::GetInstanceId();
-			MessageEndpoint->Publish(SessionPingData, EMessageScope::Network);
-		}
-	}
+			if (const TSharedPtr<FChaosVDOptionalDataChannel> LockedChannel = Channel.Pin())
+			{
+				MessageData.DataChannelsStates.Emplace(
+					FChaosVDDataChannelState
+					{
+						LockedChannel->GetId().ToString()
+						, LockedChannel->IsChannelEnabled()
+						, LockedChannel->CanChangeEnabledState()
+					});
+			}
+			return true;
+		});
+
+	InMessage.DebuggerSpecificData.Emplace(FInstancedStruct::Make(MoveTemp(MessageData)));
 }
 
-void FChaosVDRemoteSessionsManager::SendPong(const FChaosVDSessionPing& InMessage)
+void FChaosVDRemoteSessionsHandler::OnHandlingFullSessionInfoResponseMessage(
+	const UE::TraceBasedDebuggers::FFullSessionInfoResponseMessage& InMessage
+	, const TSharedPtr<UE::TraceBasedDebuggers::FSessionInfo>& InSessionInfo)
 {
-	if (!ensure(MessageEndpoint))
+	if (!InSessionInfo.IsValid())
 	{
 		return;
 	}
 
-	if (FChaosVDSessionPong* PongMessage = FMessageEndpoint::MakeMessage<FChaosVDSessionPong>())
-	{
-		PongMessage->InstanceId = FApp::GetInstanceId();
-		PongMessage->SessionId = FApp::GetSessionId();
-		PongMessage->BuildTargetType = static_cast<uint8>(FApp::GetBuildTargetType());
-
-		if (InMessage.ControllerInstanceId == PongMessage->InstanceId)
-		{
-			PongMessage->SessionName = LocalEditorSessionName;
-		}
-		else
-		{
-			FString AppSessionName = FApp::GetSessionName();
-			PongMessage->SessionName = AppSessionName == TEXT("None") || AppSessionName.IsEmpty() ? FString::Format(TEXT("{0} {1} {2}"), { FApp::GetProjectName(), FString(LexToString(FApp::GetBuildTargetType())), FString::FromInt(FPlatformProcess::GetCurrentProcessId()) }) : AppSessionName;	
-		}
-
-		MessageEndpoint->Publish(PongMessage, EMessageScope::Network);
-	}
-}
-
-void FChaosVDRemoteSessionsManager::RegisterSessionInMultiSessionWrapper(const TSharedRef<FChaosVDSessionInfo>& InSessionInfoRef)
-{
-	if (InSessionInfoRef->SessionName != LocalEditorSessionName )
+	if (const FChaosVDFullSessionInfoResponseData* MessageData = InMessage.GetDebuggerData<FChaosVDFullSessionInfoResponseData>())
 	{
-		StaticCastSharedPtr<FChaosVDMultiSessionInfo>(ActiveSessionsByInstanceId.FindChecked(AllRemoteSessionsWrapperGUID))->InnerSessionsByInstanceID.Emplace(InSessionInfoRef->InstanceId, InSessionInfoRef);
-
-		if (InSessionInfoRef->BuildTargetType == EBuildTargetType::Server)
-		{
-			StaticCastSharedPtr<FChaosVDMultiSessionInfo>(ActiveSessionsByInstanceId.FindChecked(AllRemoteServersWrapperGUID))->InnerSessionsByInstanceID.Emplace(InSessionInfoRef->InstanceId, InSessionInfoRef);
-		}
-		else
+		FChaosVDSessionData NewSessionData;
+		NewSessionData.DataChannelsStatesByName.Reserve(MessageData->DataChannelsStates.Num());
+		for (const FChaosVDDataChannelState& ChannelState : MessageData->DataChannelsStates)
 		{
-			StaticCastSharedPtr<FChaosVDMultiSessionInfo>(ActiveSessionsByInstanceId.FindChecked(AllRemoteClientsWrapperGUID))->InnerSessionsByInstanceID.Emplace(InSessionInfoRef->InstanceId, InSessionInfoRef);
-		}
-	}
-
-	StaticCastSharedPtr<FChaosVDMultiSessionInfo>(ActiveSessionsByInstanceId.FindChecked(AllSessionsWrapperGUID))->InnerSessionsByInstanceID.Emplace(InSessionInfoRef->InstanceId, InSessionInfoRef);
-}
-
-void FChaosVDRemoteSessionsManager::DeRegisterSessionInMultiSessionWrapper(const TSharedRef<FChaosVDSessionInfo>& InSessionInfoRef)
-{
-	StaticCastSharedPtr<FChaosVDMultiSessionInfo>(ActiveSessionsByInstanceId.FindChecked(AllRemoteSessionsWrapperGUID))->InnerSessionsByInstanceID.Remove(InSessionInfoRef->InstanceId);
-	StaticCastSharedPtr<FChaosVDMultiSessionInfo>(ActiveSessionsByInstanceId.FindChecked(AllSessionsWrapperGUID))->InnerSessionsByInstanceID.Remove(InSessionInfoRef->InstanceId);
-	StaticCastSharedPtr<FChaosVDMultiSessionInfo>(ActiveSessionsByInstanceId.FindChecked(AllRemoteServersWrapperGUID))->InnerSessionsByInstanceID.Remove(InSessionInfoRef->InstanceId);
-	StaticCastSharedPtr<FChaosVDMultiSessionInfo>(ActiveSessionsByInstanceId.FindChecked(AllRemoteClientsWrapperGUID))->InnerSessionsByInstanceID.Remove(InSessionInfoRef->InstanceId);
-}
-
-void FChaosVDRemoteSessionsManager::ProcessPendingMessagesForSession(const FChaosVDSessionPong& InMessage, const TSharedRef<FChaosVDSessionInfo>& InSessionInfoPtr)
-{
-	PendingRecordingStatusMessages.RemoveAndCopyValue(InMessage.InstanceId, InSessionInfoPtr->LastKnownRecordingState);
-	PendingRecordingConnectionDetailsMessages.RemoveAndCopyValue(InMessage.InstanceId, InSessionInfoPtr->LastKnownConnectionDetails);
-}
-
-void FChaosVDRemoteSessionsManager::HandleSessionPongMessage(const FChaosVDSessionPong& InMessage, const TSharedRef<IMessageContext>& InContext)
-{
-	TSharedPtr<FChaosVDSessionInfo>& SessionInfoPtr = ActiveSessionsByInstanceId.FindOrAdd(InMessage.InstanceId);
-
-	if (!SessionInfoPtr)
-	{
-		SessionInfoPtr = MakeShared<FChaosVDSessionInfo>();
-		SessionInfoPtr->Address = InContext->GetSender();
-		SessionInfoPtr->InstanceId = InMessage.InstanceId;
-		SessionInfoPtr->SessionName = InMessage.SessionName;
-		SessionInfoPtr->BuildTargetType = static_cast<EBuildTargetType>(InMessage.BuildTargetType);
+			NewSessionData.DataChannelsStatesByName.Emplace(ChannelState.ChannelName, ChannelState);
+		};
 
-		RegisterSessionInMultiSessionWrapper(SessionInfoPtr.ToSharedRef());
-
-		SessionDiscoveredDelegate.Broadcast(SessionInfoPtr->InstanceId);
-
-		// This is the first time we see this session, so we need to request the rest of its state so we can properly populate the UI
-		SendFullSessionStateRequestCommand(SessionInfoPtr->Address);
+		InSessionInfo->SetDebuggerData<FChaosVDSessionData>(MoveTemp(NewSessionData));
 	}
-
-	SessionInfoPtr->LastPingTime = FDateTime::UtcNow();
-
-	ProcessPendingMessagesForSession(InMessage, SessionInfoPtr.ToSharedRef());
-	
-	SessionsUpdatedDelegate.Broadcast();
 }
 
-void FChaosVDRemoteSessionsManager::HandleSessionPingMessage(const FChaosVDSessionPing& InMessage, const TSharedRef<IMessageContext>& InContext)
-{
-	// If this instance is not running trace, don't answer to the ping as the CVD instance will not be able to do anything useful with it
-#if	!CHAOS_VISUAL_DEBUGGER_WITHOUT_TRACE
-	SendPong(InMessage);
-#endif
-}
+PRAGMA_DISABLE_DEPRECATION_WARNINGS
 
-void FChaosVDRemoteSessionsManager::HandleRecordingStatusUpdateMessage(const FChaosVDRecordingStatusMessage& Message, const TSharedRef<IMessageContext>& InContext)
+const FChaosVDTraceDetails& FChaosVDSessionInfo::GetConnectionDetails()
 {
-	if (TSharedPtr<FChaosVDSessionInfo>* SessionInfoPtrPtr = ActiveSessionsByInstanceId.Find(Message.InstanceId))
-	{
-		TSharedPtr<FChaosVDSessionInfo>& SessionInfoPtr = *SessionInfoPtrPtr;
-		check(SessionInfoPtr);
-
-		if (SessionInfoPtr->LastKnownRecordingState.bIsRecording != Message.bIsRecording)
-		{
-			if (Message.bIsRecording)
-			{
-				RecordingStartedDelegate.Broadcast(SessionInfoPtr);
-			}
-			else
-			{
-				RecordingStoppedDelegate.Broadcast(SessionInfoPtr);
-			}
-		}
-
-		SessionInfoPtr->LastKnownRecordingState = Message;
-	}
-	else
-	{
-		PendingRecordingStatusMessages.FindOrAdd(Message.InstanceId) = Message;
-	}
+	return LastKnownConnectionDetails;
 }
 
-void FChaosVDRemoteSessionsManager::HandleConnectionDetailsUpdateMessage(const FChaosVDTraceConnectionDetailsMessage& InMessage, const TSharedRef<IMessageContext>& InContext)
+EChaosVDRecordingMode FChaosVDSessionInfo::GetRecordingMode() const
 {
-	if (TSharedPtr<FChaosVDSessionInfo>* SessionInfoPtrPtr = ActiveSessionsByInstanceId.Find(InMessage.InstanceId))
-	{
-		TSharedPtr<FChaosVDSessionInfo>& SessionInfoPtr = *SessionInfoPtrPtr;
-		check(SessionInfoPtr);
-
-		SessionInfoPtr->LastKnownConnectionDetails = InMessage.TraceDetails;
-	}
-	else
-	{
-		PendingRecordingConnectionDetailsMessages.FindOrAdd(InMessage.InstanceId) = InMessage.TraceDetails;
-	}
+	return LastKnownConnectionDetails.Mode;
 }
 
-void FChaosVDRemoteSessionsManager::HandleRecordingStartCommandMessage(const FChaosVDStartRecordingCommandMessage& InMessage, const TSharedRef<IMessageContext>& InContext)
+EChaosVDRecordingMode FChaosVDSessionInfo::GetLastRequestedRecordingMode() const
 {
-#if WITH_CHAOS_VISUAL_DEBUGGER
-	UE_AUTORTFM_ONCOMMIT(InMessage, InContext)
-	{
-		FChaosVisualDebuggerTrace::OverrideDefaultEnabledDataChannels(InMessage.DataChannelsEnabledOverrideList);
-		FChaosVDRuntimeModule::Get().StartRecording(InMessage);
-	};
-
-#endif
+	return EChaosVDRecordingMode::Invalid;
 }
 
-void FChaosVDRemoteSessionsManager::HandleRecordingStopCommandMessage(const FChaosVDStopRecordingCommandMessage& InMessage, const TSharedRef<IMessageContext>& InContext)
+void FChaosVDSessionInfo::SetLastRequestedRecordingMode(EChaosVDRecordingMode NewRecordingMode)
 {
-#if WITH_CHAOS_VISUAL_DEBUGGER
-	UE_AUTORTFM_ONCOMMIT()
-	{
-		FChaosVDRuntimeModule::Get().StopRecording();
-	};
-#endif
+	LastRequestedRecordingMode = NewRecordingMode;
 }
 
-void FChaosVDRemoteSessionsManager::HandleChangeDataChannelStateCommandMessage(const FChaosVDChannelStateChangeCommandMessage& InMessage, const TSharedRef<IMessageContext>& InContext)
+bool FChaosVDSessionInfo::IsConnected() const
 {
-#if WITH_CHAOS_VISUAL_DEBUGGER
-	UE_AUTORTFM_ONCOMMIT(InMessage)
-	{
-		using namespace Chaos::VisualDebugger;
-		TWeakPtr<FChaosVDOptionalDataChannel> ChannelInstance = FChaosVDDataChannelsManager::Get().GetChannelById(FName(InMessage.NewState.ChannelName));
-		if (TSharedPtr<FChaosVDOptionalDataChannel> LockedChannelInstance = ChannelInstance.Pin())
-		{
-			LockedChannelInstance->SetChannelEnabled(InMessage.NewState.bIsEnabled);
-		}
-	};
-#endif
+	return false;
 }
 
-void FChaosVDRemoteSessionsManager::HandleChangeDataChannelStateResponseMessage(const FChaosVDChannelStateChangeResponseMessage& InMessage, const TSharedRef<IMessageContext>& InContext)
+EChaosVDRecordingMode FChaosVDMultiSessionInfo::GetRecordingMode() const
 {
-	if (TSharedPtr<FChaosVDSessionInfo>* SessionInfoPtrPtr = ActiveSessionsByInstanceId.Find(InMessage.InstanceID))
-	{
-		if (TSharedPtr<FChaosVDSessionInfo>& SessionInfoPtr = *SessionInfoPtrPtr)
+	EChaosVDRecordingMode FirstValidInstanceRecordingMode = EChaosVDRecordingMode::Invalid;
+	EnumerateInnerSessions([&FirstValidInstanceRecordingMode](const TSharedRef<FChaosVDSessionInfo>& InSessionRef)
 		{
-			if (FChaosVDDataChannelState* FoundChannelState = SessionInfoPtr->DataChannelsStatesByName.Find(InMessage.NewState.ChannelName))
+			EChaosVDRecordingMode RecordingMode = InSessionRef->GetRecordingMode();
+			if (RecordingMode == EChaosVDRecordingMode::Invalid)
 			{
-				*FoundChannelState = InMessage.NewState;
+				// In multi-session, all recordings will have the same recording mode, but not of them might report connected state at the same time
+				// Therefore we need to continue searching until one of the session has a valid state before giving up.
+				return true;
 			}
-		}
-	}
-}
 
-void FChaosVDRemoteSessionsManager::HandleFullSessionStateRequestMessage(const FChaosVDFullSessionInfoRequestMessage& InMessage, const TSharedRef<IMessageContext>& InContext)
-{
-	if (!ensure(MessageEndpoint))
-	{
-		return;
-	}
-	
-	if (FChaosVDFullSessionInfoResponseMessage* FullSessionStateResponse = FMessageEndpoint::MakeMessage<FChaosVDFullSessionInfoResponseMessage>())
-	{
-		FullSessionStateResponse->InstanceId = FApp::GetInstanceId();
-
-#if WITH_CHAOS_VISUAL_DEBUGGER
-		FullSessionStateResponse->bIsRecording = FChaosVDRuntimeModule::Get().IsRecording();
-
-		using namespace Chaos::VisualDebugger;
-		FChaosVDDataChannelsManager::Get().EnumerateDataChannels([&FullSessionStateResponse](const TWeakPtr<FChaosVDOptionalDataChannel>& Channel)
-		{
-			if (TSharedPtr<FChaosVDOptionalDataChannel> LockedChannel = Channel.Pin())
-			{
-				FChaosVDChannelStateChangeCommandMessage DataChannelStateMessage = { LockedChannel->GetId().ToString(), LockedChannel->IsChannelEnabled(), LockedChannel->CanChangeEnabledState() };
-				FullSessionStateResponse->DataChannelsStates.Emplace(DataChannelStateMessage.NewState);
-			}
-			return true;
+			FirstValidInstanceRecordingMode = RecordingMode;
+			return false;
 		});
-#endif
-
-		MessageEndpoint->Send(
-			FullSessionStateResponse,
-			FChaosVDFullSessionInfoResponseMessage::StaticStruct(),
-			EMessageFlags::Reliable,
-			nullptr,
-			TArrayBuilder<FMessageAddress>().Add(InContext->GetSender()),
-			FTimespan::Zero(),
-			FDateTime::MaxValue());
-	}
-}
 
-void FChaosVDRemoteSessionsManager::HandleFullSessionStateResponseMessage(const FChaosVDFullSessionInfoResponseMessage& InMessage, const TSharedRef<IMessageContext>& InContext)
-{
-	if (TSharedPtr<FChaosVDSessionInfo>* SessionInfoPtrPtr = ActiveSessionsByInstanceId.Find(InMessage.InstanceId))
-	{
-		if (TSharedPtr<FChaosVDSessionInfo>& SessionInfoPtr = *SessionInfoPtrPtr)
-		{
-			SessionInfoPtr->LastKnownRecordingState.bIsRecording = InMessage.bIsRecording;
-
-			for (const FChaosVDDataChannelState& ChannelState : InMessage.DataChannelsStates)
-			{
-				SessionInfoPtr->DataChannelsStatesByName.Emplace(ChannelState.ChannelName, ChannelState);
-			}
-		}
-	}
+	return FirstValidInstanceRecordingMode;	
 }
 
-void FChaosVDRemoteSessionsManager::RemoveExpiredSessions(ERemoveSessionOptions Options)
-{
-	bool bAnySessionRemoved = false;
-	FDateTime CurrentTime = FDateTime::UtcNow();
-	for (TMap<FGuid, TSharedPtr<FChaosVDSessionInfo>>::TIterator RemoveIterator = ActiveSessionsByInstanceId.CreateIterator(); RemoveIterator; ++RemoveIterator)
-	{
-		TSharedPtr<FChaosVDSessionInfo>& SessionInfoPtr = RemoveIterator.Value();
-		if (!SessionInfoPtr)
-		{
-			RemoveIterator.RemoveCurrent();
-			bAnySessionRemoved = true;
-			continue;
-		}
-
-		if (!EnumHasAnyFlags(SessionInfoPtr->GetSessionTypeAttributes(), EChaosVDRemoteSessionAttributes::CanExpire))
-		{
-			continue;
-		}
-
-		FTimespan ElapsedTime = CurrentTime - SessionInfoPtr->LastPingTime;
-
-		// A session goes into busy state if we are attempting to issue a command that might stall the target, currently that only happens on recording start commands of complex maps
-		// In these cases, we need to allow more time between pings. If a recording command failed, it is expected the state to be changed to Ready again
-		const float MaxAllowedTimeBetweenPings = SessionInfoPtr->ReadyState == EChaosVDRemoteSessionReadyState::Busy || SessionInfoPtr->IsRecording() ? 60.0f : 3.0f;
-		if (EnumHasAnyFlags(Options, ERemoveSessionOptions::ForceRemoveAll) || ElapsedTime > FTimespan::FromSeconds(MaxAllowedTimeBetweenPings))
-		{
-			SessionExpiredDelegate.Broadcast(SessionInfoPtr->InstanceId);
-
-			DeRegisterSessionInMultiSessionWrapper(SessionInfoPtr.ToSharedRef());
-			RemoveIterator.RemoveCurrent();
-			bAnySessionRemoved = true;
-		}
-	}
-
-	if (bAnySessionRemoved)
-	{
-		SessionsUpdatedDelegate.Broadcast();
-	}
-}
+PRAGMA_ENABLE_DEPRECATION_WARNINGS
\ No newline at end of file
diff --git a/Engine/Source/Runtime/Experimental/ChaosSolverEngine/Public/Chaos/ChaosVDEngineEditorBridge.h b/Engine/Source/Runtime/Experimental/ChaosSolverEngine/Public/Chaos/ChaosVDEngineEditorBridge.h
--- a/Engine/Source/Runtime/Experimental/ChaosSolverEngine/Public/Chaos/ChaosVDEngineEditorBridge.h
+++ b/Engine/Source/Runtime/Experimental/ChaosSolverEngine/Public/Chaos/ChaosVDEngineEditorBridge.h
@@ -1,81 +1,84 @@
 // Copyright Epic Games, Inc. All Rights Reserved.
 
 #pragma once
+
+#if UE_ENABLE_INCLUDE_ORDER_DEPRECATED_IN_5_8
 #include "Chaos/ChaosVDRemoteSessionsManager.h"
+#endif // UE_ENABLE_INCLUDE_ORDER_DEPRECATED_IN_5_8
+
+#include "HAL/Platform.h"
+#include "Templates/SharedPointer.h"
 
 namespace Chaos::VD
 {
-	class ITraceDataRelayTransport;
+class ITraceDataRelayTransport;
+}
+
+namespace UE::TraceBasedDebuggers
+{
+struct FRemoteSessionsManager;
+class IDataRelayTransport;
 }
 
+PRAGMA_DISABLE_DEPRECATION_WARNINGS
+class FChaosVDRemoteSessionsManager;
+PRAGMA_ENABLE_DEPRECATION_WARNINGS
+
 #if WITH_CHAOS_VISUAL_DEBUGGER
 
+#include "ChaosLog.h"
 #include "Chaos/ParticleHandleFwd.h"
 #include "DataWrappers/ChaosVDParticleDataWrapper.h"
 #include "Delegates/IDelegateInstance.h"
-#include "HAL/Platform.h"
+#include "EngineEditorBridge.h"
 
 class IPhysicsProxyBase;
-class UBodySetup;
-class UChaosVDTraceDataChannel;
 class UGameInstance;
 
-namespace Chaos
+namespace Chaos::VisualDebugger
 {
-	namespace VisualDebugger
-	{
-		struct FChaosVDOptionalDataChannel;
-	}
+struct FChaosVDOptionalDataChannel;
 }
 
-#ifndef WITH_CHAOS_VISUAL_DEBUGGER_EXTERNAL_MESSAGING
-	#define WITH_CHAOS_VISUAL_DEBUGGER_EXTERNAL_MESSAGING 0
-#endif
-
 /** Object that bridges the gap between the ChaosVD runtime module and the Engine & CVD editor.
  * As the ChaosVDRuntime module does not have access to the engine module, this object reacts to events and performs necessary operation the runtime module cannot do directly
  */
-class FChaosVDEngineEditorBridge
+class FChaosVDEngineEditorBridge : public UE::TraceBasedDebuggers::FEngineEditorBridge
 {
 public:
-	
-	FChaosVDEngineEditorBridge() : RemoteSessionsManager(MakeShared<FChaosVDRemoteSessionsManager>())
-	{
-	}
-	
-	~FChaosVDEngineEditorBridge()= default;
+	CHAOSSOLVERENGINE_API FChaosVDEngineEditorBridge();
 
-	CHAOSSOLVERENGINE_API static FChaosVDEngineEditorBridge& Get();
+	~FChaosVDEngineEditorBridge() = default;
 
-	void Initialize();
-	void TearDown();
+	CHAOSSOLVERENGINE_API static FChaosVDEngineEditorBridge& Get();
 
+PRAGMA_DISABLE_DEPRECATION_WARNINGS
+	UE_DEPRECATED(5.8, "Use GetSessionsManager instead")
 	TSharedPtr<FChaosVDRemoteSessionsManager> GetRemoteSessionsManager()
 	{
-		return RemoteSessionsManager;
+		return nullptr;
 	}
 
 	using FChaosVDDataDataChannel = Chaos::VisualDebugger::FChaosVDOptionalDataChannel;
 
+	UE_DEPRECATED(5.8, "Use GetDataTransportInstance instead")
 	TSharedPtr<Chaos::VD::ITraceDataRelayTransport> GetRelayTransportInstance()
 	{
-		return RelayTraceDataTransportInstance;
-	}
-
-	bool IsInitialized() const
-	{
-		return bIsInitialized;
-	}
-
-	FSimpleMulticastDelegate& OnEngineEditorBridgeInitialized()
-	{
-		return EngineEditorBridgeInitializedDelegate;
+		return nullptr;
 	}
 
+	UE_DEPRECATED(5.8, "Use SetDataRelayTransportInstance instead")
 	CHAOSSOLVERENGINE_API void SetExternalTraceRelayInstance(const TSharedPtr<Chaos::VD::ITraceDataRelayTransport>& InExternalTraceRelayInstance);
+PRAGMA_ENABLE_DEPRECATION_WARNINGS
 
 private:
-	void HandleTraceConnectionDetailsUpdated();
+	virtual void OnInitializeInternal() override;
+	virtual void OnTearDownInternal() override;
+	virtual void OnRecordingStartedInternal() override;
+	virtual void OnRecordingStoppedInternal() override;
+	virtual void BuildRecordingStatusInternal(UE::TraceBasedDebuggers::FRecordingStatusMessage& OutStatusMessage) const override;
+
+	void HandleTraceConnectionDetailsUpdated() const;
 	bool AddOnScreenRecordingMessage(float DummyDeltaTime = 0.1f);
 	void RemoveOnScreenRecordingMessage();
 	void HandleCVDRecordingStarted();
@@ -87,8 +90,6 @@ private:
 	void HandleDataChannelChanged(TWeakPtr<FChaosVDDataDataChannel> ChannelWeakPtr);
 
 	void SerializeCollisionChannelsNames();
-	
-	bool BroadcastSessionStatus(float DeltaTime);
 
 	static FChaosVDParticleMetadata GenerateParticleMetadata(const IPhysicsProxyBase* ParticleProxy, const Chaos::FGeometryParticleHandle* ParticleHandle);
 
@@ -102,33 +103,39 @@ private:
 	FDelegateHandle PIEStartedHandle;
 #endif
 
-	TSharedPtr<Chaos::VD::ITraceDataRelayTransport> RelayTraceDataTransportInstance;
-
-	TSharedRef<FChaosVDRemoteSessionsManager> RemoteSessionsManager;
-	
-	FTSTicker::FDelegateHandle RecordingStatusUpdateHandle;
 	FTSTicker::FDelegateHandle DeferredShowMessageOnScreenHandle;
-
-	FSimpleMulticastDelegate EngineEditorBridgeInitializedDelegate;
-
-	bool bIsInitialized = false;
 };
+
 #else
 
 class FChaosVDEngineEditorBridge
 {
 public:
 	CHAOSSOLVERENGINE_API static FChaosVDEngineEditorBridge& Get();
-	
+
+PRAGMA_DISABLE_DEPRECATION_WARNINGS
+	UE_DEPRECATED(5.8, "Use GetSessionsManager instead")
 	TSharedPtr<FChaosVDRemoteSessionsManager> GetRemoteSessionsManager()
 	{
 		return nullptr;
 	}
 
+	UE_DEPRECATED(5.8, "Use GetDataTransportInstance instead")
 	TSharedPtr<Chaos::VD::ITraceDataRelayTransport> GetRelayTransportInstance()
 	{
 		return nullptr;
 	}
+PRAGMA_ENABLE_DEPRECATION_WARNINGS
+
+	TSharedPtr<UE::TraceBasedDebuggers::FRemoteSessionsManager> GetSessionsManager()
+	{
+		return nullptr;
+	}
+
+	TSharedPtr<UE::TraceBasedDebuggers::IDataRelayTransport> GetDataTransportInstance()
+	{
+		return nullptr;
+	}
 };
 
 #endif
diff --git a/Engine/Source/Runtime/Experimental/ChaosSolverEngine/Public/Chaos/ChaosVDRemoteSessionsManager.h b/Engine/Source/Runtime/Experimental/ChaosSolverEngine/Public/Chaos/ChaosVDRemoteSessionsManager.h
--- a/Engine/Source/Runtime/Experimental/ChaosSolverEngine/Public/Chaos/ChaosVDRemoteSessionsManager.h
+++ b/Engine/Source/Runtime/Experimental/ChaosSolverEngine/Public/Chaos/ChaosVDRemoteSessionsManager.h
@@ -1,102 +1,58 @@
 // Copyright Epic Games, Inc. All Rights Reserved.
 
 #pragma once
+
+#if UE_ENABLE_INCLUDE_ORDER_DEPRECATED_IN_5_8
 #include "ChaosVDRuntimeModule.h"
-#include "ChaosVDRecordingDetails.h"
 #include "Containers/Ticker.h"
 #include "IMessageContext.h"
 #include "Misc/DateTime.h"
 #include "Misc/Guid.h"
 #include "Templates/SharedPointer.h"
+#endif //UE_ENABLE_INCLUDE_ORDER_DEPRECATED_IN_5_8
 
+#include "ChaosVDRecordingDetails.h"
+#include "RemoteSessionsManager.h"
 #include "ChaosVDRemoteSessionsManager.generated.h"
 
-class IMessageContext;
-struct FChaosVDSessionPong;
-class FMessageEndpoint;
-class IMessageBus;
+#define UE_API CHAOSSOLVERENGINE_API
 
-USTRUCT()
-struct FChaosVDSessionPing
+struct UE_DEPRECATED(5.8, "Use UE::TraceBasedDebuggers::FSessionPing instead.") FChaosVDSessionPing : UE::TraceBasedDebuggers::FSessionPing
 {
-	GENERATED_BODY()
-	
-	UPROPERTY()
-	FGuid ControllerInstanceId;
 };
 
-USTRUCT()
-struct FChaosVDSessionPong
+struct UE_DEPRECATED(5.8, "Use UE::TraceBasedDebuggers::FSessionPong instead.") FChaosVDSessionPong : UE::TraceBasedDebuggers::FSessionPong
 {
-	GENERATED_BODY()
-	
-	UPROPERTY()
-	FGuid InstanceId;
-	
-	UPROPERTY()
-	FGuid SessionId;
-
-	UPROPERTY()
-	FString SessionName;
-
-	UPROPERTY()
-	uint8 BuildTargetType = static_cast<uint8>(EBuildTargetType::Unknown);
 };
 
+/** Message specific to ChaosVD to send a stop recording command. */
 USTRUCT()
-struct FChaosVDStopRecordingCommandMessage
+struct FChaosVDStopRecordingCommandMessage : public UE::TraceBasedDebuggers::FStopRecordingCommandMessage
 {
 	GENERATED_BODY()
 };
 
-USTRUCT()
-struct FChaosVDRecordingStatusMessage
+PRAGMA_DISABLE_DEPRECATION_WARNINGS
+struct UE_DEPRECATED(5.8, "Use UE::TraceBasedDebuggers::FRecordingStatusMessage instead.") FChaosVDRecordingStatusMessage : UE::TraceBasedDebuggers::FRecordingStatusMessage
 {
-
-	GENERATED_BODY()
-
 	FChaosVDRecordingStatusMessage()
 	{
 	}
 
-	PRAGMA_DISABLE_DEPRECATION_WARNINGS
 	FChaosVDRecordingStatusMessage(const FChaosVDRecordingStatusMessage& Other) = default;
 	FChaosVDRecordingStatusMessage(FChaosVDRecordingStatusMessage&& Other) noexcept = default;
 	FChaosVDRecordingStatusMessage& operator=(const FChaosVDRecordingStatusMessage& Other) = default;
 	FChaosVDRecordingStatusMessage& operator=(FChaosVDRecordingStatusMessage&& Other) noexcept = default;
-	PRAGMA_ENABLE_DEPRECATION_WARNINGS
-
-	UPROPERTY()
-	FGuid InstanceId = FGuid();
-	
-	UPROPERTY()
-	bool bIsRecording = false;
-
-	UPROPERTY()
-	float ElapsedTime = false;
-
-	UPROPERTY()
-	int64 BufferedDataBytesSize = 0;
 
 	UE_DEPRECATED(5.7, "Please get the trace details directly from the session info object")
 	FChaosVDTraceDetails TraceDetails;
 };
 
-USTRUCT()
-struct FChaosVDTraceConnectionDetailsMessage
+struct UE_DEPRECATED(5.8, "Use UE::TraceBasedDebuggers::FTraceConnectionDetailsMessage instead.") FChaosVDTraceConnectionDetailsMessage : UE::TraceBasedDebuggers::FTraceConnectionDetailsMessage
 {
-	GENERATED_BODY()
-
-	FChaosVDTraceConnectionDetailsMessage()
-	{
-	}
-
-	UPROPERTY()
-	FGuid InstanceId = FGuid();
-
-	UPROPERTY()
 	FChaosVDTraceDetails TraceDetails;
 };
+PRAGMA_ENABLE_DEPRECATION_WARNINGS
 
 USTRUCT()
 struct FChaosVDRelayTraceDataMessage
@@ -161,6 +117,7 @@ USTRUCT()
 struct FChaosVDChannelStateChangeCommandMessage
 {
 	GENERATED_BODY()
+
 	UPROPERTY()
 	FChaosVDDataChannelState NewState;
 };
@@ -177,29 +134,55 @@ struct FChaosVDChannelStateChangeResponseMessage
 	FChaosVDDataChannelState NewState;
 };
 
+/**
+ * Structure specific to ChaosVD to store the states of the data channels.
+ * The struct is stored as the DebuggerSpecificData for the message type UE::TraceBasedDebuggers::FFullSessionInfoResponseMessage.
+ */
 USTRUCT()
-struct FChaosVDFullSessionInfoRequestMessage
+struct FChaosVDFullSessionInfoResponseData
 {
 	GENERATED_BODY()
+
+	UPROPERTY()
+	TArray<FChaosVDDataChannelState> DataChannelsStates;
 };
 
+/**
+ * Structure specific to ChaosVD to store the states of the data channels in the session info struct.
+ * The struct is stored as a DebuggerSpecificSessionData inside UE::TraceBasedDebuggers::FSessionInfo
+ * and get be retrieved using GetDebuggerData<FChaosVDSessionData>().
+ */
 USTRUCT()
-struct FChaosVDFullSessionInfoResponseMessage
+struct FChaosVDSessionData
 {
 	GENERATED_BODY()
 
-	UPROPERTY()
-	FGuid InstanceId;
+	TMap<FString, FChaosVDDataChannelState> DataChannelsStatesByName;
+};
 
-	UPROPERTY()
-	TArray<FChaosVDDataChannelState> DataChannelsStates;
+/**
+ * Remote sessions handler to handle ChaosVD specific message types.
+ */
+struct FChaosVDRemoteSessionsHandler : UE::TraceBasedDebuggers::IRemoteSessionsHandler
+{
+	UE_API virtual void OnCreatingMessageEndpoint(const TSharedPtr<UE::TraceBasedDebuggers::FRemoteSessionsManager>&, FMessageEndpointBuilder&) override;
+	UE_API virtual void OnSubscribingMessageTypes(const TSharedPtr<UE::TraceBasedDebuggers::FRemoteSessionsManager>&, FMessageEndpoint&) override;
+	UE_API virtual void OnBuildingFullSessionInfoResponseMessage(UE::TraceBasedDebuggers::FFullSessionInfoResponseMessage&) override;
+	UE_API virtual void OnHandlingFullSessionInfoResponseMessage(const UE::TraceBasedDebuggers::FFullSessionInfoResponseMessage&, const TSharedPtr<UE::TraceBasedDebuggers::FSessionInfo>&) override;
+};
 
-	UPROPERTY()
-	bool bIsRecording = false;
+PRAGMA_DISABLE_DEPRECATION_WARNINGS
+
+struct UE_DEPRECATED(5.8, "Use UE::TraceBasedDebuggers::FFullSessionInfoRequestMessage instead.") FChaosVDFullSessionInfoRequestMessage
+{
 };
 
-UENUM()
-enum class EChaosVDRemoteSessionAttributes
+struct UE_DEPRECATED(5.8, "Use UE::TraceBasedDebuggers::FFullSessionInfoResponseMessage instead.") FChaosVFFullSessionInfoResponseMessage : UE::TraceBasedDebuggers::FFullSessionInfoResponseMessage
+{
+	TArray<FChaosVDDataChannelState> DataChannelsStates;
+};
+
+enum class UE_DEPRECATED(5.8, "Use UE::TraceBasedDebuggers::ERemoteSessionAttributes instead.") EChaosVDRemoteSessionAttributes
 {
 	None = 0,
 	SupportsDataChannelChange = 1 << 0,
@@ -208,26 +191,19 @@ enum class EChaosVDRemoteSessionAttributes
 };
 ENUM_CLASS_FLAGS(EChaosVDRemoteSessionAttributes)
 
-UENUM()
-enum class EChaosVDRemoteSessionReadyState : uint8
+enum class UE_DEPRECATED(5.8, "Use UE::TraceBasedDebuggers::ERemoteSessionReadyState instead.") EChaosVDRemoteSessionReadyState : uint8
 {
-	/** The session is ready to execute commands */
 	Ready,
-	/** We are executing a command in the session we expect to take a while without hearing anything from the target. */
 	Busy
 };
 
-
-/**
- * Session object that contains all the information needed to communicate with a remote instance, and the state of that instance
- */
-struct FChaosVDSessionInfo
+struct UE_DEPRECATED(5.8, "Use UE::TraceBasedDebuggers::FSessionInfo instead.") FChaosVDSessionInfo
 {
 	explicit FChaosVDSessionInfo() : InstanceId(FGuid()),
 									SessionTypeAttributes(EChaosVDRemoteSessionAttributes::CanExpire | EChaosVDRemoteSessionAttributes::SupportsDataChannelChange)
 	{
 	}
-	
+
 	virtual ~FChaosVDSessionInfo() = default;
 
 protected:
@@ -257,13 +233,13 @@ public:
 	UE_DEPRECATED(5.8, "This method will no longer be used and there is not a replacement planned.")
 	CHAOSSOLVERENGINE_API virtual EChaosVDRecordingMode GetRecordingMode() const;
 	UE_DEPRECATED(5.8, "This method will no longer be used and there is not a replacement planned.")
-	CHAOSSOLVERENGINE_API virtual EChaosVDRecordingMode GetLastRequestedRecordingMode() const;
+	CHAOSSOLVERENGINE_API virtual EChaosVDRecordingMode GetLastRequestedRecordingMode() const final;
 	UE_DEPRECATED(5.8, "This method will no longer be used and there is not a replacement planned.")
-	CHAOSSOLVERENGINE_API virtual void SetLastRequestedRecordingMode(EChaosVDRecordingMode NewRecordingMode);
+	CHAOSSOLVERENGINE_API virtual void SetLastRequestedRecordingMode(EChaosVDRecordingMode NewRecordingMode) final;
 
 
 	UE_DEPRECATED(5.7, "This method will no longer be used and there is not a replacement planned.")
-	virtual bool IsConnected() const;
+	virtual bool IsConnected() const final;
 
 	EChaosVDRemoteSessionAttributes GetSessionTypeAttributes() const
 	{
@@ -281,15 +257,9 @@ protected:
 	uint64 ReceivedBytesPerSecond = 0;
 	
 	const EChaosVDRemoteSessionAttributes SessionTypeAttributes;
-
-	friend class FChaosVDRemoteSessionsManager;
 };
 
-/**
- * Session object that is able to control and provide information about multiple session objects.
- * Used to the UI can use the same API to control multiple session, than for single sessions.
- */
-struct FChaosVDMultiSessionInfo : public FChaosVDSessionInfo
+struct UE_DEPRECATED(5.8, "Use UE::TraceBasedDebuggers::FMultiSessionInfo instead.")  FChaosVDMultiSessionInfo : public FChaosVDSessionInfo
 {
 	explicit FChaosVDMultiSessionInfo() : FChaosVDSessionInfo(EChaosVDRemoteSessionAttributes::IsMultiSessionWrapper)
 	{
@@ -323,26 +293,20 @@ struct FChaosVDMultiSessionInfo : public FChaosVDSessionInfo
 
 DECLARE_LOG_CATEGORY_EXTERN(LogChaosVDRemoteSession, Log, VeryVerbose);
 
+UE_DEPRECATED(5.8, "Use UE::TraceBasedDebuggers::FRecordingStateChangeDelegate()")
 DECLARE_MULTICAST_DELEGATE_OneParam(FChaosVDRecordingStateChangeDelegate, TWeakPtr<FChaosVDSessionInfo> Session)
+
+UE_DEPRECATED(5.8, "Use UE::TraceBasedDebuggers::FSessionLifetimeDelegate()")
 DECLARE_MULTICAST_DELEGATE_OneParam(FChaosVDSessionDiscoverDelegate, FGuid SessionID)
 
 /** Object that is able to discover, issue and execute commands back and forth between CVD and client/server/editor instances */
-class FChaosVDRemoteSessionsManager
+class UE_DEPRECATED(5.8, "Use UE::TraceBasedDebuggers::FRemoteSessionsManager instead.") FChaosVDRemoteSessionsManager : public UE::TraceBasedDebuggers::FRemoteSessionsManager
 {
 public:
 	FChaosVDRemoteSessionsManager();
 
 	UE_DEPRECATED(5.7, "Use Initialize()")
 	CHAOSSOLVERENGINE_API void Initialize(const TSharedPtr<IMessageBus>& InMessageBus);
-	CHAOSSOLVERENGINE_API void Initialize();
-	CHAOSSOLVERENGINE_API void Shutdown();
-
-	CHAOSSOLVERENGINE_API void ReInitializeMessagingSystem(const TSharedPtr<IMessageBus>& InMessageBus);
-
-	FSimpleMulticastDelegate& OnSessionsUpdated()
-	{
-		return SessionsUpdatedDelegate;
-	}
 
 	FChaosVDSessionDiscoverDelegate& OnSessionDiscovered()
 	{
@@ -354,24 +318,6 @@ public:
 		return SessionExpiredDelegate;
 	}
 
-	typedef TFunction<void(const UScriptStruct*)> FVisitorFunction;
-	
-	/**
-	 * Enumerate messages types allowed by this module.
-	 * @param InVisitor
-	 */
-	CHAOSSOLVERENGINE_API void EnumerateMessageTypes(const FVisitorFunction& InVisitor);
-
-	/**
-	 * Starts sending ping requests that potential targets (client, servers or other editors) can answer to and report themselves as available
-	 */
-	CHAOSSOLVERENGINE_API void StartSessionDiscovery();
-	
-	/**
-     * Stop sending ping requests to potential targets (client, servers or other editors).
-     */
-	CHAOSSOLVERENGINE_API void StopSessionDiscovery();
-
 	/**
 	 * Issues a command to the provided address that will start a CVD recording
 	 * @param InDestinationAddress Message bus address that will execute the command
@@ -379,12 +325,6 @@ public:
 	 */
 	CHAOSSOLVERENGINE_API void SendStartRecordingCommand(const FMessageAddress& InDestinationAddress, const FChaosVDStartRecordingCommandMessage& RecordingStartCommandParams);
 
-	/**
-	 * Issues a command to the provided address to stop a CVD recording 
-	 * 	@param InDestinationAddress Message bus address that will execute the command
-	 */
-	CHAOSSOLVERENGINE_API void SendStopRecordingCommand(const FMessageAddress& InDestinationAddress);
-
 	/**
 	 * Issues a command to the provided address to change the state of a data channel
 	 * 	@param InDestinationAddress Message bus address that will execute the command
@@ -447,7 +387,7 @@ public:
 	 */
 	FChaosVDRecordingStateChangeDelegate& OnSessionRecordingStarted()
 	{
-		return RecordingStartedDelegate;	
+		return RecordingStartedDelegate;
 	}
 
 	/**
@@ -457,148 +397,17 @@ public:
 	{
 		return RecordingStoppedDelegate;
 	}
-	
-	CHAOSSOLVERENGINE_API static const FString AllRemoteSessionsTargetName;
-	CHAOSSOLVERENGINE_API static const FString AllRemoteServersTargetName;
-	CHAOSSOLVERENGINE_API static const FString AllRemoteClientsTargetName;
-	CHAOSSOLVERENGINE_API static const FString AllSessionsTargetName;
-	CHAOSSOLVERENGINE_API static const FString CustomSessionsTargetName;
-	CHAOSSOLVERENGINE_API static const FString LocalEditorSessionName;
-	CHAOSSOLVERENGINE_API static const FName MessageBusEndPointName;
-	CHAOSSOLVERENGINE_API static const FGuid AllRemoteSessionsWrapperGUID;
-	CHAOSSOLVERENGINE_API static const FGuid AllRemoteServersWrapperGUID;
-	CHAOSSOLVERENGINE_API static const FGuid AllRemoteClientsWrapperGUID;
-	CHAOSSOLVERENGINE_API static const FGuid AllSessionsWrapperGUID;
-	CHAOSSOLVERENGINE_API static const FGuid CustomSessionsWrapperGUID;
-	CHAOSSOLVERENGINE_API static const FGuid InvalidSessionGUID;
-	CHAOSSOLVERENGINE_API static const FGuid LocalEditorSessionID;
-	CHAOSSOLVERENGINE_API static const FGuid LocalSessionID;
-
-	DECLARE_MULTICAST_DELEGATE_TwoParams(FMessagingInitialized, TSharedPtr<IMessageBus> NewMessageBus, TSharedPtr<FMessageEndpoint> NewEndpoint);
-	FMessagingInitialized& OnMessagingInitialized()
-	{
-		return MessagingInitializedDelegate;
-	}
-
-	TWeakPtr<IMessageBus> GetMessageBusInstance()
-	{
-		return MessageBusPtr;
-	}
-	
-	TWeakPtr<FMessageEndpoint> GetMessageBusEndPoint()
-	{
-		return MessageEndpoint;
-	}
-
-	/**
-	 * Returns true if this instance has controller capabilities (is either an editor or CVD Standalone, which is also an editor)
-	 */
-	CHAOSSOLVERENGINE_API static constexpr bool IsController();
-
-	/**
-	 * Adds a message type to the list of supported message types for this session system.
-	 * @param ScriptStruct Message type info
-	 */
-	CHAOSSOLVERENGINE_API void RegisterExternalSupportedMessageType(const UScriptStruct* ScriptStruct);
 
 private:
-
-	void RegisterBuiltInMessageTypes();
-
-	void InitializeMessagingSystem(const TSharedPtr<IMessageBus>& InMessageBus);
-	void ShutdownMessagingSystem();
-
-	/**
-	 * Sends a request to obtain the full session information to the provided message bus address
-	 * @param InDestinationAddress 
-	 */
-	void SendFullSessionStateRequestCommand(const FMessageAddress& InDestinationAddress);
-
-	/**
-	 * Registers a session object that is able to control multiple session instances
-	 * @param InSessionInfoRef Multi Session object
-	 */
-	void RegisterSessionInMultiSessionWrapper(const TSharedRef<FChaosVDSessionInfo>& InSessionInfoRef);
-	
-	/**
-	 * Deegisters a session object that is able to control multiple session instances
-	 * @param InSessionInfoRef Multi Session object
-	 */
-	void DeRegisterSessionInMultiSessionWrapper(const TSharedRef<FChaosVDSessionInfo>& InSessionInfoRef);
-
-	/**
-	 * Creates a session object that is able to control multiple other session objects.
-	 * @param InstanceId Instance ID for this object
-	 * @param SessionName Session name that will be used in the UI
-	 * @return 
-	 */
-	TSharedPtr<FChaosVDSessionInfo> CreatedWrapperSessionInfo(FGuid InstanceId, const FString& SessionName);
-
-	/***
-	 * Creates the messagebus endpoint this session manager will use
-	 */
-	TSharedPtr<FMessageEndpoint> CreateEndPoint(const TSharedRef<IMessageBus>& InMessageBus);
-
-	bool Tick(float DeltaTime);
-
-	/**
-	 * Usually used by the controller. Broadcast to the network this controller exists.
-	 */
-	void SendPing();
-
-	/**
-	 * Broadcast to the network a small subset of this instance information, in to a received session ping.
-	 * @param InMessage Ping message that was received
-	 */
-	void SendPong(const FChaosVDSessionPing& InMessage);
-	
-	void HandleSessionPongMessage(const FChaosVDSessionPong& InMessage, const TSharedRef<IMessageContext>& InContext);
-	void HandleSessionPingMessage(const FChaosVDSessionPing& InMessage, const TSharedRef<IMessageContext>& InContext);
-	void HandleRecordingStatusUpdateMessage(const FChaosVDRecordingStatusMessage& InMessage, const TSharedRef<IMessageContext>& InContext);
-	void HandleConnectionDetailsUpdateMessage(const FChaosVDTraceConnectionDetailsMessage& InMessage, const TSharedRef<IMessageContext>& InContext);
-	void HandleRecordingStartCommandMessage(const FChaosVDStartRecordingCommandMessage& InMessage, const TSharedRef<IMessageContext>& InContext);
-	void HandleRecordingStopCommandMessage(const FChaosVDStopRecordingCommandMessage& InMessage, const TSharedRef<IMessageContext>& InContext);
-	void HandleChangeDataChannelStateCommandMessage(const FChaosVDChannelStateChangeCommandMessage& InMessage, const TSharedRef<IMessageContext>& InContext);
-	void HandleChangeDataChannelStateResponseMessage(const FChaosVDChannelStateChangeResponseMessage& InMessage, const TSharedRef<IMessageContext>& InContext);
-	void HandleFullSessionStateRequestMessage(const FChaosVDFullSessionInfoRequestMessage& InMessage, const TSharedRef<IMessageContext>& InContext);
-	void HandleFullSessionStateResponseMessage(const FChaosVDFullSessionInfoResponseMessage& InMessage, const TSharedRef<IMessageContext>& InContext);
-
-	enum class ERemoveSessionOptions : uint8
-	{
-		None = 0,
-		ForceRemoveAll = 1 << 0
-	};
-	FRIEND_ENUM_CLASS_FLAGS(ERemoveSessionOptions)
-
-	void RemoveExpiredSessions(ERemoveSessionOptions Options = ERemoveSessionOptions::None);
-
-	void ProcessPendingMessagesForSession(const FChaosVDSessionPong& InMessage, const TSharedRef<FChaosVDSessionInfo>& InSessionInfoPtr);
-
-	/** Holds the time at which the last ping was sent. */
-    FDateTime LastPingTime;
-
-    /** Holds a pointer to the message bus. */
-    TWeakPtr<IMessageBus> MessageBusPtr;
-
-    /** Holds the messaging endpoint. */
-    TSharedPtr<FMessageEndpoint> MessageEndpoint;
-
 	TMap<FGuid, TSharedPtr<FChaosVDSessionInfo>> ActiveSessionsByInstanceId;
-	TMap<FGuid, FChaosVDRecordingStatusMessage> PendingRecordingStatusMessages;
-	TMap<FGuid, FChaosVDTraceDetails> PendingRecordingConnectionDetailsMessages;
 
-	FSimpleMulticastDelegate SessionsUpdatedDelegate;
 	FChaosVDSessionDiscoverDelegate SessionDiscoveredDelegate;
 	FChaosVDSessionDiscoverDelegate SessionExpiredDelegate;
 
 	FChaosVDRecordingStateChangeDelegate RecordingStartedDelegate;
 	FChaosVDRecordingStateChangeDelegate RecordingStoppedDelegate;
-	
-	FTSTicker::FDelegateHandle TickHandle;
-
-	FMessagingInitialized MessagingInitializedDelegate;
+};
 
-	TArray<const UScriptStruct*> SupportedMessageTypes;
+PRAGMA_ENABLE_DEPRECATION_WARNINGS
 
-	bool bInitialized = false;
-};
+#undef UE_API
\ No newline at end of file
diff --git a/Engine/Source/Runtime/Experimental/ChaosSolverEngine/Public/Chaos/ChaosVDTraceRelayTransport.h b/Engine/Source/Runtime/Experimental/ChaosSolverEngine/Public/Chaos/ChaosVDTraceRelayTransport.h
--- a/Engine/Source/Runtime/Experimental/ChaosSolverEngine/Public/Chaos/ChaosVDTraceRelayTransport.h
+++ b/Engine/Source/Runtime/Experimental/ChaosSolverEngine/Public/Chaos/ChaosVDTraceRelayTransport.h
@@ -3,27 +3,17 @@
 #pragma once
 
 #include "Containers/Array.h"
-#include "ChaosVDTraceRelayTransport.generated.h"
+#include "Containers/UnrealString.h"
+#include "Misc/CoreMiscdefines.h"
 
 namespace Chaos::VD
 {
 
 /** Structure containing the details necessary to connect to a relay instance */
-USTRUCT()
-struct FRelayConnectionInfo
+struct UE_DEPRECATED(5.8, "Use UE::TraceBasedDebuggers::FRelayConnectionInfo instead") FRelayConnectionInfo
 {
-	GENERATED_BODY()
-
-	/** Relay IP address */
-	UPROPERTY()
 	FString Address;
-
-	/** Relay port */
-	UPROPERTY()
 	uint16 Port = 0;
-
-	/** Encoded SSL certificate */
-	UPROPERTY()
 	TArray<uint8> CertificateAuthority;
 };
 }
@@ -31,7 +21,6 @@ struct FRelayConnectionInfo
 #if WITH_CHAOS_VISUAL_DEBUGGER
 #include "Delegates/Delegate.h"
 
-struct FChaosVDRelayTraceDataMessage;
 struct FGuid;
 
 namespace Chaos::VD
@@ -39,7 +28,7 @@ namespace Chaos::VD
 class FRelayTraceWriter;
 
 /** Possible connection attempt results */
-enum EConnectionAttemptResult
+enum UE_DEPRECATED(5.8, "Use UE::TraceBasedDebuggers::EConnectionAttemptResult instead") EConnectionAttemptResult
 {
 	NotStarted,
 	InProgress,
@@ -47,10 +36,11 @@ enum EConnectionAttemptResult
 	Failed
 };
 
+PRAGMA_DISABLE_DEPRECATION_WARNINGS
 /**
  * Interface for any CVD trace relay transport implementation
  */
-class ITraceDataRelayTransport
+class UE_DEPRECATED(5.8, "Use UE::TraceBasedDebuggers::IDataRelayTransport instead") ITraceDataRelayTransport
 {
 public:
 	virtual ~ITraceDataRelayTransport() = default;
@@ -71,4 +61,5 @@ public:
 
 }
 
-#endif // WITH_CHAOS_VISUAL_DEBUGGER
+PRAGMA_ENABLE_DEPRECATION_WARNINGS
+#endif // WITH_CHAOS_VISUAL_DEBUGGER
\ No newline at end of file
diff --git a/Engine/Source/Runtime/Experimental/ChaosVisualDebugger/ChaosVDRuntime.Build.cs b/Engine/Source/Runtime/Experimental/ChaosVisualDebugger/ChaosVDRuntime.Build.cs
--- a/Engine/Source/Runtime/Experimental/ChaosVisualDebugger/ChaosVDRuntime.Build.cs
+++ b/Engine/Source/Runtime/Experimental/ChaosVisualDebugger/ChaosVDRuntime.Build.cs
@@ -13,7 +13,8 @@ namespace UnrealBuildTool.Rules
 				new string[] {
 					"Core",
 					"CoreUObject",
-					"TraceLog"
+					"TraceLog",
+					"TraceBasedDebuggers",
 				}
 			);
 
diff --git a/Engine/Source/Runtime/Experimental/ChaosVisualDebugger/Private/ChaosVDRuntimeModule.cpp b/Engine/Source/Runtime/Experimental/ChaosVisualDebugger/Private/ChaosVDRuntimeModule.cpp
--- a/Engine/Source/Runtime/Experimental/ChaosVisualDebugger/Private/ChaosVDRuntimeModule.cpp
+++ b/Engine/Source/Runtime/Experimental/ChaosVisualDebugger/Private/ChaosVDRuntimeModule.cpp
@@ -6,26 +6,12 @@
 
 #if WITH_CHAOS_VISUAL_DEBUGGER
 
-#include "Containers/Array.h"
 #include "HAL/IConsoleManager.h"
-#include "HAL/PlatformProcess.h"
-#include "Internationalization/Internationalization.h"
-#include "Misc/CommandLine.h"
-#include "Misc/FileHelper.h"
-#include "Misc/MessageDialog.h"
-#include "Misc/Paths.h"
 
 #if WITH_EDITOR
 #include "Misc/ScopedSlowTask.h"
 #endif
 
-#include "ChaosVDRecordingDetails.h"
-#include "Math/GuardedInt.h"
-#include "Misc/App.h"
-#include "Modules/ModuleManager.h"
-#include "ProfilingDebugging/TraceAuxiliary.h"
-#include "Trace/RelayTraceWriter.h"
-
 IMPLEMENT_MODULE(FChaosVDRuntimeModule, ChaosVDRuntime);
 
 DEFINE_LOG_CATEGORY(LogChaosVDRuntime);
@@ -61,28 +47,22 @@ static FAutoConsoleVariable CVarChaosVDGTimeBetweenFullCaptures(
 	10,
 	TEXT("Time interval in seconds after which a full capture (not only delta changes) should be recorded"));
 
-static FAutoConsoleVariable CVarChaosVDMaxTimeToWaitForDisconnect(
-	TEXT("p.Chaos.VD.MaxTimeToWaitForDisconnectSeconds"),
-	5.0f,
-	TEXT("Max time to wait after attempting to stop an active trace session. After that time has passed if we are still connected, CVD will continue and eventually error out."));
-
-static FAutoConsoleVariable CVarChaosVDMaxBytesPerRelayChunk(
-	TEXT("p.Chaos.VD.MaxBytesPerRelayChunk"),
-	65536,
-	TEXT("How many bytes per relay data message we should send"));
-
-static FAutoConsoleVariable CVarChaosVDMaxPendingBytesInRelayWriter(
-	TEXT("p.Chaos.VD.MaxPendingBytesInRelayWriter"),
-	536870912,
-	TEXT("How many bytes the relay writer can hold before overflowing, triggering a close. Defaults to ~512MB"));
-
-FChaosVDRecordingStateChangedDelegate FChaosVDRuntimeModule::RecordingStartedDelegate = FChaosVDRecordingStateChangedDelegate();
-FChaosVDRecordingStateChangedDelegate FChaosVDRuntimeModule::PostRecordingStartedDelegate = FChaosVDRecordingStateChangedDelegate();
-FChaosVDRecordingStateChangedDelegate FChaosVDRuntimeModule::RecordingStopDelegate = FChaosVDRecordingStateChangedDelegate();
-FChaosVDRecordingStartFailedDelegate FChaosVDRuntimeModule::RecordingStartFailedDelegate = FChaosVDRecordingStartFailedDelegate();
 FChaosVDCaptureRequestDelegate FChaosVDRuntimeModule::PerformFullCaptureDelegate = FChaosVDCaptureRequestDelegate();
+UE::TraceBasedDebuggers::FRecordingStateChangedDelegate FChaosVDRuntimeModule::RecordingStartedDelegate = UE::TraceBasedDebuggers::FRecordingStateChangedDelegate();
+UE::TraceBasedDebuggers::FRecordingStateChangedDelegate FChaosVDRuntimeModule::PostRecordingStartedDelegate = UE::TraceBasedDebuggers::FRecordingStateChangedDelegate();
+UE::TraceBasedDebuggers::FRecordingStateChangedDelegate FChaosVDRuntimeModule::RecordingStopDelegate = UE::TraceBasedDebuggers::FRecordingStateChangedDelegate();
+UE::TraceBasedDebuggers::FRecordingStartFailedDelegate FChaosVDRuntimeModule::RecordingStartFailedDelegate = UE::TraceBasedDebuggers::FRecordingStartFailedDelegate();
 FTransactionallySafeRWLock FChaosVDRuntimeModule::DelegatesRWLock = FTransactionallySafeRWLock();
 
+
+FChaosVDRuntimeModule::FChaosVDRuntimeModule()
+	: FRuntimeModule(LogChaosVDRuntime, Chaos::VD::DebuggerGuid)
+{
+	DebuggerName = TEXT("ChaosVD");
+	AutoRecordCommand = TEXT("StartCVDRecording");
+	HostAddressCommand = TEXT("CVDHost");
+}
+
 FChaosVDRuntimeModule& FChaosVDRuntimeModule::Get()
 {
 	return FModuleManager::Get().LoadModuleChecked<FChaosVDRuntimeModule>(TEXT("ChaosVDRuntime"));
@@ -90,57 +70,20 @@ FChaosVDRuntimeModule& FChaosVDRuntimeModule::Get()
 
 void FChaosVDRuntimeModule::StartupModule()
 {
-	if (FParse::Param(FCommandLine::Get(), TEXT("StartCVDRecording")))
-	{
-		TArray<FString, TInlineAllocator<1>> CVDOptions;
-
-		{
-			FString CVDHostAddress;
-			if (FParse::Value(FCommandLine::Get(), TEXT("CVDHost="), CVDHostAddress))
-			{
-				CVDOptions.Emplace(TEXT("Server"));
-				CVDOptions.Emplace(MoveTemp(CVDHostAddress));
-			}
-		}
-        
-        StartRecording(CVDOptions);
-	}
-	else
+	if (!TryRecordOnStartup())
 	{
-		
 #if UE_TRACE_ENABLED || UE_TRACE_MINIMAL_ENABLED
 		UE::Trace::ToggleChannel(TEXT("ChaosVDChannel"), false);
 #endif
-
 	}
 }
 
 void FChaosVDRuntimeModule::ShutdownModule()
 {
-	if (bIsRecording)
+	if (IsRecording())
 	{
 		StopRecording();
 	}
-
-	FTraceAuxiliary::OnTraceStopped.RemoveAll(this);
-}
-
-void FChaosVDRuntimeModule::StartRecording(TConstArrayView<FString> Args)
-{
-	FChaosVDStartRecordingCommandMessage StartRecordingCommand;
-	if (Args.Num() == 0 || Args[0] == TEXT("File"))
-	{
-		StartRecordingCommand.RecordingMode = EChaosVDRecordingMode::File;
-		StartRecordingCommand.TransportMode = EChaosVDTransportMode::FileSystem;
-	}
-	else if(Args[0] == TEXT("Server"))
-	{
-		StartRecordingCommand.RecordingMode = EChaosVDRecordingMode::Live;
-		StartRecordingCommand.TransportMode = EChaosVDTransportMode::TraceServer;
-		StartRecordingCommand.Target = Args.IsValidIndex(1) ? Args[1] : TEXT("127.0.0.1");
-	}
-
-	StartRecording(StartRecordingCommand);
 }
 
 int32 FChaosVDRuntimeModule::GenerateUniqueID()
@@ -154,505 +97,83 @@ int32 FChaosVDRuntimeModule::GenerateUniqueID()
 	return NewID;
 }
 
-int64 FChaosVDRuntimeModule::GetBufferedDataBytesSize() const
-{
-	FGuardedInt64 RelaySendBufferSize(0);
-
-	if (RelayWriter)
-	{
-		RelaySendBufferSize += RelayWriter->GetQueuedBytesNum();
-	}
-
-	if (ExternalRelayBufferSizeProviderDelegate.IsBound())
-	{
-		RelaySendBufferSize += ExternalRelayBufferSizeProviderDelegate.Execute();
-	}
-
-	constexpr int64 DefaultValue = std::numeric_limits<int64>::max();
-	return RelaySendBufferSize.Get(DefaultValue); 
-}
-
 FString FChaosVDRuntimeModule::GetLastRecordingFileNamePath() const
 {
 	return FString();
 }
 
-FChaosVDTraceDetails FChaosVDRuntimeModule::GetCurrentTraceSessionDetails() const
-{
-	FChaosVDTraceDetails Details;
-	if (ExternalTraceStatusDelegate.IsBound())
-	{
-		return ExternalTraceStatusDelegate.Execute();
-	}
-	
-	if (FTraceAuxiliary::IsConnected(Details.SessionGuid, Details.TraceGuid))
-	{
-		Details.TraceTarget = FTraceAuxiliary::GetTraceDestinationString();
-		Details.TransportMode = CurrentTransportMode;
-		Details.MarkAsValid();
-		return Details;
-	}
-
-	return Details;
-}
-
-void FChaosVDRuntimeModule::StopTrace()
-{
-	bRequestedStop = true;
-	FTraceAuxiliary::Stop();
-}
-
-void FChaosVDRuntimeModule::GenerateRecordingFileName(FString& OutFileName)
-{
-	FStringFormatOrderedArguments NameArgs { FString(FApp::GetProjectName()), LexToString(FApp::GetBuildTargetType()), FDateTime::Now().ToString(TEXT("%Y%m%d_%H%M%S%s")) };
-	
-	OutFileName = FString::Format(TEXT("ChaosVD-{0}-{1}-{2}.utrace"), NameArgs);
-}
-
-void FChaosVDRuntimeModule::RegisterExternalRelayExecutor(const FChaosVDRelayTraceDataDelegate& InExternalRelayExecutor)
-{
-	if (!ensureMsgf(!ExternalRelayExecutorDelegate.IsBound(), TEXT("An external relay executor is already registered!")))
-	{
-		return;
-	}
-
-	ExternalRelayExecutorDelegate = InExternalRelayExecutor;
-}
-
-void FChaosVDRuntimeModule::RegisterExternalRelayOverflowHandler(const FSimpleDelegate& InExternalRelayOverflowHandler)
-{
-	if (!ensureMsgf(!ExternalRelayOverflowHandlerDelegate.IsBound(), TEXT("An external relay overflow handler is already registered!")))
-	{
-		return;
-	}
-
-	ExternalRelayOverflowHandlerDelegate = InExternalRelayOverflowHandler;
-}
-
-void FChaosVDRuntimeModule::UnregisterExternalRelayOverflowHandler()
-{
-	ExternalRelayOverflowHandlerDelegate.Unbind();
-}
-
-void FChaosVDRuntimeModule::UnregisterCurrentExternalRelayExecutor()
-{
-	ExternalRelayExecutorDelegate.Unbind();
-}
-
-void FChaosVDRuntimeModule::RegisterExternalTraceStatusProvider(const FChaosVDExternalTraceStatusDelegate& InExternalTraceStatusCallback)
-{
-	ExternalTraceStatusDelegate = InExternalTraceStatusCallback;
-}
-
-void FChaosVDRuntimeModule::UnregisterExternalTraceStatusProvider()
-{
-	ExternalTraceStatusDelegate.Unbind();
-}
-
-void FChaosVDRuntimeModule::RegisterExternalRelayBufferSizeProvider(const FChaosVDExternalRelayBufferQueueSizeDelegate& InExternalRelayBufferSizeCallback)
+void FChaosVDRuntimeModule::OnEnableRequiredTraceChannelsInternal()
 {
-	if (!ensureMsgf(!ExternalRelayBufferSizeProviderDelegate.IsBound(), TEXT("An external relay InExternalRelayBufferSizeCallback is already registered!")))
-	{
-		return;
-	}
-
-	ExternalRelayBufferSizeProviderDelegate = InExternalRelayBufferSizeCallback;
-}
-
-void FChaosVDRuntimeModule::UnregisterExternalRelayBufferSizeProvider()
-{
-	ExternalRelayBufferSizeProviderDelegate.Unbind();
-}
-
-void FChaosVDRuntimeModule::HandleRelayOverflow()
-{
-	if (ExternalRelayOverflowHandlerDelegate.IsBound())
-	{
-		UE_LOG(LogChaosVDRuntime, Error, TEXT("[%hs] Relay writer has overflowed. Executing external overflow handler... "), __func__);
-		ExternalRelayOverflowHandlerDelegate.Execute();
-		return;
-	}
-
-	UE_LOG(LogChaosVDRuntime, Error, TEXT("[%hs] Relay writer has overflowed. Stopping Recording... "), __func__);
-	StopRecording();
-}
-
-void FChaosVDRuntimeModule::RelayTraceData()
-{
-	if (ensure(ExternalRelayExecutorDelegate.IsBound() && RelayWriter.IsValid()))
-	{
-		ExternalRelayExecutorDelegate.Execute(*RelayWriter);
-	}
-}
-
-void FChaosVDRuntimeModule::SetupRelayDataPumpDelegates()
-{
-	if (!ensure(RelayWriter))
-	{
-		return;
-	}
-
-	RelayWriter->OnNewDataAvailable().BindRaw(this, &FChaosVDRuntimeModule::RelayTraceData);
-	RelayWriter->OnDataOverflow().BindRaw(this, &FChaosVDRuntimeModule::HandleRelayOverflow);
-}
-
-void FChaosVDRuntimeModule::ClearRelayDataPumpDelegates()
-{
-	if (RelayWriter)
-	{
-		RelayWriter->OnNewDataAvailable().Unbind();
-		RelayWriter->OnDataOverflow().Unbind();
-	}
-}
-
-bool FChaosVDRuntimeModule::RequestFullCapture(float DeltaTime)
-{
-	// Full capture intervals are clamped to be no lower than 1 sec
-	UE::TReadScopeLock ReadLock (DelegatesRWLock);
-	PerformFullCaptureDelegate.Broadcast(EChaosVDFullCaptureFlags::Particles);
-	return true;
-}
+#if UE_TRACE_ENABLED || UE_TRACE_MINIMAL_ENABLED
+	UE::Trace::ToggleChannel(TEXT("ChaosVDChannel"), true);
+#endif
 
-bool FChaosVDRuntimeModule::RecordingTimerTick(float DeltaTime)
-{
-	if (bIsRecording)
-	{
-		AccumulatedRecordingTime += DeltaTime;
-	}
-	
-	return true;
+#if UE_TRACE_ENABLED
+	UE::Trace::ToggleChannel(TEXT("Log"), true);
+#endif
 }
 
-void FChaosVDRuntimeModule::HandleTraceConnectionEstablished()
+void FChaosVDRuntimeModule::OnRecordingStartedInternal()
 {
-	UE_LOG(LogChaosVDRuntime, Log, TEXT("Trace connection established."));
+	constexpr int32 MinAllowedTimeInSecondsBetweenCaptures = 1;
+	const int32 ConfiguredTimeBetweenCaptures = CVarChaosVDGTimeBetweenFullCaptures->GetInt();
 
-	if (OnConnectionDelegateHandle.IsValid())
-	{
-		FTraceAuxiliary::OnConnection.RemoveAll(this);
-		OnConnectionDelegateHandle.Reset();
-	}
+	ensureAlwaysMsgf(ConfiguredTimeBetweenCaptures > MinAllowedTimeInSecondsBetweenCaptures,
+		TEXT("The minimum allowed time interval between full captures is [%d] seconds, but [%d] seconds were configured. Clamping to [%d] seconds"),
+		MinAllowedTimeInSecondsBetweenCaptures, ConfiguredTimeBetweenCaptures, MinAllowedTimeInSecondsBetweenCaptures);
 
-	if (!ensureAlwaysMsgf(bIsRecording, TEXT("Received a trace connection established callback but no CVD trace is active. This should not happen.")))
-	{
-		return;
-	}
+	FullCaptureRequesterHandle = FTSTicker::GetCoreTicker().AddTicker(FTickerDelegate::CreateRaw(this, &FChaosVDRuntimeModule::RequestFullCapture),
+		FMath::Clamp(ConfiguredTimeBetweenCaptures, MinAllowedTimeInSecondsBetweenCaptures, TNumericLimits<int32>::Max()));
 
 	{
 		UE::TReadScopeLock ReadLock(DelegatesRWLock);
 		RecordingStartedDelegate.Broadcast();
 	}
 
-	TraceConnectionDetailsUpdatedDelegate.Broadcast();
-	
-	constexpr int32 MinAllowedTimeInSecondsBetweenCaptures = 1;
-	int32 ConfiguredTimeBetweenCaptures = CVarChaosVDGTimeBetweenFullCaptures->GetInt();
-
-	ensureAlwaysMsgf(ConfiguredTimeBetweenCaptures > MinAllowedTimeInSecondsBetweenCaptures,
-	                 TEXT("The minimum allowed time interval between full captures is [%d] seconds, but [%d] seconds were configured. Clamping to [%d] seconds"),
-	                 MinAllowedTimeInSecondsBetweenCaptures, ConfiguredTimeBetweenCaptures, MinAllowedTimeInSecondsBetweenCaptures);
-
-	FullCaptureRequesterHandle = FTSTicker::GetCoreTicker().AddTicker(FTickerDelegate::CreateRaw(this, &FChaosVDRuntimeModule::RequestFullCapture),
-	                                                                  FMath::Clamp(ConfiguredTimeBetweenCaptures, MinAllowedTimeInSecondsBetweenCaptures, TNumericLimits<int32>::Max()));
-
-	RecordingTimerHandle = FTSTicker::GetCoreTicker().AddTicker(FTickerDelegate::CreateRaw(this, &FChaosVDRuntimeModule::RecordingTimerTick));
-
 	{
 		UE::TReadScopeLock ReadLock(DelegatesRWLock);
 		PostRecordingStartedDelegate.Broadcast();
 	}
 }
 
-
-void FChaosVDRuntimeModule::StartRecording(const FChaosVDStartRecordingCommandMessage& InRecordingStartCommand)
+void FChaosVDRuntimeModule::OnRecordingStartFailedInternal(const FText& FailureReason)
 {
-	if (!ensureAlwaysMsgf(!bIsRecording, TEXT("Received a recording start command while a recording is already active")))
-	{
-		UE_LOG(LogChaosVDRuntime, Log, TEXT("[%hs] There is an active CVD recording. Attempting to stop it..."), __func__)
-		StopRecording();
-	}
-
-	// Start with a generic Failure reason
-	FText FailureReason = LOCTEXT("SeeLogsForErrorDetailsText","Please see the logs for more details...");
-
-#if UE_TRACE_ENABLED || UE_TRACE_MINIMAL_ENABLED
-
-	// Other tools could bee using trace
-	// This is aggressive but until Trace supports multi-sessions, just take over.
-	if (FTraceAuxiliary::IsConnected())
 	{
-		UE_LOG(LogChaosVDRuntime, Log, TEXT("[%s] There is an active trace session. attempting to disconnect..."), ANSI_TO_TCHAR(__FUNCTION__));
-
-		//TODO: We should make the wait async like we do whe we attempt to connect to a live session
-		if (FTraceAuxiliary::Stop() && WaitForTraceSessionDisconnect())
-		{
-			UE_LOG(LogChaosVDRuntime, Log, TEXT("[%s] Successful disconnect attempt!."), ANSI_TO_TCHAR(__FUNCTION__));
-		}
-		else
-		{
-			FailureReason = LOCTEXT("FailedToStopActiveRecordingErrorMessage", "Failed to Stop active Trace Session.");
-		}
-	}
-
-	SaveAndDisabledCurrentEnabledTraceChannels();
-
-	EnableRequiredTraceChannels();
-
-	FTraceAuxiliary::FOptions TracingOptions;
-	TracingOptions.bExcludeTail = true;
-
-	if (InRecordingStartCommand.RecordingMode == EChaosVDRecordingMode::File)
-	{
-
-		if (InRecordingStartCommand.TransportMode == EChaosVDTransportMode::FileSystem)
-		{
-			FString FileTarget;
-			GenerateRecordingFileName(FileTarget);
-			UE_LOG(LogChaosVDRuntime, Log, TEXT("[%hs] Generated trace file name [%s]"), __func__, *FileTarget);
-
-			bIsRecording = FTraceAuxiliary::Start(FTraceAuxiliary::EConnectionType::File, *FileTarget, nullptr, &TracingOptions);
-		}
-		else
-		{
-			UE_LOG(LogChaosVDRuntime, Error, TEXT("[%hs] Unsupported Transport mode [%s]"), __func__, *UEnum::GetValueAsString(InRecordingStartCommand.TransportMode))
-		}
-	}
-	else if(InRecordingStartCommand.RecordingMode == EChaosVDRecordingMode::Live)
-	{
-		switch (InRecordingStartCommand.TransportMode)
-		{
-		case EChaosVDTransportMode::Direct:
-		case EChaosVDTransportMode::TraceServer:
-			bIsRecording = FTraceAuxiliary::Start(FTraceAuxiliary::EConnectionType::Network, *InRecordingStartCommand.Target, nullptr, &TracingOptions);
-			break;
-		case EChaosVDTransportMode::Relay:
-			{
-				if (!ensureMsgf(ExternalRelayExecutorDelegate.IsBound(), TEXT("Cannot start a trace in relay mode without a Relay Executor")))
-				{
-					UE_LOG(LogChaosVDRuntime, Error, TEXT("[%hs] Cannot start a trace in relay mode without a Relay Executor"), __func__);
-					break;
-				}
-				
-				if (RelayWriter && !RelayWriter->IsClosed())
-				{
-					UE_LOG(LogChaosVDRuntime, Warning, TEXT("[%hs] Attempting to relay trace, but there is a relay writer still open"), __func__);
-					RelayWriter->Close();
-				}
-
-				RelayWriter = MakeUnique<Chaos::VD::FRelayTraceWriter>();
-
-				constexpr int32 MinBunchSize = 4096;
-				RelayWriter->SetMaxBytesPerBunch(FMath::Clamp(CVarChaosVDMaxBytesPerRelayChunk->GetInt(), MinBunchSize, std::numeric_limits<int32>::max()));
-
-				constexpr int32 MinBufferSizeLimit = 2097152;
-				RelayWriter->SetMaxAllowedPendingBytes(FMath::Clamp(CVarChaosVDMaxPendingBytesInRelayWriter->GetInt(), MinBufferSizeLimit, std::numeric_limits<int32>::max()));
-				
-				bIsRecording = FTraceAuxiliary::Relay(reinterpret_cast<UPTRINT>(RelayWriter.Get()), Chaos::VD::FRelayTraceWriter::WriteHelper, Chaos::VD::FRelayTraceWriter::CloseHelper, nullptr, &TracingOptions);
-
-				if (bIsRecording)
-				{
-					SetupRelayDataPumpDelegates();
-				}
-
-				break;
-			}
-		case EChaosVDTransportMode::Invalid:
-		case EChaosVDTransportMode::FileSystem:
-		UE_LOG(LogChaosVDRuntime, Error, TEXT("[%hs] Unsupported Transport mode [%s]"), __func__, *UEnum::GetValueAsString(InRecordingStartCommand.TransportMode))
-			break;
-		}
-	}
-	else
-	{
-		FailureReason = LOCTEXT("WrongCommandArgumentsError", "The start recording command was called with invalid arguments");
-	}
-#endif
-	
-	AccumulatedRecordingTime = 0.0f;
-
-	if (ensureMsgf(bIsRecording, TEXT("Failed to start CVD recording | Reason [%s]"), *FailureReason.ToString()))
-	{
-		CurrentTransportMode = InRecordingStartCommand.TransportMode;
-
-		if (ensureMsgf(!OnConnectionDelegateHandle.IsValid(), TEXT("Starting a trace while we are waiting for a pending connection")))
-		{
-			OnConnectionDelegateHandle = FTraceAuxiliary::OnConnection.AddRaw(this, &FChaosVDRuntimeModule::HandleTraceConnectionEstablished);
-		}
-
-		// Start Listening for Trace Stopped events, in case Trace is stopped outside our control so we can gracefully stop CVD recording and log a warning
-		if (ensureMsgf(!OnTraceStoppedDelegateHandle.IsValid(), TEXT("Starting a trace while we are waiting for a trace session to stop.")))
-		{
-			OnTraceStoppedDelegateHandle = FTraceAuxiliary::OnTraceStopped.AddRaw(this, &FChaosVDRuntimeModule::HandleTraceStopRequest);
-		}
-	}
-	else
-	{
-		UE_LOG(LogChaosVDRuntime, Error, TEXT("[%s] Failed to start CVD recording | Reason: [%s]"), ANSI_TO_TCHAR(__FUNCTION__), *FailureReason.ToString());
-
-#if WITH_EDITOR
-		FMessageDialog::Open(EAppMsgType::Ok, FText::FormatOrdered(LOCTEXT("StartRecordingFailedMessage", "Failed to start CVD recording. \n\n{0}"), FailureReason));
-#endif
-
-		{
-			UE::TReadScopeLock ReadLock(DelegatesRWLock);
-			RecordingStartFailedDelegate.Broadcast(FailureReason);
-		}
-
-		RelayWriter = nullptr;
+		UE::TReadScopeLock ReadLock(DelegatesRWLock);
+		RecordingStartFailedDelegate.Broadcast(FailureReason);
 	}
 }
 
-void FChaosVDRuntimeModule::StopRecording()
+void FChaosVDRuntimeModule::OnRecordingStoppedInternal()
 {
-	FTraceAuxiliary::OnConnection.RemoveAll(this);
-	OnConnectionDelegateHandle = FDelegateHandle();
-
-	FTraceAuxiliary::OnTraceStopped.RemoveAll(this);
-	OnTraceStoppedDelegateHandle = FDelegateHandle();
-
-	CurrentTransportMode = EChaosVDTransportMode::Invalid;
-
-	if (!bIsRecording)
 	{
-		UE_LOG(LogChaosVDRuntime, Warning, TEXT("[%hs] Attempted to stop recorded when there is no CVD recording active."), __func__);
-		return;
+		UE::TReadScopeLock ReadLock(DelegatesRWLock);
+		RecordingStopDelegate.Broadcast();
 	}
 
-#if UE_TRACE_ENABLED || UE_TRACE_MINIMAL_ENABLED
-	RestoreTraceChannelsToPreRecordingState();
-
-	StopTrace();
-#endif
-
 	if (FullCaptureRequesterHandle.IsValid())
 	{
 		FTSTicker::RemoveTicker(FullCaptureRequesterHandle);
 
 		FullCaptureRequesterHandle.Reset();
 	}
-	
-	if (RecordingTimerHandle.IsValid())
-	{
-		FTSTicker::RemoveTicker(RecordingTimerHandle);
-		RecordingTimerHandle.Reset();
-	}
-
-	ClearRelayDataPumpDelegates();
-
-	if (RelayWriter)
-	{
-		RelayWriter->Close();
-	}
-	
-	bIsRecording = false;
-	AccumulatedRecordingTime = 0.0f;
-
-	{
-		UE::TReadScopeLock ReadLock(DelegatesRWLock);
-		RecordingStopDelegate.Broadcast();
-	}
 }
 
-void FChaosVDRuntimeModule::HandleTraceStopRequest(FTraceAuxiliary::EConnectionType TraceType, const FString& TraceDestination)
-{
-	if (bIsRecording)
-	{
-		if (!ensure(bRequestedStop))
-		{
-			UE_LOG(LogChaosVDRuntime, Warning, TEXT("Trace Recording has been stopped unexpectedly"));
-
-#if WITH_EDITOR
-			FMessageDialog::Open(EAppMsgType::Ok, LOCTEXT("UnexpectedStopMessage", "Trace recording has been stopped unexpectedly. CVD cannot continue with the recording session... "));
-#endif
-		}
+PRAGMA_DISABLE_DEPRECATION_WARNINGS
 
-		StopRecording();
-	}
-
-	bRequestedStop = false;
-}
-
-bool FChaosVDRuntimeModule::WaitForTraceSessionDisconnect()
-{
-	float MaxWaitTime = CVarChaosVDMaxTimeToWaitForDisconnect->GetFloat();
-	float CurrentWaitTime = 0.0f;
-
-#if WITH_EDITOR
-	FScopedSlowTask DisconnectAttemptSlowTask(MaxWaitTime, LOCTEXT("DisconnectAttemptMessage", " Active Trace Session detected, attempting to disconnect ..."));
-
-	constexpr bool bShowCancelButton = false;
-	constexpr bool bAllowInPIE = true;
-	DisconnectAttemptSlowTask.MakeDialog(bShowCancelButton, bAllowInPIE);
-#endif
-
-	while (CurrentWaitTime < MaxWaitTime)
-	{
-		constexpr float WaitInterval = 0.1f;
-		FPlatformProcess::Sleep(0.1f);
-
-		if (!FTraceAuxiliary::IsConnected())
-		{
-			return true;
-		}
-
-		// We don't need to be precise for this, we can just accumulate the wait
-		CurrentWaitTime += WaitInterval;
-
-#if WITH_EDITOR
-		DisconnectAttemptSlowTask.EnterProgressFrame(CurrentWaitTime);
-#endif
-	}
-
-	return FTraceAuxiliary::IsConnected();
-}
-
-void FChaosVDRuntimeModule::SaveAndDisabledCurrentEnabledTraceChannels()
+FChaosVDTraceDetails FChaosVDRuntimeModule::GetCurrentTraceSessionDetails() const
 {
-	// Until we support allowing other channels, indicate in the logs that we are disabling everything else
-	UE_LOG(LogChaosVDRuntime, Log, TEXT("[%s] Disabling additional trace channels..."), ANSI_TO_TCHAR(__FUNCTION__));
-
-#if UE_TRACE_ENABLED || UE_TRACE_MINIMAL_ENABLED
-	OriginalTraceChannelsState.Reset();
-
-	// Disable any enabled additional channel
-	UE::Trace::EnumerateChannels([](const ANSICHAR* ChannelName, bool bEnabled, void* SavedTraceChannelsPtr)
-	{
-		TMap<FString, bool>* SavedTraceChannels = static_cast<TMap<FString, bool>*>(SavedTraceChannelsPtr);
-		FString ChannelNameFString(ChannelName);
-		SavedTraceChannels->Add(ChannelNameFString, bEnabled);
-		if (bEnabled)
-		{
-			UE::Trace::ToggleChannel(ChannelNameFString.GetCharArray().GetData(), false);
-		}
-	}
-	, &OriginalTraceChannelsState);
-#endif
+	return {};
 }
+PRAGMA_ENABLE_DEPRECATION_WARNINGS
 
-void FChaosVDRuntimeModule::RestoreTraceChannelsToPreRecordingState()
-{
-#if UE_TRACE_ENABLED || UE_TRACE_MINIMAL_ENABLED
-	UE_LOG(LogChaosVDRuntime, Log, TEXT("[%s] Restoring trace channels state..."), ANSI_TO_TCHAR(__FUNCTION__));
-
-	for (const TPair<FString, bool>& ChannelWithState : OriginalTraceChannelsState)
-	{
-		UE::Trace::ToggleChannel(GetData(ChannelWithState.Key), ChannelWithState.Value); 
-	}
-	
-	OriginalTraceChannelsState.Reset();
-#endif
-}
 
-void FChaosVDRuntimeModule::EnableRequiredTraceChannels()
+bool FChaosVDRuntimeModule::RequestFullCapture(float DeltaTime)
 {
-#if UE_TRACE_ENABLED || UE_TRACE_MINIMAL_ENABLED
-	UE::Trace::ToggleChannel(TEXT("ChaosVDChannel"), true); 
-	UE::Trace::ToggleChannel(TEXT("Frame"), true);
-#endif
-
-#if UE_TRACE_ENABLED
-	UE::Trace::ToggleChannel(TEXT("Log"), true);
-#endif
+	// Full capture intervals are clamped to be no lower than 1 sec
+	UE::TReadScopeLock ReadLock(DelegatesRWLock);
+	PerformFullCaptureDelegate.Broadcast(EChaosVDFullCaptureFlags::Particles);
+	return true;
 }
 
 #undef LOCTEXT_NAMESPACE
diff --git a/Engine/Source/Runtime/Experimental/ChaosVisualDebugger/Private/Trace/RelayTraceWriter.cpp b/Engine/Source/Runtime/Experimental/ChaosVisualDebugger/Private/Trace/RelayTraceWriter.cpp
deleted file mode 100644
diff --git a/Engine/Source/Runtime/Experimental/ChaosVisualDebugger/Public/ChaosVDRecordingDetails.h b/Engine/Source/Runtime/Experimental/ChaosVisualDebugger/Public/ChaosVDRecordingDetails.h
--- a/Engine/Source/Runtime/Experimental/ChaosVisualDebugger/Public/ChaosVDRecordingDetails.h
+++ b/Engine/Source/Runtime/Experimental/ChaosVisualDebugger/Public/ChaosVDRecordingDetails.h
@@ -1,107 +1,64 @@
 // Copyright Epic Games, Inc. All Rights Reserved.
 
 #pragma once
+#include "RemoteSessionsManager.h"
 
+#if UE_ENABLE_INCLUDE_ORDER_DEPRECATED_IN_5_8
 #include "Containers/UnrealString.h"
+#endif // UE_ENABLE_INCLUDE_ORDER_DEPRECATED_IN_5_8
+
 #include "Misc/Guid.h"
+#include "TraceDataRelayTransport.h"
 #include "ChaosVDRecordingDetails.generated.h"
 
-UENUM()
-enum class EChaosVDRecordingMode : uint8
+namespace Chaos::VD
+{
+/** Debugger Guid used to identify the debugger type when sending command and update messages */
+static constexpr FGuid DebuggerGuid = FGuid(0x34CFFA28, 0x72434D08, 0xB07AD3FE, 0xE2D96272);
+}
+
+
+/** CVD specific message to send a start recording command with optional list of data channels to force to be enabled. */
+USTRUCT()
+struct FChaosVDStartRecordingCommandMessage : public UE::TraceBasedDebuggers::FStartRecordingCommandMessage
+{
+	GENERATED_BODY()
+
+	UPROPERTY()
+	TArray<FString> DataChannelsEnabledOverrideList;
+};
+
+PRAGMA_DISABLE_DEPRECATION_WARNINGS
+
+enum class UE_DEPRECATED(5.8, "Use UE::TraceBasedDebuggers::ERecordingMode instead.") EChaosVDRecordingMode : uint8
 {
 	Invalid,
 	Live,
 	File
 };
 
-/** Available transport modes for CVD Trace Data*/
-UENUM()
-enum class EChaosVDTransportMode : uint8
+enum class UE_DEPRECATED(5.8, "Use UE::TraceBasedDebuggers::ETraceTransportMode instead.") EChaosVDTransportMode : uint8
 {
-	Invalid UMETA(Hidden),
+	Invalid,
 	/** The data will be traced directly to a file in the local file system */
-	FileSystem UMETA(Hidden),
+	FileSystem,
 	/** Data will be traced to the selected trace store server */
-	TraceServer UMETA(Hidden),
+	TraceServer,
 	/** Data will be traced directly to the editor, accessible via sockets */
 	Direct,
 	/** Data will be traced via the trace relay system using a custom transport, like normal UE Networking or a direct socket connection */
 	Relay
 };
 
-USTRUCT()
-struct FChaosVDStartRecordingCommandMessage
+struct UE_DEPRECATED(5.8, "Use UE::TraceBasedDebuggers::FTraceConnectionDetails instead.") FChaosVDTraceDetails : public UE::TraceBasedDebuggers::FTraceConnectionDetails
 {
-	GENERATED_BODY()
-
-	UPROPERTY()
-	EChaosVDRecordingMode RecordingMode = EChaosVDRecordingMode::Invalid;
-
-	UPROPERTY()
-	EChaosVDTransportMode TransportMode = EChaosVDTransportMode::Invalid;
-	
-	UPROPERTY()
-	FString Target;
-
-	UPROPERTY()
-	TArray<FString> DataChannelsEnabledOverrideList;
-};
-
-/** Structure containing the info necessary to connect to a trace session and start receiving data */
-USTRUCT()
-struct FChaosVDTraceDetails
-{
-	GENERATED_BODY()
-
 	FChaosVDTraceDetails() = default;
 
-	PRAGMA_DISABLE_DEPRECATION_WARNINGS
-	FChaosVDTraceDetails(const FChaosVDTraceDetails& Other) = default;
-	FChaosVDTraceDetails(FChaosVDTraceDetails&& Other) noexcept = default;
-	FChaosVDTraceDetails& operator=(const FChaosVDTraceDetails& Other) = default;
-	FChaosVDTraceDetails& operator=(FChaosVDTraceDetails&& Other) noexcept = default;
-	PRAGMA_ENABLE_DEPRECATION_WARNINGS
-
-	void MarkAsValid()
-	{
-		bIsValid = true;
-	}
-
-	bool IsValid() const
-	{
-		return bIsValid;
-	}
-
-	/** ID used to find the active trace session in the trace system */
-	UPROPERTY()
-	FGuid TraceGuid;
-
-	/** ID from the CVD remote session from which this connection details come from */
-	UPROPERTY()
-	FGuid SessionGuid;
-
-	/** Address of the trace server or editor in direct trace mode that will receive the trace data */
-	UPROPERTY()
-	FString TraceTarget;
-
-	/*** Port number used for the trace connection, if any*/
-	UPROPERTY()
-	uint16 Port = 0;
-
 	UE_DEPRECATED(5.7, "This property is no longer used and it will removed in the future")
 	bool bIsConnected = false;
 
 	UE_DEPRECATED(5.8, "Recording mode is no longer valid and will not be replaced")
 	EChaosVDRecordingMode Mode = EChaosVDRecordingMode::Invalid;
+};
 
-	/** How the data is being transported to the editor (Trace Server, Custom Relay, Direct, etc.)*/
-	UPROPERTY()
-	EChaosVDTransportMode TransportMode = EChaosVDTransportMode::Invalid;
-
-	/** Encoded SSL certificate for the trace connection, if any */
-	UPROPERTY()
-	TArray<uint8> CertAuth;
-
-	UPROPERTY()
-	bool bIsValid = false;
-};
\ No newline at end of file
+PRAGMA_ENABLE_DEPRECATION_WARNINGS
\ No newline at end of file
diff --git a/Engine/Source/Runtime/Experimental/ChaosVisualDebugger/Public/ChaosVDRuntimeModule.h b/Engine/Source/Runtime/Experimental/ChaosVisualDebugger/Public/ChaosVDRuntimeModule.h
--- a/Engine/Source/Runtime/Experimental/ChaosVisualDebugger/Public/ChaosVDRuntimeModule.h
+++ b/Engine/Source/Runtime/Experimental/ChaosVisualDebugger/Public/ChaosVDRuntimeModule.h
@@ -4,24 +4,34 @@
 
 #if WITH_CHAOS_VISUAL_DEBUGGER
 
-#include "ChaosVDRecordingDetails.h"
-#include "Templates/SharedPointer.h"
-#include "Modules/ModuleInterface.h"
-#include "ProfilingDebugging/TraceAuxiliary.h"
+#if UE_ENABLE_INCLUDE_ORDER_DEPRECATED_IN_5_8
 #include "Containers/Ticker.h"
+#include "ProfilingDebugging/TraceAuxiliary.h"
+#include "Templates/SharedPointer.h"
+#endif // UE_ENABLE_INCLUDE_ORDER_DEPRECATED_IN_5_8
+
+#include "ChaosVDRecordingDetails.h"
 #include "Misc/ScopeRWLock.h"
 #include "Misc/TransactionallySafeRWLock.h"
+#include "Modules/ModuleInterface.h"
+#include "TraceBasedDebuggerRuntime.h"
 
 #define UE_API CHAOSVDRUNTIME_API
 
-namespace Chaos::VD
+namespace UE::TraceBasedDebuggers
 {
-	class FRelayTraceWriter;
+struct FStartRecordingCommandMessage;
+struct FRelayTraceDataWriter;
 }
 
 struct FChaosVDRecording;
 class FText;
 
+namespace Chaos::VD
+{
+class FRelayTraceWriter;
+}
+
 /* Option flags that controls what should be recorded when doing a full capture **/
 enum class EChaosVDFullCaptureFlags : int32
 {
@@ -33,38 +43,36 @@ ENUM_CLASS_FLAGS(EChaosVDFullCaptureFlags)
 DECLARE_MULTICAST_DELEGATE(FChaosVDRecordingStateChangedDelegate)
 DECLARE_MULTICAST_DELEGATE_OneParam(FChaosVDCaptureRequestDelegate, EChaosVDFullCaptureFlags)
 DECLARE_MULTICAST_DELEGATE_OneParam(FChaosVDRecordingStartFailedDelegate, const FText&)
+
+PRAGMA_DISABLE_DEPRECATION_WARNINGS
+UE_DEPRECATED(5.8, "Use UE::TraceBasedDebuggers::FRelayTraceDataDelegate instead")
 DECLARE_DELEGATE_OneParam(FChaosVDRelayTraceDataDelegate, Chaos::VD::FRelayTraceWriter&)
-DECLARE_DELEGATE_RetVal(FChaosVDTraceDetails, FChaosVDExternalTraceStatusDelegate)
+UE_DEPRECATED(5.8, "Use UE::TraceBasedDebuggers::FExternalRelayBufferQueueSizeDelegate instead")
 DECLARE_DELEGATE_RetVal(uint64, FChaosVDExternalRelayBufferQueueSizeDelegate)
+UE_DEPRECATED(5.8, "Use UE::TraceBasedDebuggers::FTraceConnectionDetails instead")
+DECLARE_DELEGATE_RetVal(FChaosVDTraceDetails, FChaosVDExternalTraceStatusDelegate)
+PRAGMA_ENABLE_DEPRECATION_WARNINGS
 
 DECLARE_LOG_CATEGORY_EXTERN(LogChaosVDRuntime, Log, All)
 
-class FChaosVDRuntimeModule : public IModuleInterface
+class FChaosVDRuntimeModule : public UE::TraceBasedDebuggers::FRuntimeModule, public IModuleInterface
 {
 public:
 
+	UE_API FChaosVDRuntimeModule();
 	static UE_API FChaosVDRuntimeModule& Get();
-	
+
 	UE_API virtual void StartupModule() override;
 	UE_API virtual void ShutdownModule() override;
 
-	/** Starts a CVD recording by starting a Trace session. It will stop any existing trace session
-	 * @param Args : Arguments array usually provided by the commandline. Used to determine if we want to record to file or a local trace server
-	 */
-	UE_API void StartRecording(TConstArrayView<FString> Args);
+	using FRuntimeModule::StartRecording;
 
 	/** Starts a CVD recording by starting a Trace session. It will stop any existing trace session
 	 * @param InRecordingStartCommand : Used to determine if we want to record to file or a local trace server among other settings
 	 */
-	UE_API void StartRecording(const FChaosVDStartRecordingCommandMessage& InRecordingStartCommand);
-	
-	/* Stops an active recording */
-	UE_API void StopRecording();
-
-	/** Returns true if we are currently recording a Physics simulation */
-	bool IsRecording() const
+	void StartRecording(const FChaosVDStartRecordingCommandMessage& InRecordingStartCommand)
 	{
-		return bIsRecording;
+		FRuntimeModule::StartRecording(InRecordingStartCommand);
 	}
 
 	/** Returns a unique ID used to be used to identify CVD (Chaos Visual Debugger) data */
@@ -130,20 +138,13 @@ public:
 		return PerformFullCaptureDelegate.Remove(InDelegateToRemove);
 	}
 
-	/** Returns the accumulated recording time in seconds since the recording started */
-	float GetAccumulatedRecordingTime() const
-	{
-		return AccumulatedRecordingTime;
-	}
-
-	/** Returns the accumulated current send buffer size used by the relay transport system - This should be used for logging or UI feedback only */
-	UE_API int64 GetBufferedDataBytesSize() const;
+PRAGMA_DISABLE_DEPRECATION_WARNINGS
 
 	/** [Deprecated] Returns the full path of the active recording file*/
-	UE_DEPRECATED(5.7, "This method is no longer used. Use GetCurrentTraceSessionDetails")
+	UE_DEPRECATED(5.7, "This method is no longer used. Use GetCurrentTraceConnectionDetails")
 	UE_API FString GetLastRecordingFileNamePath() const;
 
-	/** Return the current connection details for the underlying trace session */
+	UE_DEPRECATED(5.8, "This method is no longer used. Use GetCurrentTraceConnectionDetails")
 	UE_API FChaosVDTraceDetails GetCurrentTraceSessionDetails() const;
 
 	UE_DEPRECATED(5.8, "This method is no longer used and will not be replaced")
@@ -152,130 +153,40 @@ public:
 		return EChaosVDRecordingMode::Invalid;
 	}
 
-	/** Finds a valid file name for a new file - Used to generate the file name for the Trace recording */
-	UE_API void GenerateRecordingFileName(FString& OutFileName);
-
-	/**
-	 * Registers an external handler that will be in charge of relaying the traced data to a remote CVD instance if the relay transport mode is selected
-	 * @param InExternalRelayExecutor Delegate that will be executed to relay data to a remote CVD instance, periodically
-	 */
-	UE_API void RegisterExternalRelayExecutor(const FChaosVDRelayTraceDataDelegate& InExternalRelayExecutor);
-
-	/**
-	 * Registers an external handler that will be in charge handling what happens when the Relay writer overflows
-	 * @param InExternalRelayOverflowHandler Delegate that will be executed when the Relay writer overflows
-	 */
-	UE_API void RegisterExternalRelayOverflowHandler(const FSimpleDelegate& InExternalRelayOverflowHandler);
-
-	/**
-	 * Unbinds the delegate to an external relay overflow handler
-	 */
-	UE_API void UnregisterExternalRelayOverflowHandler();
-
-	/**
-	 * Unbinds the delegate to an external relay executor
-	 */
-	UE_API void UnregisterCurrentExternalRelayExecutor();
-
-	/**
-	 * Registers an external handler that will be in charge of providing the trace status if the recording mode is set to Relay
-	 * @param InExternalTraceStatusCallback Delegate that will be executed to get the trace status, periodically
-	 */
-	UE_API void RegisterExternalTraceStatusProvider(const FChaosVDExternalTraceStatusDelegate& InExternalTraceStatusCallback);
-
-	/**
-	 * Unbinds the delegate to the current external trace status provider
-	 */
-	UE_API void UnregisterExternalTraceStatusProvider();
-
-	/**
-	 *  Registers an external handler that will provide information about any external send buffer the current relay transport implementation might be using, if any.
-	 *  This is used for UI metrics. To handle buffer overflows, use RegisterExternalRelayOverflowHandler 
-	 * @param InExternalRelayBufferSizeCallback 
-	 */
-	UE_API void RegisterExternalRelayBufferSizeProvider(const FChaosVDExternalRelayBufferQueueSizeDelegate& InExternalRelayBufferSizeCallback);
-
-	/**
-	 * Unbinds the delegate to the current external Relay buffer size provider
-	 */
-	UE_API void UnregisterExternalRelayBufferSizeProvider();
-
-	FSimpleMulticastDelegate& OnTraceConnectionDetailsUpdated()
+	UE_DEPRECATED(5.8, "Use the overload taking FChaosVDExternalTraceConnectionStatusDelegate instead")
+	void RegisterExternalTraceStatusProvider(const FChaosVDExternalTraceStatusDelegate& InExternalTraceStatusCallback)
 	{
-		return TraceConnectionDetailsUpdatedDelegate;
 	}
 
+	UE_DEPRECATED(5.8, "Use GetRelayTraceWriterInstance instead.")
 	Chaos::VD::FRelayTraceWriter* GetRelayWriterInstance() const
 	{
-		return RelayWriter.Get();
+		return nullptr;
 	}
+PRAGMA_ENABLE_DEPRECATION_WARNINGS
 
 private:
 
-	void HandleRelayOverflow();
-
-	void RelayTraceData();
-
-	void SetupRelayDataPumpDelegates();
-	void ClearRelayDataPumpDelegates();
-
-	/** Stops the current Trace session */
-	UE_API void StopTrace();
-
-	/** Queues a full Capture of the simulation on the next frame */
-	UE_API bool RequestFullCapture(float DeltaTime);
+	virtual void OnRecordingStartedInternal() override;
+	virtual void OnRecordingStartFailedInternal(const FText& FailureReason) override;
+	virtual void OnRecordingStoppedInternal() override;
+	virtual void OnEnableRequiredTraceChannelsInternal() override;
 
 	/** Queues a full Capture of the simulation on the next frame */
-	UE_API bool RecordingTimerTick(float DeltaTime);
-	
-	void HandleTraceConnectionEstablished();
-
-	/** Used to handle stop requests to the active trace session that were not done by us
-	 * That is a possible scenario because Trace is shared by other In-Editor tools
-	 */
-	UE_API void HandleTraceStopRequest(FTraceAuxiliary::EConnectionType TraceType, const FString& TraceDestination);
-
-	UE_API bool WaitForTraceSessionDisconnect();
-
-	UE_API void EnableRequiredTraceChannels();
-	UE_API void SaveAndDisabledCurrentEnabledTraceChannels();
-	UE_API void RestoreTraceChannelsToPreRecordingState();
-
-	bool bIsRecording = false;
-	bool bRequestedStop = false;
-
-	float AccumulatedRecordingTime = 0.0f;
+	bool RequestFullCapture(float DeltaTime);
 
 	FTSTicker::FDelegateHandle FullCaptureRequesterHandle;
-	FTSTicker::FDelegateHandle RecordingTimerHandle;
-	FTSTicker::FDelegateHandle RelayExecutorTickerHandle;
-
-	FDelegateHandle OnConnectionDelegateHandle;
-	FDelegateHandle OnTraceStoppedDelegateHandle;
 
 	static UE_API FChaosVDRecordingStateChangedDelegate RecordingStartedDelegate;
 	static UE_API FChaosVDRecordingStateChangedDelegate PostRecordingStartedDelegate;
 	static UE_API FChaosVDRecordingStateChangedDelegate RecordingStopDelegate;
 	static UE_API FChaosVDRecordingStartFailedDelegate RecordingStartFailedDelegate;
 	static UE_API FChaosVDCaptureRequestDelegate PerformFullCaptureDelegate;
-
-	std::atomic<int32> LastGeneratedID;
-
-	TMap<FString, bool> OriginalTraceChannelsState;
-
-	EChaosVDTransportMode CurrentTransportMode = EChaosVDTransportMode::Invalid;
-
-	TUniquePtr<Chaos::VD::FRelayTraceWriter> RelayWriter;
-	FChaosVDRelayTraceDataDelegate ExternalRelayExecutorDelegate;
-	FSimpleDelegate ExternalRelayOverflowHandlerDelegate;
-	FChaosVDExternalRelayBufferQueueSizeDelegate ExternalRelayBufferSizeProviderDelegate;
-
 	static UE_API FTransactionallySafeRWLock DelegatesRWLock;
 
-	FChaosVDExternalTraceStatusDelegate ExternalTraceStatusDelegate;
-
-	FSimpleMulticastDelegate TraceConnectionDetailsUpdatedDelegate;
+	std::atomic<int32> LastGeneratedID;
 };
-#endif
 
-#undef UE_API 
\ No newline at end of file
+#undef UE_API
+
+#endif // WITH_CHAOS_VISUAL_DEBUGGER
\ No newline at end of file
diff --git a/Engine/Source/Runtime/Experimental/ChaosVisualDebugger/Public/Trace/RelayTraceWriter.h b/Engine/Source/Runtime/Experimental/ChaosVisualDebugger/Public/Trace/RelayTraceWriter.h
--- a/Engine/Source/Runtime/Experimental/ChaosVisualDebugger/Public/Trace/RelayTraceWriter.h
+++ b/Engine/Source/Runtime/Experimental/ChaosVisualDebugger/Public/Trace/RelayTraceWriter.h
@@ -3,129 +3,26 @@
 #pragma once
 
 #if WITH_CHAOS_VISUAL_DEBUGGER
+
+#if UE_ENABLE_INCLUDE_ORDER_DEPRECATED_IN_5_8
 #include <atomic>
 #include "Containers/Array.h"
 #include "Containers/SpscQueue.h"
 #include "Delegates/Delegate.h"
 #include "Serialization/BitWriter.h"
 #include "Templates/UniquePtr.h"
+#endif // UE_ENABLE_INCLUDE_ORDER_DEPRECATED_IN_5_8
 
-#pragma once
-
-class FBufferArchive;
+#include "RelayTraceDataWriter.h"
 
 namespace Chaos::VD
 {
-/**
- * Writer object that chunks trace data into smaller bunches and queues them so a relay instance
- * can dequeue and send them over the network later
- */
-class FRelayTraceWriter
+class UE_DEPRECATED(5.8, "Use UE::TraceBasedDebuggers::FRelayTraceDataWriter directly instead.") FRelayTraceWriter : public UE::TraceBasedDebuggers::FRelayTraceDataWriter
 {
 public:
 	FRelayTraceWriter();
 	~FRelayTraceWriter();
-
-	/**
-	 * Sets max number of bytes a queued bunch can have
-	 * @param MaxBytes Greater than zero number of bytes to set as limit
-	 */
-	CHAOSVDRUNTIME_API void SetMaxBytesPerBunch(uint32 MaxBytes);
-
-	
-	/**
-	 * Sets max number of bytes this writer can buffer before overflowing, triggering a close
-	 * @param MaxBytes number of bytes to set as limit
-	 */
-	CHAOSVDRUNTIME_API void SetMaxAllowedPendingBytes(uint64 MaxBytes);
-
-	static bool WriteHelper(UPTRINT WriterHandle, const void* Data, uint32 Size)
-	{
-		if (!WriterHandle)
-		{
-			return false;
-		}
-
-		return reinterpret_cast<FRelayTraceWriter*>(WriterHandle)->Write_Internal(Data, Size);
-	}
-
-	static void CloseHelper(UPTRINT WriterHandle)
-	{
-		if (!WriterHandle)
-		{
-			return;
-		}
-
-		reinterpret_cast<FRelayTraceWriter*>(WriterHandle)->Close();
-	}
-
-	/**
-	 * Dequeues a bunch ready to send over the network, if any
-	 */
-	CHAOSVDRUNTIME_API TOptional<TUniquePtr<FBufferArchive>> DequeuePendingBunch();
-
-	/**
-	 * Returns the size of the next bunch in the queue
-	 */
-	CHAOSVDRUNTIME_API int32 NextPendingBunchSize() const;
-
-	/**
-	 * Returns true if we have any data pending to be sent over the network
-	 */
-	bool HasPendingBunches() const
-	{
-		return PendingBunchesCount > 0;
-	}
-
-	/**
-	 * returns the number of bytes overall that are in the queue waiting to be sent
-	 */
-	CHAOSVDRUNTIME_API int64 GetQueuedBytesNum() const;
-
-	/**
-	 * Closes this Writer. After this is called, no new data will be processed, but any pending data can still be dequeued and sent over
-	 */
-	void Close();
-
-	/**
-	 * Returns true if this writer is closed and no new data is being processed
-	 */
-	bool IsClosed() const
-	{
-		return bClosed;
-	}
-
-	DECLARE_DELEGATE(FNewDataAvailableDelegate)
-	FNewDataAvailableDelegate& OnNewDataAvailable()
-	{
-		return NewDataAvailableDelegate;
-	}
-
-	FSimpleDelegate& OnDataOverflow()
-	{
-		return DataOverflowDelegate;
-	}
-
-	private:
-
-	FNewDataAvailableDelegate NewDataAvailableDelegate;
-	FNewDataAvailableDelegate DataOverflowDelegate;
-
-	TUniquePtr<FBufferArchive> CreateBunch();
-
-	bool Write_Internal(const void* Data, uint32 Size);
-
-	bool CanWriteData(uint32 NewDataSize) const;
-
-	TSpscQueue<TUniquePtr<FBufferArchive>> PendingBunchesToRelay;
-	TUniquePtr<FBufferArchive> CurrentPendingBunch;
-	std::atomic<int32> PendingBunchesCount = 0;
-	std::atomic<bool> bClosed = false;
-	uint32 MaxPendingBytesPerBunch = 0;
-	std::atomic<int64> QueuedBytesNum = 0;
-	uint64 MaxAllowedPendingBytes = 1073741824;
 };
 
 }
 #endif // WITH_CHAOS_VISUAL_DEBUGGER
-
diff --git a/Engine/Source/Runtime/Experimental/TraceBasedDebuggers/Config/BaseEngine.ini b/Engine/Source/Runtime/Experimental/TraceBasedDebuggers/Config/BaseEngine.ini
new file mode 100644
--- /dev/null
+++ b/Engine/Source/Runtime/Experimental/TraceBasedDebuggers/Config/BaseEngine.ini
@@ -0,0 +1,14 @@
+[CoreRedirects]
++EnumRedirects=(OldName="/Script/ChaosVDRuntime.EChaosVDRecordingMode", NewName="/Script/TraceBasedDebuggers.ERecordingMode")
++EnumRedirects=(OldName="/Script/ChaosVDRuntime.EChaosVDTransportMode", NewName="/Script/TraceBasedDebuggers.ETraceTransportMode")
++EnumRedirects=(OldName="/Script/ChaosVDRuntime.EChaosVDRemoteSessionAttributes", NewName="/Script/TraceBasedDebuggers.ERemoteSessionAttributes")
++EnumRedirects=(OldName="/Script/ChaosVDRuntime.EChaosVDRemoteSessionReadyState", NewName="/Script/TraceBasedDebuggers.ERemoteSessionReadyState")
+
++StructRedirects=(OldName="/Script/ChaosVDRuntime.ChaosVDTraceDetails", NewName="/Script/TraceBasedDebuggers.TraceConnectionDetails")
++StructRedirects=(OldName="/Script/ChaosSolverEngine.RelayConnectionInfo", NewName="/Script/TraceBasedDebuggers.RelayConnectionInfo")
++StructRedirects=(OldName="/Script/ChaosSolverEngine.ChaosVDSessionPing", NewName="/Script/TraceBasedDebuggers.SessionPing")
++StructRedirects=(OldName="/Script/ChaosSolverEngine.ChaosVDSessionPong", NewName="/Script/TraceBasedDebuggers.SessionPong")
++StructRedirects=(OldName="/Script/ChaosSolverEngine.ChaosVDTraceConnectionDetailsMessage", NewName="/Script/TraceBasedDebuggers.TraceConnectionDetailsMessage")
++StructRedirects=(OldName="/Script/ChaosSolverEngine.ChaosVDRecordingStatusMessage", NewName="/Script/TraceBasedDebuggers.RecordingStatusMessage")
++StructRedirects=(OldName="/Script/ChaosSolverEngine.ChaosVDFullSessionInfoRequestMessage", NewName="/Script/TraceBasedDebuggers.FullSessionInfoRequestMessage")
++StructRedirects=(OldName="/Script/ChaosSolverEngine.ChaosVFFullSessionInfoResponseMessage", NewName="/Script/TraceBasedDebuggers.FullSessionInfoResponseMessage")
diff --git a/Engine/Source/Runtime/Experimental/TraceBasedDebuggers/Private/EngineEditorBridge.cpp b/Engine/Source/Runtime/Experimental/TraceBasedDebuggers/Private/EngineEditorBridge.cpp
new file mode 100644
--- /dev/null
+++ b/Engine/Source/Runtime/Experimental/TraceBasedDebuggers/Private/EngineEditorBridge.cpp
@@ -0,0 +1,123 @@
+// Copyright Epic Games, Inc. All Rights Reserved.
+
+#include "EngineEditorBridge.h"
+
+#if WITH_TRACE_BASED_DEBUGGERS
+
+#include "Misc/App.h"
+
+namespace UE::TraceBasedDebuggers
+{
+
+TWeakPtr<FRemoteSessionsManager> FEngineEditorBridge::SharedRemoteSessionsManager;
+
+TSharedPtr<FRemoteSessionsManager> FEngineEditorBridge::GetSessionsManager()
+{
+	return SharedRemoteSessionsManager.Pin();
+}
+
+void FEngineEditorBridge::HandleRecordingStarted()
+{
+	BroadcastSessionStatus();
+
+	// Create periodic status update (every 0.5 seconds)
+	if (!RecordingStatusUpdateHandle.IsValid())
+	{
+		constexpr float UpdateInterval = 0.5f;
+		RecordingStatusUpdateHandle = FTSTicker::GetCoreTicker().AddTicker(FTickerDelegate::CreateLambda([this](float)
+			{
+				BroadcastSessionStatus();
+				return true;
+			}), UpdateInterval);
+	}
+
+	// Allow derived classes to perform additional actions
+	OnRecordingStartedInternal();
+}
+
+void FEngineEditorBridge::HandleRecordingStopped()
+{
+	// Remove periodic status update
+	FTSTicker::RemoveTicker(RecordingStatusUpdateHandle);
+	RecordingStatusUpdateHandle.Reset();
+
+	BroadcastSessionStatus();
+
+	// Allow derived classes to perform additional actions
+	OnRecordingStoppedInternal();
+}
+
+void FEngineEditorBridge::BroadcastSessionStatus() const
+{
+	FRecordingStatusMessage StatusMessage;
+
+	// Allow derived classes to populate the message
+	BuildRecordingStatusInternal(StatusMessage);
+	if (ensureMsgf(StatusMessage.DebuggerId.IsValid(), TEXT("It is mandatory to provide the debugger type identifier.")))
+	{
+		StatusMessage.InstanceId = FApp::GetInstanceId();
+		RemoteSessionsManager->PublishMessage(StatusMessage);
+	}
+}
+
+void FEngineEditorBridge::SetDataRelayTransportInstance(const TSharedPtr<IDataRelayTransport>& InExternalTraceRelayInstance)
+{
+	if (!ensure(!RelayTraceDataTransportInstance))
+	{
+		UE_LOG_REF(LogCategory, Error, TEXT("An external relay instance was already set. This cannot be changed afterwards. Ignoring set request..."));
+		return;
+	}
+
+	RelayTraceDataTransportInstance = InExternalTraceRelayInstance;
+}
+
+FEngineEditorBridge::FEngineEditorBridge(const FLogCategoryBase& LogCategory)
+	: RemoteSessionsManager(SharedRemoteSessionsManager.IsValid() ? SharedRemoteSessionsManager.Pin().ToSharedRef() : MakeShared<FRemoteSessionsManager>())
+	, LogCategory(LogCategory)
+{
+	if (!SharedRemoteSessionsManager.IsValid())
+	{
+		SharedRemoteSessionsManager = RemoteSessionsManager;
+	}
+}
+
+void FEngineEditorBridge::Initialize()
+{
+	// Allow derived classes to perform their initialization
+	OnInitializeInternal();
+
+	bIsInitialized = true;
+
+	InitializedDelegate.Broadcast();
+}
+
+void FEngineEditorBridge::RegisterRemoteSessionsHandler(TSharedRef<IRemoteSessionsHandler> Handler)
+{
+	ensureMsgf(!RemoteSessionsHandler.IsValid(), TEXT("Handler already registered."));
+	RemoteSessionsHandler = Handler;
+	GetSessionsManager()->RegisterExternalHandler(Handler);
+}
+
+void FEngineEditorBridge::TearDown()
+{
+	if (RelayTraceDataTransportInstance)
+	{
+		RelayTraceDataTransportInstance->Shutdown();
+	}
+
+	FTSTicker::RemoveTicker(RecordingStatusUpdateHandle);
+	RecordingStatusUpdateHandle.Reset();
+
+	// Allow derived classes to perform their cleanup
+	OnTearDownInternal();
+
+	if (RemoteSessionsHandler)
+	{
+		GetSessionsManager()->UnregisterExternalHandler(RemoteSessionsHandler.ToSharedRef());
+	}
+
+	bIsInitialized = false;
+}
+
+} // namespace UE::TraceBasedDebuggers
+#endif // WITH_TRACE_BASED_DEBUGGERS
diff --git a/Engine/Source/Runtime/Experimental/TraceBasedDebuggers/Private/RelayTraceDataWriter.cpp b/Engine/Source/Runtime/Experimental/TraceBasedDebuggers/Private/RelayTraceDataWriter.cpp
new file mode 100644
--- /dev/null
+++ b/Engine/Source/Runtime/Experimental/TraceBasedDebuggers/Private/RelayTraceDataWriter.cpp
@@ -0,0 +1,199 @@
+// Copyright Epic Games, Inc. All Rights Reserved.
+
+#include "RelayTraceDataWriter.h"
+
+#if WITH_TRACE_BASED_DEBUGGERS
+#include "Math/GuardedInt.h"
+#include "Serialization/BufferArchive.h"
+#include "Serialization/MemoryReader.h"
+
+namespace UE::TraceBasedDebuggers
+{
+	FRelayTraceDataWriter::FRelayTraceDataWriter(const FLogCategoryBase& LogCategory)
+		: LogCategory(LogCategory)
+	{
+	}
+
+	FRelayTraceDataWriter::~FRelayTraceDataWriter()
+	{
+		Close();
+	}
+
+	void FRelayTraceDataWriter::SetMaxBytesPerBunch(uint32 MaxBytes)
+	{
+		if (MaxBytes == 0)
+		{
+			UE_LOG_REF(LogCategory, Error, TEXT("[%hs] Only values greater than 0 are supported"), __func__);
+			return;
+		}
+
+		if (MaxBytes != MaxPendingBytesPerBunch)
+		{
+			if (!ensure(PendingBunchesToRelay.IsEmpty()))
+			{
+				return;
+			}
+		}
+
+		MaxPendingBytesPerBunch = MaxBytes;
+	}
+
+	void FRelayTraceDataWriter::SetMaxAllowedPendingBytes(uint64 MaxBytes)
+	{
+		if (MaxBytes == 0)
+		{
+			UE_LOG_REF(LogCategory, Error, TEXT("[%hs] Only values greater than 0 are supported"), __func__);
+			return;
+		}
+
+		MaxAllowedPendingBytes = MaxBytes;
+	}
+
+	bool FRelayTraceDataWriter::WriteHelper(UPTRINT WriterHandle, const void* Data, uint32 Size)
+	{
+		if (!WriterHandle)
+		{
+			return false;
+		}
+
+		return reinterpret_cast<FRelayTraceDataWriter*>(WriterHandle)->Write_Internal(Data, Size);
+	}
+
+	TOptional<TUniquePtr<FBufferArchive>> FRelayTraceDataWriter::DequeuePendingBunch()
+	{
+		if (PendingBunchesCount == 0)
+		{
+			return TOptional<TUniquePtr<FBufferArchive>>();
+		}
+
+		QueuedBytesNum -= NextPendingBunchSize();
+
+		--PendingBunchesCount;
+		return PendingBunchesToRelay.Dequeue();
+	}
+
+	int32 FRelayTraceDataWriter::NextPendingBunchSize() const
+	{
+		if (PendingBunchesCount == 0)
+		{
+			return 0;
+		}
+
+		TUniquePtr<FBufferArchive>* PendingBunch = PendingBunchesToRelay.Peek();
+
+		if (PendingBunch && PendingBunch->IsValid())
+		{
+			return (*PendingBunch)->Num();
+		}
+
+		return 0;
+	}
+
+	TUniquePtr<FBufferArchive> FRelayTraceDataWriter::CreateBunch()
+	{
+		TUniquePtr<FBufferArchive> NewBunch = MakeUnique<FBufferArchive>();
+		return NewBunch;
+	}
+
+	bool FRelayTraceDataWriter::Write_Internal(const void* Data, uint32 Size)
+	{
+		if (bClosed)
+		{
+			return false;
+		}
+
+		const bool bHasValidLimits = MaxPendingBytesPerBunch != 0;
+		if (!ensure(bHasValidLimits))
+		{
+			return false;
+		}
+
+		if (!CanWriteData(Size))
+		{
+			DataOverflowDelegate.ExecuteIfBound();
+			Close();
+			return false;
+		}
+
+		if (!CurrentPendingBunch)
+		{
+			CurrentPendingBunch = CreateBunch();
+		}
+
+		uint32 RemainingSizeToCopy = Size;
+
+		while (RemainingSizeToCopy > 0)
+		{
+			int32 AvailableSpaceInCurrentBunch = MaxPendingBytesPerBunch - CurrentPendingBunch->Num();
+
+			if (AvailableSpaceInCurrentBunch <= 0)
+			{
+				QueuedBytesNum += CurrentPendingBunch->Num();
+				PendingBunchesToRelay.Enqueue(MoveTemp(CurrentPendingBunch));
+				++PendingBunchesCount;
+				CurrentPendingBunch = CreateBunch();
+				AvailableSpaceInCurrentBunch = MaxPendingBytesPerBunch;
+
+				NewDataAvailableDelegate.ExecuteIfBound();
+			}
+
+			const uint32 BytesToCopyNum = static_cast<uint32>(AvailableSpaceInCurrentBunch) > RemainingSizeToCopy ? RemainingSizeToCopy : static_cast<uint32>(AvailableSpaceInCurrentBunch);
+
+			const uint8* DestBufferPtr = static_cast<const uint8*>(Data);
+			const int32 Offset = Size - RemainingSizeToCopy;
+
+			const bool bIsValidOffset = Offset >= 0 && static_cast<uint32>(Offset) < Size;
+			if (!ensureMsgf(bIsValidOffset, TEXT("Unexpected calculated offset [%d] | Size [%u]"), Offset, Size))
+			{
+				return false;
+			}
+
+			CurrentPendingBunch->Serialize(const_cast<uint8*>(DestBufferPtr + Offset), BytesToCopyNum);
+
+			if (CurrentPendingBunch->IsError())
+			{
+				return false;
+			}
+
+			RemainingSizeToCopy -= BytesToCopyNum;
+		}
+
+		return !CurrentPendingBunch->IsError();
+	}
+
+	bool FRelayTraceDataWriter::CanWriteData(uint32 NewDataSize) const
+	{
+		TGuardedInt<uint64> ExpectedQueuedBytes(0);
+
+		ExpectedQueuedBytes += QueuedBytesNum;
+
+		ExpectedQueuedBytes += NewDataSize;
+
+		if (CurrentPendingBunch)
+		{
+			ExpectedQueuedBytes += CurrentPendingBunch->Num();
+		}
+
+		constexpr uint64 DefaultValue = std::numeric_limits<uint64>::max();
+		return ExpectedQueuedBytes.Get(DefaultValue) <= MaxAllowedPendingBytes;
+	}
+
+	int64 FRelayTraceDataWriter::GetQueuedBytesNum() const
+	{
+		return QueuedBytesNum;
+	}
+
+	void FRelayTraceDataWriter::Close()
+	{
+		int64 PendingBytes = QueuedBytesNum;
+		if (PendingBytes > 0)
+		{
+			// TODO: We should handle the case where we close the trace session but we still have data pending in the send queue (UE-320563).
+			// We can ignore it for now as it is not a critical failure
+			UE_LOG_REF(LogCategory, Warning, TEXT("Closing Trace Session with [%lld] bytes in the send queue. This data will be discarded"), PendingBytes);
+		}
+
+		bClosed = true;
+	}
+}
+#endif // WITH_TRACE_BASED_DEBUGGERS
\ No newline at end of file
diff --git a/Engine/Source/Runtime/Experimental/TraceBasedDebuggers/Private/RemoteSessionsManager.cpp b/Engine/Source/Runtime/Experimental/TraceBasedDebuggers/Private/RemoteSessionsManager.cpp
new file mode 100644
--- /dev/null
+++ b/Engine/Source/Runtime/Experimental/TraceBasedDebuggers/Private/RemoteSessionsManager.cpp
@@ -0,0 +1,622 @@
+// Copyright Epic Games, Inc. All Rights Reserved.
+
+#include "RemoteSessionsManager.h"
+
+#include "Containers/ArrayBuilder.h"
+#include "IMessageBus.h"
+#include "IMessagingModule.h"
+#include "MessageEndpoint.h"
+#include "MessageEndpointBuilder.h"
+#include "Misc/App.h"
+
+#include UE_INLINE_GENERATED_CPP_BY_NAME(RemoteSessionsManager)
+
+#ifndef TRACE_BASED_DEBUGGERS_WITHOUT_TRACE
+#define TRACE_BASED_DEBUGGERS_WITHOUT_TRACE 0
+#endif
+
+DEFINE_LOG_CATEGORY(LogTraceBasedDebuggers)
+
+namespace UE::TraceBasedDebuggers
+{
+const FGuid FRemoteSessionsManager::AllRemoteSessionsWrapperGUID = FGuid::NewGuid();
+const FGuid FRemoteSessionsManager::AllRemoteServersWrapperGUID = FGuid::NewGuid();
+const FGuid FRemoteSessionsManager::AllRemoteClientsWrapperGUID = FGuid::NewGuid();
+const FGuid FRemoteSessionsManager::AllSessionsWrapperGUID = FGuid::NewGuid();
+const FGuid FRemoteSessionsManager::CustomSessionsWrapperGUID = FGuid::NewGuid();
+const FGuid FRemoteSessionsManager::InvalidSessionGUID = FGuid();
+
+const FString FRemoteSessionsManager::LocalEditorSessionName = TEXT("Local Editor");
+const FGuid FRemoteSessionsManager::LocalEditorSessionID = IsController() ? FApp::GetInstanceId() : InvalidSessionGUID;
+const FName FRemoteSessionsManager::MessageBusEndpointName = FName("TraceBasedDebuggerSessionManagerEndpoint");
+const FString FRemoteSessionsManager::AllRemoteSessionsTargetName = TEXT("All Remote");
+const FString FRemoteSessionsManager::AllRemoteServersTargetName = TEXT("All Remote Servers");
+const FString FRemoteSessionsManager::AllRemoteClientsTargetName = TEXT("All Remote Clients");
+const FString FRemoteSessionsManager::AllSessionsTargetName = TEXT("All Sessions");
+const FString FRemoteSessionsManager::CustomSessionsTargetName = TEXT("Custom Selection");
+
+const FTraceConnectionDetails& FSessionInfo::GetTraceConnectionDetails()
+{
+	return LastKnownConnectionDetails;
+}
+
+bool FSessionInfo::IsAnyDebuggerRecording() const
+{
+	return LastKnownRecordingState.bIsRecording;
+}
+
+bool FSessionInfo::IsRecording(const FGuid& DebuggerTypeId) const
+{
+	return LastKnownRecordingState.DebuggerId == DebuggerTypeId
+		&& LastKnownRecordingState.bIsRecording;
+}
+
+uint64 FSessionInfo::GetBufferedBytesNum() const
+{
+	return LastKnownRecordingState.BufferedDataBytesSize;
+}
+
+void FSessionInfo::SetReceivedBytesPerSecond(uint64 InNewBytesPerSecond)
+{
+	ReceivedBytesPerSecond = InNewBytesPerSecond;
+}
+
+uint64 FSessionInfo::GetReceivedBytesPerSecond() const
+{
+	return ReceivedBytesPerSecond;
+}
+
+bool FMultiSessionInfo::IsAnyDebuggerRecording() const
+{
+	bool bIsRecording = false;
+	EnumerateInnerSessions([&bIsRecording](const TSharedRef<FSessionInfo>& InSessionRef)
+	{
+		if (InSessionRef->IsAnyDebuggerRecording())
+		{
+			bIsRecording = true;
+			return false;
+		}
+		return true;
+	});
+
+	return bIsRecording;
+}
+
+bool FMultiSessionInfo::IsRecording(const FGuid& DebuggerTypeId) const
+{
+	bool bIsRecording = false;
+	EnumerateInnerSessions([&bIsRecording, DebuggerTypeId](const TSharedRef<FSessionInfo>& InSessionRef)
+	{
+		if (InSessionRef->IsRecording(DebuggerTypeId))
+		{
+			bIsRecording = true;
+			return false;
+		}
+		return true;
+	});
+
+	return bIsRecording;
+}
+
+uint64 FMultiSessionInfo::GetBufferedBytesNum() const
+{
+	uint64 BufferedBytesNum = 0;
+	EnumerateInnerSessions([&BufferedBytesNum](const TSharedRef<FSessionInfo>& InSessionRef)
+	{
+		BufferedBytesNum += InSessionRef->GetBufferedBytesNum();
+		return true;
+	});
+
+	return BufferedBytesNum;
+}
+
+void FRemoteSessionsManager::RegisterBuiltInMessageTypes()
+{
+	SupportedMessageTypes.Emplace(FSessionPong::StaticStruct());
+	SupportedMessageTypes.Emplace(FRecordingStatusMessage::StaticStruct());
+	SupportedMessageTypes.Emplace(FSessionPing::StaticStruct());
+	SupportedMessageTypes.Emplace(FFullSessionInfoRequestMessage::StaticStruct());
+	SupportedMessageTypes.Emplace(FFullSessionInfoResponseMessage::StaticStruct());
+	SupportedMessageTypes.Emplace(FTraceConnectionDetailsMessage::StaticStruct());
+}
+
+FRemoteSessionsManager::FRemoteSessionsManager()
+{
+	RegisterBuiltInMessageTypes();
+	Initialize();
+}
+
+FRemoteSessionsManager::~FRemoteSessionsManager()
+{
+	Shutdown();
+}
+
+TWeakPtr<FSessionInfo> FRemoteSessionsManager::GetSessionInfo(const FGuid Id)
+{
+	if (TSharedPtr<FSessionInfo>* FoundSessionPtrPtr = ActiveSessionsById.Find(Id))
+	{
+		return *FoundSessionPtrPtr;
+	}
+
+	return nullptr;
+}
+
+TWeakPtr<FSessionInfo> FRemoteSessionsManager::GetFirstAvailableSingleSessionInfo()
+{
+	for (const TPair<FGuid, TSharedPtr<FSessionInfo>>& SessionWithID : ActiveSessionsById)
+	{
+		if (SessionWithID.Value && !EnumHasAnyFlags(SessionWithID.Value->GetSessionTypeAttributes(), ERemoteSessionAttributes::IsMultiSessionWrapper))
+		{
+			return SessionWithID.Value;
+		}
+	}
+
+	return nullptr;
+}
+
+TSharedPtr<FSessionInfo> FRemoteSessionsManager::CreatedWrapperSessionInfo(FGuid InstanceId, const FStringView SessionName)
+{
+	TSharedPtr<FMultiSessionInfo> NewSessionInfo = MakeShared<FMultiSessionInfo>();
+	NewSessionInfo->InstanceId = InstanceId;
+	NewSessionInfo->SessionName = SessionName;
+
+	return NewSessionInfo;
+}
+
+TSharedPtr<FMessageEndpoint> FRemoteSessionsManager::CreateEndPoint(const TSharedRef<IMessageBus>& InMessageBus)
+{
+	FMessageEndpointBuilder EndpointBuilder = FMessageEndpoint::Builder(MessageBusEndpointName, InMessageBus)
+		.Handling<FSessionPing>(this, &FRemoteSessionsManager::HandleSessionPingMessage)
+		.Handling<FFullSessionInfoRequestMessage>(this, &FRemoteSessionsManager::HandleFullSessionStateRequestMessage);
+
+	if (IsController())
+	{
+		EndpointBuilder.Handling<FSessionPong>(this, &FRemoteSessionsManager::HandleSessionPongMessage)
+			.Handling<FRecordingStatusMessage>(this, &FRemoteSessionsManager::HandleRecordingStatusUpdateMessage)
+			.Handling<FFullSessionInfoResponseMessage>(this, &FRemoteSessionsManager::HandleFullSessionStateResponseMessage)
+			.Handling<FTraceConnectionDetailsMessage>(this, &FRemoteSessionsManager::HandleConnectionDetailsUpdateMessage);
+	}
+
+	for (const TSharedRef<IRemoteSessionsHandler>& Handler : Handlers)
+	{
+		Handler->OnCreatingMessageEndpoint(AsShared(), EndpointBuilder);
+	}
+
+	return EndpointBuilder.Build();
+}
+
+void FRemoteSessionsManager::ReInitializeMessagingSystem(const TSharedPtr<IMessageBus>& InMessageBus)
+{
+	ShutdownMessagingSystem();
+	InitializeMessagingSystem(InMessageBus);
+}
+
+void FRemoteSessionsManager::RegisterExternalSupportedMessageType(const UScriptStruct* ScriptStruct)
+{
+	if (ensure(ScriptStruct))
+	{
+		SupportedMessageTypes.Emplace(ScriptStruct);
+	}
+}
+
+void FRemoteSessionsManager::InitializeMessagingSystem(const TSharedPtr<IMessageBus>& InMessageBus)
+{
+	if (!InMessageBus)
+	{
+		return;
+	}
+
+	MessageBusPtr = InMessageBus;
+	MessageEndpoint = CreateEndPoint(InMessageBus.ToSharedRef());
+
+	if (!MessageEndpoint)
+	{
+		return;
+	}
+
+	// Subscribe built-in message types
+	if (IsController())
+	{
+		MessageEndpoint->Subscribe<FSessionPong>();
+		MessageEndpoint->Subscribe<FRecordingStatusMessage>();
+		MessageEndpoint->Subscribe<FFullSessionInfoResponseMessage>();
+		MessageEndpoint->Subscribe<FTraceConnectionDetailsMessage>();
+	}
+
+	MessageEndpoint->Subscribe<FSessionPing>();
+	MessageEndpoint->Subscribe<FFullSessionInfoRequestMessage>();
+
+	// Allow external owner to subscribe
+	for (const TSharedRef<IRemoteSessionsHandler>& Handler : Handlers)
+	{
+		Handler->OnSubscribingMessageTypes(AsShared(), *MessageEndpoint);
+	}
+
+	if (bInitialized)
+	{
+		MessageEndpoint->Enable();
+	}
+	else
+	{
+		MessageEndpoint->Disable();
+	}
+
+	// Notify that messaging initialization is done
+	MessagingInitializedDelegate.Broadcast(InMessageBus, MessageEndpoint);
+}
+
+void FRemoteSessionsManager::ShutdownMessagingSystem()
+{
+	if (MessageEndpoint)
+	{
+		// Unsubscribe from all message types (i.e. no delegate required)
+		MessageEndpoint->Unsubscribe();
+	}
+
+	MessageBusPtr.Reset();
+	MessageEndpoint = nullptr;
+}
+
+void FRemoteSessionsManager::Initialize()
+{
+#if !defined(WITH_TRACE_BASED_DEBUGGERS_EXTERNAL_MESSAGING) || !WITH_TRACE_BASED_DEBUGGERS_EXTERNAL_MESSAGING
+	InitializeMessagingSystem(IMessagingModule::Get().GetDefaultBus());
+#endif
+
+	ActiveSessionsById.Add(AllRemoteSessionsWrapperGUID, CreatedWrapperSessionInfo(AllRemoteSessionsWrapperGUID, AllRemoteSessionsTargetName));
+	ActiveSessionsById.Add(AllRemoteServersWrapperGUID, CreatedWrapperSessionInfo(AllRemoteServersWrapperGUID, AllRemoteServersTargetName));
+	ActiveSessionsById.Add(AllRemoteClientsWrapperGUID, CreatedWrapperSessionInfo(AllRemoteClientsWrapperGUID, AllRemoteClientsTargetName));
+	ActiveSessionsById.Add(AllSessionsWrapperGUID, CreatedWrapperSessionInfo(AllSessionsWrapperGUID, AllSessionsTargetName));
+	ActiveSessionsById.Add(CustomSessionsWrapperGUID, CreatedWrapperSessionInfo(CustomSessionsWrapperGUID, CustomSessionsTargetName));
+
+	if (MessageEndpoint)
+	{
+		MessageEndpoint->Enable();
+	}
+
+	bInitialized = true;
+}
+
+void FRemoteSessionsManager::Shutdown()
+{
+	bInitialized = false;
+	ShutdownMessagingSystem();
+
+	if (TickHandle.IsValid())
+	{
+		FTSTicker::RemoveTicker(TickHandle);
+		TickHandle.Reset();
+	}
+}
+
+void FRemoteSessionsManager::RegisterExternalHandler(TSharedRef<IRemoteSessionsHandler> Handler)
+{
+	Handlers.Emplace(Handler);
+	if (const TSharedPtr<IMessageBus> MessageBus = MessageBusPtr.Pin())
+	{
+		ReInitializeMessagingSystem(MessageBus);
+	}
+}
+
+void FRemoteSessionsManager::UnregisterExternalHandler(TSharedRef<IRemoteSessionsHandler> Handler)
+{
+	Handlers.Remove(Handler);
+	if (const TSharedPtr<IMessageBus> MessageBus = MessageBusPtr.Pin())
+	{
+		ReInitializeMessagingSystem(MessageBus);
+	}
+}
+
+void FRemoteSessionsManager::EnumerateMessageTypes(const FVisitorFunction& InVisitor)
+{
+	for (const UScriptStruct* MessageType : SupportedMessageTypes)
+	{
+		if (MessageType)
+		{
+			InVisitor(MessageType);
+		}
+	}
+}
+
+void FRemoteSessionsManager::StartSessionDiscovery()
+{
+	if (TickHandle.IsValid())
+	{
+		UE_LOG(LogTraceBasedDebuggers, Warning, TEXT("[%hs] Session discovery already started"), __func__);
+		return;
+	}
+
+	constexpr float TickInterval = 1.0f;
+	TickHandle = FTSTicker::GetCoreTicker().AddTicker(FTickerDelegate::CreateRaw(this, &FRemoteSessionsManager::Tick), TickInterval);
+}
+
+void FRemoteSessionsManager::StopSessionDiscovery()
+{
+	if (TickHandle.IsValid())
+	{
+		RemoveExpiredSessions(ERemoveSessionOptions::ForceRemoveAll);
+		FTSTicker::RemoveTicker(TickHandle);
+		TickHandle.Reset();
+	}
+}
+
+void FRemoteSessionsManager::SendFullSessionStateRequestCommand(const FMessageAddress& InDestinationAddress)
+{
+	if (!MessageEndpoint)
+	{
+		UE_LOG(LogTraceBasedDebuggers, Error, TEXT("[%hs] Failed to send command | Invalid endpoint."), __FUNCTION__);
+		return;
+	}
+
+	MessageEndpoint->Send(
+	FMessageEndpoint::MakeMessage<FFullSessionInfoRequestMessage>(),
+	FFullSessionInfoRequestMessage::StaticStruct(),
+	EMessageFlags::Reliable,
+	nullptr,
+	TArrayBuilder<FMessageAddress>().Add(InDestinationAddress),
+	FTimespan::Zero(),
+	FDateTime::MaxValue());
+}
+
+bool FRemoteSessionsManager::Tick(float DeltaTime)
+{
+	if (IsController())
+	{
+		SendPing();
+		RemoveExpiredSessions();
+	}
+
+	return true;
+}
+
+void FRemoteSessionsManager::SendPing()
+{
+	if (ensure(MessageEndpoint))
+	{
+		if (FSessionPing* SessionPingData = FMessageEndpoint::MakeMessage<FSessionPing>())
+		{
+			SessionPingData->ControllerInstanceId = FApp::GetInstanceId();
+			MessageEndpoint->Publish(SessionPingData, EMessageScope::Network);
+		}
+	}
+}
+
+void FRemoteSessionsManager::SendPong(const FSessionPing& InMessage)
+{
+	if (!ensure(MessageEndpoint))
+	{
+		return;
+	}
+
+	if (FSessionPong* PongMessage = FMessageEndpoint::MakeMessage<FSessionPong>())
+	{
+		PongMessage->InstanceId = FApp::GetInstanceId();
+		PongMessage->SessionId = FApp::GetSessionId();
+		PongMessage->BuildTargetType = static_cast<uint8>(FApp::GetBuildTargetType());
+
+		if (InMessage.ControllerInstanceId == PongMessage->InstanceId)
+		{
+			PongMessage->SessionName = LocalEditorSessionName;
+		}
+		else
+		{
+			FString AppSessionName = FApp::GetSessionName();
+			PongMessage->SessionName = AppSessionName == TEXT("None") || AppSessionName.IsEmpty()
+				? FString::Format(TEXT("{0} {1} {2}"),
+						{FApp::GetProjectName()
+						, FString(LexToString(FApp::GetBuildTargetType()))
+						, FString::FromInt(FPlatformProcess::GetCurrentProcessId())})
+				: AppSessionName;
+		}
+
+		MessageEndpoint->Publish(PongMessage, EMessageScope::Network);
+	}
+}
+
+void FRemoteSessionsManager::RegisterSessionInMultiSessionWrapper(const TSharedRef<FSessionInfo>& InSessionInfoRef)
+{
+	if (InSessionInfoRef->SessionName != LocalEditorSessionName)
+	{
+		StaticCastSharedPtr<FMultiSessionInfo>(ActiveSessionsById.FindChecked(AllRemoteSessionsWrapperGUID))->InnerSessionsByInstanceID.Emplace(InSessionInfoRef->InstanceId, InSessionInfoRef);
+
+		if (InSessionInfoRef->BuildTargetType == EBuildTargetType::Server)
+		{
+			StaticCastSharedPtr<FMultiSessionInfo>(ActiveSessionsById.FindChecked(AllRemoteServersWrapperGUID))->InnerSessionsByInstanceID.Emplace(InSessionInfoRef->InstanceId, InSessionInfoRef);
+		}
+		else
+		{
+			StaticCastSharedPtr<FMultiSessionInfo>(ActiveSessionsById.FindChecked(AllRemoteClientsWrapperGUID))->InnerSessionsByInstanceID.Emplace(InSessionInfoRef->InstanceId, InSessionInfoRef);
+		}
+	}
+
+	StaticCastSharedPtr<FMultiSessionInfo>(ActiveSessionsById.FindChecked(AllSessionsWrapperGUID))->InnerSessionsByInstanceID.Emplace(InSessionInfoRef->InstanceId, InSessionInfoRef);
+}
+
+void FRemoteSessionsManager::DeRegisterSessionInMultiSessionWrapper(const TSharedRef<FSessionInfo>& InSessionInfoRef)
+{
+	StaticCastSharedPtr<FMultiSessionInfo>(ActiveSessionsById.FindChecked(AllRemoteSessionsWrapperGUID))->InnerSessionsByInstanceID.Remove(InSessionInfoRef->InstanceId);
+	StaticCastSharedPtr<FMultiSessionInfo>(ActiveSessionsById.FindChecked(AllSessionsWrapperGUID))->InnerSessionsByInstanceID.Remove(InSessionInfoRef->InstanceId);
+	StaticCastSharedPtr<FMultiSessionInfo>(ActiveSessionsById.FindChecked(AllRemoteServersWrapperGUID))->InnerSessionsByInstanceID.Remove(InSessionInfoRef->InstanceId);
+	StaticCastSharedPtr<FMultiSessionInfo>(ActiveSessionsById.FindChecked(AllRemoteClientsWrapperGUID))->InnerSessionsByInstanceID.Remove(InSessionInfoRef->InstanceId);
+}
+
+void FRemoteSessionsManager::ProcessPendingMessagesForSession(const FSessionPong& InMessage, const TSharedRef<FSessionInfo>& InSessionInfoPtr)
+{
+	PendingRecordingStatusMessages.RemoveAndCopyValue(InMessage.InstanceId, InSessionInfoPtr->LastKnownRecordingState);
+	PendingRecordingConnectionDetailsMessages.RemoveAndCopyValue(InMessage.InstanceId, InSessionInfoPtr->LastKnownConnectionDetails);
+}
+
+void FRemoteSessionsManager::HandleSessionPongMessage(const FSessionPong& InMessage, const TSharedRef<IMessageContext>& InContext)
+{
+	TSharedPtr<FSessionInfo>& SessionInfoPtr = ActiveSessionsById.FindOrAdd(InMessage.InstanceId);
+
+	if (!SessionInfoPtr)
+	{
+		SessionInfoPtr = MakeShared<FSessionInfo>();
+		SessionInfoPtr->Address = InContext->GetSender();
+		SessionInfoPtr->InstanceId = InMessage.InstanceId;
+		SessionInfoPtr->SessionName = InMessage.SessionName;
+		SessionInfoPtr->BuildTargetType = static_cast<EBuildTargetType>(InMessage.BuildTargetType);
+
+		RegisterSessionInMultiSessionWrapper(SessionInfoPtr.ToSharedRef());
+
+		SessionDiscoveredDelegate.Broadcast(SessionInfoPtr);
+
+		// This is the first time we see this session, so we need to request the rest of its state so we can properly populate the UI
+		SendFullSessionStateRequestCommand(SessionInfoPtr->Address);
+	}
+
+	SessionInfoPtr->LastPingTime = FDateTime::UtcNow();
+
+	ProcessPendingMessagesForSession(InMessage, SessionInfoPtr.ToSharedRef());
+
+	SessionsUpdatedDelegate.Broadcast();
+}
+
+void FRemoteSessionsManager::HandleSessionPingMessage(const FSessionPing& InMessage, const TSharedRef<IMessageContext>& InContext)
+{
+	// If this instance is not running trace, don't answer to the ping as the debugger instance will not be able to do anything useful with it
+#if !TRACE_BASED_DEBUGGERS_WITHOUT_TRACE
+	SendPong(InMessage);
+#endif
+}
+
+void FRemoteSessionsManager::HandleRecordingStatusUpdateMessage(const FRecordingStatusMessage& Message, const TSharedRef<IMessageContext>& InContext)
+{
+	if (TSharedPtr<FSessionInfo>* SessionInfoPtrPtr = ActiveSessionsById.Find(Message.InstanceId))
+	{
+		TSharedPtr<FSessionInfo>& SessionInfoPtr = *SessionInfoPtrPtr;
+		check(SessionInfoPtr);
+
+		const FGuid& LastKnownDebuggerTypeId = SessionInfoPtr->LastKnownRecordingState.DebuggerId;
+		const bool bLastKnownDebuggerIsRecording = SessionInfoPtr->LastKnownRecordingState.bIsRecording;
+		if (LastKnownDebuggerTypeId.IsValid()
+			&& LastKnownDebuggerTypeId != Message.DebuggerId
+			&& bLastKnownDebuggerIsRecording)
+		{
+			ensureMsgf(false, TEXT("A trace-based debugger changed is recording status while another debugger is already recording."
+				" This is not supported and indicates a problem on the runtime modules side."));
+			return;
+		}
+
+		if (bLastKnownDebuggerIsRecording != Message.bIsRecording)
+		{
+			if (Message.bIsRecording)
+			{
+				RecordingStartedDelegate.Broadcast(SessionInfoPtr);
+			}
+			else
+			{
+				RecordingStoppedDelegate.Broadcast(SessionInfoPtr);
+			}
+		}
+
+		SessionInfoPtr->LastKnownRecordingState = Message;
+	}
+	else
+	{
+		PendingRecordingStatusMessages.FindOrAdd(Message.InstanceId) = Message;
+	}
+}
+
+void FRemoteSessionsManager::HandleConnectionDetailsUpdateMessage(const FTraceConnectionDetailsMessage& InMessage, const TSharedRef<IMessageContext>& InContext)
+{
+	if (TSharedPtr<FSessionInfo>* SessionInfoPtrPtr = ActiveSessionsById.Find(InMessage.InstanceId))
+	{
+		const TSharedPtr<FSessionInfo>& SessionInfoPtr = *SessionInfoPtrPtr;
+		check(SessionInfoPtr);
+
+		SessionInfoPtr->LastKnownConnectionDetails = InMessage.TraceDetails;
+	}
+	else
+	{
+		PendingRecordingConnectionDetailsMessages.FindOrAdd(InMessage.InstanceId) = InMessage.TraceDetails;
+	}
+}
+
+void FRemoteSessionsManager::HandleFullSessionStateRequestMessage(const FFullSessionInfoRequestMessage& InMessage, const TSharedRef<IMessageContext>& InContext)
+{
+	if (!ensure(MessageEndpoint))
+	{
+		return;
+	}
+
+	if (FFullSessionInfoResponseMessage* FullSessionStateResponse = FMessageEndpoint::MakeMessage<FFullSessionInfoResponseMessage>())
+	{
+		FullSessionStateResponse->InstanceId = FApp::GetInstanceId();
+
+		// Allow external handlers to append information
+		for (const TSharedRef<IRemoteSessionsHandler>& Handler : Handlers)
+		{
+			Handler->OnBuildingFullSessionInfoResponseMessage(*FullSessionStateResponse);
+		}
+
+		MessageEndpoint->Send(
+			FullSessionStateResponse,
+			FFullSessionInfoResponseMessage::StaticStruct(),
+			EMessageFlags::Reliable,
+			nullptr,
+			TArrayBuilder<FMessageAddress>().Add(InContext->GetSender()),
+			FTimespan::Zero(),
+			FDateTime::MaxValue());
+	}
+}
+
+void FRemoteSessionsManager::HandleFullSessionStateResponseMessage(const FFullSessionInfoResponseMessage& InMessage, const TSharedRef<IMessageContext>& InContext)
+{
+	if (const TSharedPtr<FSessionInfo>* SessionInfoPtr = ActiveSessionsById.Find(InMessage.InstanceId))
+	{
+		if (const TSharedPtr<FSessionInfo>& SessionInfo = *SessionInfoPtr)
+		{
+			SessionInfo->LastKnownRecordingState.bIsRecording = InMessage.bIsRecording;
+			SessionInfo->LastKnownRecordingState.DebuggerId = InMessage.DebuggerId;
+
+			// Allow external handlers to process the message
+			for (const TSharedRef<IRemoteSessionsHandler>& Handler : Handlers)
+			{
+				Handler->OnHandlingFullSessionInfoResponseMessage(InMessage, SessionInfo);
+			}
+		}
+	}
+}
+
+void FRemoteSessionsManager::RemoveExpiredSessions(const ERemoveSessionOptions Options)
+{
+	bool bAnySessionRemoved = false;
+	const FDateTime CurrentTime = FDateTime::UtcNow();
+	for (TMap<FGuid, TSharedPtr<FSessionInfo>>::TIterator RemoveIterator = ActiveSessionsById.CreateIterator(); RemoveIterator; ++RemoveIterator)
+	{
+		TSharedPtr<FSessionInfo>& SessionInfoPtr = RemoveIterator.Value();
+		if (!SessionInfoPtr)
+		{
+			RemoveIterator.RemoveCurrent();
+			bAnySessionRemoved = true;
+			continue;
+		}
+
+		if (!EnumHasAnyFlags(SessionInfoPtr->GetSessionTypeAttributes(), ERemoteSessionAttributes::CanExpire))
+		{
+			continue;
+		}
+
+		const FTimespan ElapsedTime = CurrentTime - SessionInfoPtr->LastPingTime;
+
+		// A session goes into busy state if we are attempting to issue a command that might stall the target, currently that only happens on recording start commands of complex maps
+		// In these cases, we need to allow more time between pings. If a recording command failed, it is expected the state to be changed to Ready again
+		const float MaxAllowedTimeBetweenPings = SessionInfoPtr->ReadyState == ERemoteSessionReadyState::Busy || SessionInfoPtr->IsAnyDebuggerRecording() ? 60.0f : 3.0f;
+		if (EnumHasAnyFlags(Options, ERemoveSessionOptions::ForceRemoveAll) || ElapsedTime > FTimespan::FromSeconds(MaxAllowedTimeBetweenPings))
+		{
+			SessionExpiredDelegate.Broadcast(SessionInfoPtr);
+
+			DeRegisterSessionInMultiSessionWrapper(SessionInfoPtr.ToSharedRef());
+			RemoveIterator.RemoveCurrent();
+			bAnySessionRemoved = true;
+		}
+	}
+
+	if (bAnySessionRemoved)
+	{
+		SessionsUpdatedDelegate.Broadcast();
+	}
+}
+} // namespace UE::TraceBasedDebuggers
\ No newline at end of file
diff --git a/Engine/Source/Runtime/Experimental/TraceBasedDebuggers/Private/TraceBasedDebuggerRuntime.cpp b/Engine/Source/Runtime/Experimental/TraceBasedDebuggers/Private/TraceBasedDebuggerRuntime.cpp
new file mode 100644
--- /dev/null
+++ b/Engine/Source/Runtime/Experimental/TraceBasedDebuggers/Private/TraceBasedDebuggerRuntime.cpp
@@ -0,0 +1,643 @@
+// Copyright Epic Games, Inc. All Rights Reserved.
+
+#include "TraceBasedDebuggerRuntime.h"
+
+#include "Misc/CommandLine.h"
+#include "RelayTraceDataWriter.h"
+#include "RemoteSessionsManager.h"
+#include "Math/GuardedInt.h"
+#include "Misc/App.h"
+#include "Misc/FileHelper.h"
+
+#if WITH_EDITOR
+#include "Misc/ScopedSlowTask.h"
+#include "Misc/MessageDialog.h"
+#endif
+
+#define LOCTEXT_NAMESPACE "TraceBasedDebuggerRuntime"
+
+namespace UE::TraceBasedDebuggers
+{
+static FAutoConsoleVariable CVarMaxTimeToWaitForDisconnect(
+	TEXT("tracebaseddebuggers.MaxTimeToWaitForDisconnectSeconds"),
+	5.0f,
+	TEXT("Max time to wait after attempting to stop an active trace session. After that time has passed if we are still connected, debuggers will continue and eventually error out."));
+
+static FAutoConsoleVariable CVarMaxBytesPerRelayChunk(
+	TEXT("tracebaseddebuggers.MaxBytesPerRelayChunk"),
+	65536,
+	TEXT("How many bytes per relay data message we should send"));
+
+static FAutoConsoleVariable CVarMaxPendingBytesInRelayWriter(
+	TEXT("tracebaseddebuggers.MaxPendingBytesInRelayWriter"),
+	536870912,
+	TEXT("How many bytes the relay writer can hold before overflowing, triggering a close. Defaults to ~512MB"));
+
+FRuntimeModule::FRuntimeModule(const FLogCategoryBase& InLogCategory, const FGuid& InDebuggerTypeID)
+	: LogCategory(InLogCategory)
+	, DebuggerGuid(InDebuggerTypeID)
+{
+}
+
+FRuntimeModule::~FRuntimeModule()
+{
+}
+
+bool FRuntimeModule::TryRecordOnStartup()
+{
+	if (AutoRecordCommand.IsEmpty()
+		|| !FParse::Param(FCommandLine::Get(), *AutoRecordCommand))
+	{
+		return false;
+	}
+
+	TArray<FString, TInlineAllocator<1>> Options;
+	FString HostAddress;
+	if (FParse::Value(FCommandLine::Get(), *HostAddressCommand, HostAddress))
+	{
+		Options.Emplace(TEXT("Server"));
+		Options.Emplace(MoveTemp(HostAddress));
+	}
+
+	StartRecording(Options);
+	return true;
+}
+
+void FRuntimeModule::StartRecording(const TConstArrayView<FString> Args)
+{
+	using namespace UE::TraceBasedDebuggers;
+	FStartRecordingCommandMessage StartRecordingCommand;
+
+	// Default (i.e., no arguments) or "File" trace to file using GenerateRecordingFileName
+	if (Args.Num() == 0 || Args[0] == TEXT("File"))
+	{
+		StartRecordingCommand.RecordingMode = ERecordingMode::File;
+		StartRecordingCommand.TransportMode = ETraceTransportMode::FileSystem;
+	}
+	// "Server [ADDRESS]": trace server using either the optional [ADDRESS] or local host (i.e., "127.0.0.1")
+	else if (Args[0] == TEXT("Server"))
+	{
+		StartRecordingCommand.RecordingMode = ERecordingMode::Live;
+		StartRecordingCommand.TransportMode = ETraceTransportMode::TraceServer;
+		StartRecordingCommand.Target = Args.IsValidIndex(1) ? Args[1] : TEXT("127.0.0.1");
+	}
+	// "-tracefile [FILENAME]": trace to file using either the optional [FILENAME] or GenerateRecordingFileName
+	// "-tracehost [ADDRESS]": trace server using either the optional [ADDRESS] or local host (i.e., "127.0.0.1")
+	else
+	{
+		FTraceAuxiliary::EConnectionType TraceType = FTraceAuxiliary::EConnectionType::None;
+		FString TraceDestination;
+
+		for (const FString& Arg : Args)
+		{
+			if (Arg.StartsWith(TEXT("-tracefile"), ESearchCase::IgnoreCase))
+			{
+				ensureMsgf(TraceType == FTraceAuxiliary::EConnectionType::None
+					, TEXT("StartRecording: Specifying more than 1 trace destination is not supported. Received: %s")
+					, *FString::Join(Args, TEXT(" ")));
+
+				TraceType = FTraceAuxiliary::EConnectionType::File;
+
+				// Try to extract filename.
+				if (FParse::Value(*Arg, TEXT("-tracefile="), TraceDestination))
+				{
+					// Make sure it's a valid filename
+					FText FilenameError;
+					if (!FFileHelper::IsFilenameValidForSaving(TraceDestination, FilenameError))
+					{
+						ensureMsgf(false, TEXT("StartRecording: Specified filename is not supported: %s"), *FilenameError.ToString());
+						TraceDestination = "";
+					}
+				}
+			}
+			else if (Arg.StartsWith(TEXT("-tracehost"), ESearchCase::IgnoreCase))
+			{
+				ensureMsgf(TraceType == FTraceAuxiliary::EConnectionType::None
+					, TEXT("StartRecording: Specifying more than 1 trace destination is not supported. Received: %s")
+					, *FString::Join(Args, TEXT(" ")));
+
+				TraceType = FTraceAuxiliary::EConnectionType::Network;
+
+				if (FParse::Value(*Arg, TEXT("-tracehost="), TraceDestination))
+				{
+					// Should we validate that TraceDestination is valid ip address?
+				}
+			}
+			else
+			{
+				ensureMsgf(false, TEXT("StartRecording: Received unknown argument: %s"), *Arg);
+			}
+		}
+
+		switch (TraceType)
+		{
+		case FTraceAuxiliary::EConnectionType::Network:
+			StartRecordingCommand.RecordingMode = ERecordingMode::Live;
+			StartRecordingCommand.TransportMode = ETraceTransportMode::TraceServer;
+			StartRecordingCommand.Target = TraceDestination;
+			break;
+		case FTraceAuxiliary::EConnectionType::File:
+			StartRecordingCommand.RecordingMode = ERecordingMode::File;
+			StartRecordingCommand.TransportMode = ETraceTransportMode::FileSystem;
+			StartRecordingCommand.Target = TraceDestination;
+			break;
+		}
+	}
+
+	StartRecording(StartRecordingCommand);
+}
+
+void FRuntimeModule::StartRecording(const FStartRecordingCommandMessage& InRecordingStartCommand)
+{
+	using namespace UE::TraceBasedDebuggers;
+	if (!ensureAlwaysMsgf(!bIsRecording, TEXT("Received a recording start command while a recording is already active")))
+	{
+		UE_LOG_REF(LogCategory, Log, TEXT("[%hs] There is an active recording. Attempting to stop it..."), __func__)
+		StopRecording();
+	}
+
+	// Start with a generic Failure reason
+	FText FailureReason = LOCTEXT("SeeLogsForErrorDetailsText", "Please see the logs for more details...");
+
+#if UE_TRACE_ENABLED || UE_TRACE_MINIMAL_ENABLED
+
+	// Other tools could be using trace
+	// This is aggressive but until Trace supports multi-sessions, just take over.
+	if (FTraceAuxiliary::IsConnected())
+	{
+		UE_LOG_REF(LogCategory, Log, TEXT("[%hs] There is an active trace session. attempting to disconnect..."), __func__);
+
+		//TODO: We should make the wait async like we do whe we attempt to connect to a live session
+		if (FTraceAuxiliary::Stop() && WaitForTraceSessionDisconnect())
+		{
+			UE_LOG_REF(LogCategory, Log, TEXT("[%hs] Successful disconnect attempt!."), __func__);
+		}
+		else
+		{
+			FailureReason = LOCTEXT("FailedToStopActiveRecordingErrorMessage", "Failed to Stop active Trace Session.");
+		}
+	}
+
+	SaveAndDisabledCurrentEnabledTraceChannels();
+
+	EnableRequiredTraceChannels();
+
+	// Allow derived classes to perform additional actions before we start the traces
+	OnRecordingStartingInternal();
+
+	FTraceAuxiliary::FOptions TracingOptions;
+	TracingOptions.bExcludeTail = true;
+
+	if (InRecordingStartCommand.RecordingMode == ERecordingMode::File)
+	{
+		if (InRecordingStartCommand.TransportMode == ETraceTransportMode::FileSystem)
+		{
+			FString FileTarget = InRecordingStartCommand.Target;
+			if (FileTarget.IsEmpty())
+			{
+				GenerateRecordingFileName(FileTarget);
+			}
+
+			UE_LOG_REF(LogCategory, Log, TEXT("[%hs] Generated trace file name [%s]"), __func__, *FileTarget);
+
+			bIsRecording = FTraceAuxiliary::Start(FTraceAuxiliary::EConnectionType::File, *FileTarget, nullptr, &TracingOptions, LogCategory);
+		}
+		else
+		{
+			UE_LOG_REF(LogCategory, Error, TEXT("[%hs] Unsupported Transport mode [%s]"), __func__, *UEnum::GetValueAsString(InRecordingStartCommand.TransportMode))
+		}
+	}
+	else if (InRecordingStartCommand.RecordingMode == ERecordingMode::Live)
+	{
+		switch (InRecordingStartCommand.TransportMode)
+		{
+		case ETraceTransportMode::Direct:
+		case ETraceTransportMode::TraceServer:
+			bIsRecording = FTraceAuxiliary::Start(FTraceAuxiliary::EConnectionType::Network, *InRecordingStartCommand.Target, nullptr, &TracingOptions);
+			break;
+		case ETraceTransportMode::Relay:
+		{
+			if (!ensureMsgf(ExternalRelayExecutorDelegate.IsBound(), TEXT("Cannot start a trace in relay mode without a Relay Executor")))
+			{
+				UE_LOG_REF(LogCategory, Error, TEXT("[%hs] Cannot start a trace in relay mode without a Relay Executor"), __func__);
+				break;
+			}
+
+			if (RelayWriter && !RelayWriter->IsClosed())
+			{
+				UE_LOG_REF(LogCategory, Warning, TEXT("[%hs] Attempting to relay trace, but there is a relay writer still open"), __func__);
+				RelayWriter->Close();
+			}
+
+			RelayWriter = MakeUnique<FRelayTraceDataWriter>(LogCategory);
+
+			constexpr int32 MinBunchSize = 4096;
+			RelayWriter->SetMaxBytesPerBunch(FMath::Clamp(CVarMaxBytesPerRelayChunk->GetInt(), MinBunchSize, std::numeric_limits<int32>::max()));
+
+			constexpr int32 MinBufferSizeLimit = 2097152;
+			RelayWriter->SetMaxAllowedPendingBytes(FMath::Clamp(CVarMaxPendingBytesInRelayWriter->GetInt(), MinBufferSizeLimit, std::numeric_limits<int32>::max()));
+
+			bIsRecording = FTraceAuxiliary::Relay(reinterpret_cast<UPTRINT>(RelayWriter.Get()), FRelayTraceDataWriter::WriteHelper, FRelayTraceDataWriter::CloseHelper, nullptr, &TracingOptions);
+
+			if (bIsRecording)
+			{
+				SetupRelayDataPumpDelegates();
+			}
+
+			break;
+		}
+		case ETraceTransportMode::Invalid:
+		case ETraceTransportMode::FileSystem:
+			UE_LOG_REF(LogCategory, Error, TEXT("[%hs] Unsupported Transport mode [%s]"), __func__, *UEnum::GetValueAsString(InRecordingStartCommand.TransportMode));
+			break;
+		}
+	}
+	else
+	{
+		FailureReason = LOCTEXT("WrongCommandArgumentsError", "The start recording command was called with invalid arguments");
+	}
+#endif
+
+	AccumulatedRecordingTime = 0.0f;
+
+	if (ensureMsgf(bIsRecording, TEXT("Failed to start recording | Reason [%s]"), *FailureReason.ToString()))
+	{
+		CurrentTransportMode = InRecordingStartCommand.TransportMode;
+
+		if (ensureMsgf(!OnConnectionDelegateHandle.IsValid(), TEXT("Starting a trace while we are waiting for a pending connection")))
+		{
+			OnConnectionDelegateHandle = FTraceAuxiliary::OnConnection.AddRaw(this, &FRuntimeModule::HandleTraceConnectionEstablished);
+		}
+
+		// Start Listening for Trace Stopped events, in case Trace is stopped outside our control so we can gracefully stop recording and log a warning
+		if (ensureMsgf(!OnTraceStoppedDelegateHandle.IsValid(), TEXT("Starting a trace while we are waiting for a trace session to stop.")))
+		{
+			OnTraceStoppedDelegateHandle = FTraceAuxiliary::OnTraceStopped.AddRaw(this, &FRuntimeModule::HandleTraceStopRequest);
+		}
+	}
+	else
+	{
+		UE_LOG_REF(LogCategory, Error, TEXT("[%hs] Failed to start recording | Reason: [%s]"), __func__, *FailureReason.ToString());
+
+#if WITH_EDITOR
+		FMessageDialog::Open(EAppMsgType::Ok, FText::FormatOrdered(LOCTEXT("StartRecordingFailedMessage", "Failed to start recording. \n\n{0}"), FailureReason));
+#endif
+
+		RelayWriter = nullptr;
+
+		// Allow derived classes to perform additional actions on failure
+		OnRecordingStartFailedInternal(FailureReason);
+	}
+}
+
+void FRuntimeModule::StopRecording()
+{
+	FTraceAuxiliary::OnConnection.RemoveAll(this);
+	OnConnectionDelegateHandle = FDelegateHandle();
+
+	FTraceAuxiliary::OnTraceStopped.RemoveAll(this);
+	OnTraceStoppedDelegateHandle = FDelegateHandle();
+
+	CurrentTransportMode = ETraceTransportMode::Invalid;
+
+	if (!bIsRecording)
+	{
+		UE_LOG_REF(LogCategory, Warning, TEXT("[%hs] Attempted to stop recorded when there is no recording active."), __func__);
+		return;
+	}
+
+#if UE_TRACE_ENABLED || UE_TRACE_MINIMAL_ENABLED
+	RestoreTraceChannelsToPreRecordingState();
+
+	StopTrace();
+#endif
+
+	if (RecordingTimerHandle.IsValid())
+	{
+		FTSTicker::RemoveTicker(RecordingTimerHandle);
+		RecordingTimerHandle.Reset();
+	}
+
+	ClearRelayDataPumpDelegates();
+
+	if (RelayWriter)
+	{
+		RelayWriter->Close();
+	}
+
+	bIsRecording = false;
+	AccumulatedRecordingTime = 0.0f;
+
+	// Allow derived classes to perform additional actions
+	OnRecordingStoppedInternal();
+}
+
+
+int64 FRuntimeModule::GetBufferedDataBytesSize() const
+{
+	FGuardedInt64 RelaySendBufferSize(0);
+
+	if (RelayWriter)
+	{
+		RelaySendBufferSize += RelayWriter->GetQueuedBytesNum();
+	}
+
+	if (ExternalRelayBufferSizeProviderDelegate.IsBound())
+	{
+		RelaySendBufferSize += ExternalRelayBufferSizeProviderDelegate.Execute();
+	}
+
+	constexpr int64 DefaultValue = std::numeric_limits<int64>::max();
+	return RelaySendBufferSize.Get(DefaultValue);
+}
+
+FTraceConnectionDetails FRuntimeModule::GetCurrentTraceConnectionDetails() const
+{
+	FTraceConnectionDetails Details;
+	if (ExternalTraceStatusDelegate.IsBound())
+	{
+		return ExternalTraceStatusDelegate.Execute();
+	}
+
+	if (FTraceAuxiliary::IsConnected(Details.SessionGuid, Details.TraceGuid))
+	{
+		Details.TraceTarget = FTraceAuxiliary::GetTraceDestinationString();
+		Details.TransportMode = CurrentTransportMode;
+		Details.MarkAsValid();
+		return Details;
+	}
+
+	return Details;
+}
+
+void FRuntimeModule::StopTrace()
+{
+	bRequestedStop = true;
+	FTraceAuxiliary::Stop();
+}
+
+void FRuntimeModule::GenerateRecordingFileName(FString& OutFileName) const
+{
+	const FStringFormatOrderedArguments NameArgs
+	{
+		DebuggerName,
+		FString(FApp::GetProjectName())
+		, LexToString(FApp::GetBuildTargetType())
+		, FDateTime::Now().ToString(TEXT("%Y%m%d_%H%M%S%s"))
+	};
+
+	OutFileName = FString::Format(TEXT("{0}-{1}-{2}-{3}.utrace"), NameArgs);
+}
+
+void FRuntimeModule::RegisterExternalRelayExecutor(const FRelayTraceDataDelegate& InExternalRelayExecutor)
+{
+	if (!ensureMsgf(!ExternalRelayExecutorDelegate.IsBound(), TEXT("An external relay executor is already registered!")))
+	{
+		return;
+	}
+
+	ExternalRelayExecutorDelegate = InExternalRelayExecutor;
+}
+
+void FRuntimeModule::RegisterExternalRelayOverflowHandler(const FSimpleDelegate& InExternalRelayOverflowHandler)
+{
+	if (!ensureMsgf(!ExternalRelayOverflowHandlerDelegate.IsBound(), TEXT("An external relay overflow handler is already registered!")))
+	{
+		return;
+	}
+
+	ExternalRelayOverflowHandlerDelegate = InExternalRelayOverflowHandler;
+}
+
+void FRuntimeModule::UnregisterExternalRelayOverflowHandler()
+{
+	ExternalRelayOverflowHandlerDelegate.Unbind();
+}
+
+void FRuntimeModule::UnregisterCurrentExternalRelayExecutor()
+{
+	ExternalRelayExecutorDelegate.Unbind();
+}
+
+void FRuntimeModule::RegisterExternalTraceStatusProvider(const FExternalTraceConnectionStatusDelegate& InExternalTraceStatusCallback)
+{
+	ExternalTraceStatusDelegate = InExternalTraceStatusCallback;
+}
+
+void FRuntimeModule::UnregisterExternalTraceStatusProvider()
+{
+	ExternalTraceStatusDelegate.Unbind();
+}
+
+void FRuntimeModule::RegisterExternalRelayBufferSizeProvider(const FExternalRelayBufferQueueSizeDelegate& InExternalRelayBufferSizeCallback)
+{
+	if (!ensureMsgf(!ExternalRelayBufferSizeProviderDelegate.IsBound(), TEXT("An external relay InExternalRelayBufferSizeCallback is already registered!")))
+	{
+		return;
+	}
+
+	ExternalRelayBufferSizeProviderDelegate = InExternalRelayBufferSizeCallback;
+}
+
+void FRuntimeModule::UnregisterExternalRelayBufferSizeProvider()
+{
+	ExternalRelayBufferSizeProviderDelegate.Unbind();
+}
+
+void FRuntimeModule::HandleRelayOverflow()
+{
+	if (ExternalRelayOverflowHandlerDelegate.IsBound())
+	{
+		UE_LOG_REF(LogCategory, Error, TEXT("[%hs] Relay writer has overflowed. Executing external overflow handler... "), __func__);
+		ExternalRelayOverflowHandlerDelegate.Execute();
+		return;
+	}
+
+	UE_LOG_REF(LogCategory, Error, TEXT("[%hs] Relay writer has overflowed. Stopping Recording... "), __func__);
+	StopRecording();
+}
+
+void FRuntimeModule::RelayTraceData() const
+{
+	if (ensure(ExternalRelayExecutorDelegate.IsBound() && RelayWriter.IsValid()))
+	{
+		ExternalRelayExecutorDelegate.Execute(*RelayWriter);
+	}
+}
+
+void FRuntimeModule::SetupRelayDataPumpDelegates()
+{
+	if (!ensure(RelayWriter))
+	{
+		return;
+	}
+
+	RelayWriter->OnNewDataAvailable().BindRaw(this, &FRuntimeModule::RelayTraceData);
+	RelayWriter->OnDataOverflow().BindRaw(this, &FRuntimeModule::HandleRelayOverflow);
+}
+
+void FRuntimeModule::ClearRelayDataPumpDelegates()
+{
+	if (RelayWriter)
+	{
+		RelayWriter->OnNewDataAvailable().Unbind();
+		RelayWriter->OnDataOverflow().Unbind();
+	}
+}
+
+bool FRuntimeModule::RecordingTimerTick(float DeltaTime)
+{
+	if (bIsRecording)
+	{
+		AccumulatedRecordingTime += DeltaTime;
+	}
+
+	return true;
+}
+
+void FRuntimeModule::HandleTraceConnectionEstablished()
+{
+	UE_LOG_REF(LogCategory, Log, TEXT("Trace connection established."));
+
+	if (OnConnectionDelegateHandle.IsValid())
+	{
+		FTraceAuxiliary::OnConnection.RemoveAll(this);
+		OnConnectionDelegateHandle.Reset();
+	}
+
+	if (!ensureAlwaysMsgf(bIsRecording, TEXT("Received a trace connection established callback but no trace is active. This should not happen.")))
+	{
+		return;
+	}
+
+	// Allow derived classes to perform additional actions
+	OnRecordingStartedInternal();
+
+	TraceConnectionDetailsUpdatedDelegate.Broadcast();
+
+	// Add ticker to keep track of accumulated recording time
+	RecordingTimerHandle = FTSTicker::GetCoreTicker().AddTicker(FTickerDelegate::CreateRaw(this, &FRuntimeModule::RecordingTimerTick));
+}
+
+void FRuntimeModule::HandleTraceStopRequest(FTraceAuxiliary::EConnectionType TraceType, const FString& TraceDestination)
+{
+	if (bIsRecording)
+	{
+		if (!ensure(bRequestedStop))
+		{
+			UE_LOG_REF(LogCategory, Warning, TEXT("Trace Recording has been stopped unexpectedly"));
+
+#if WITH_EDITOR
+			FMessageDialog::Open(EAppMsgType::Ok, LOCTEXT("UnexpectedStopMessage", "Trace recording has been stopped unexpectedly. Debugger cannot continue with the recording session... "));
+#endif
+		}
+
+		StopRecording();
+	}
+
+	bRequestedStop = false;
+}
+
+bool FRuntimeModule::WaitForTraceSessionDisconnect()
+{
+	float MaxWaitTime = CVarMaxTimeToWaitForDisconnect->GetFloat();
+	float CurrentWaitTime = 0.0f;
+
+#if WITH_EDITOR
+	FScopedSlowTask DisconnectAttemptSlowTask(MaxWaitTime, LOCTEXT("DisconnectAttemptMessage", " Active Trace Session detected, attempting to disconnect ..."));
+
+	constexpr bool bShowCancelButton = false;
+	constexpr bool bAllowInPIE = true;
+	DisconnectAttemptSlowTask.MakeDialog(bShowCancelButton, bAllowInPIE);
+#endif
+
+	while (CurrentWaitTime < MaxWaitTime)
+	{
+		constexpr float WaitInterval = 0.1f;
+		FPlatformProcess::Sleep(0.1f);
+
+		if (!FTraceAuxiliary::IsConnected())
+		{
+			return true;
+		}
+
+		// We don't need to be precise for this, we can just accumulate the wait
+		CurrentWaitTime += WaitInterval;
+
+#if WITH_EDITOR
+		DisconnectAttemptSlowTask.EnterProgressFrame(CurrentWaitTime);
+#endif
+	}
+
+	return !FTraceAuxiliary::IsConnected();
+}
+
+void FRuntimeModule::SaveAndDisabledCurrentEnabledTraceChannels()
+{
+	// Until we support allowing other channels, indicate in the logs that we are disabling everything else
+	UE_LOG_REF(LogCategory, Log, TEXT("[%hs] Disabling additional trace channels..."), __func__);
+
+#if UE_TRACE_ENABLED || UE_TRACE_MINIMAL_ENABLED
+	OriginalTraceChannelsState.Reset();
+
+	// Disable any enabled additional channel
+	Trace::EnumerateChannels([](const ANSICHAR* ChannelName, bool bEnabled, void* SavedTraceChannelsPtr)
+		{
+			TMap<FString, bool>* SavedTraceChannels = static_cast<TMap<FString, bool>*>(SavedTraceChannelsPtr);
+			FString ChannelNameFString(ChannelName);
+			SavedTraceChannels->Add(ChannelNameFString, bEnabled);
+			if (bEnabled)
+			{
+				Trace::ToggleChannel(ChannelNameFString.GetCharArray().GetData(), false);
+			}
+		}, &OriginalTraceChannelsState);
+#endif
+}
+
+void FRuntimeModule::RestoreTraceChannelsToPreRecordingState()
+{
+#if UE_TRACE_ENABLED || UE_TRACE_MINIMAL_ENABLED
+	UE_LOG_REF(LogCategory, Log, TEXT("[%hs] Restoring trace channels state..."), __func__);
+
+	for (const TPair<FString, bool>& ChannelWithState : OriginalTraceChannelsState)
+	{
+		Trace::ToggleChannel(GetData(ChannelWithState.Key), ChannelWithState.Value);
+	}
+
+	OriginalTraceChannelsState.Reset();
+#endif
+}
+
+void FRuntimeModule::EnableRequiredTraceChannels()
+{
+#if UE_TRACE_ENABLED || UE_TRACE_MINIMAL_ENABLED
+	Trace::ToggleChannel(TEXT("Frame"), true);
+#endif
+
+	// Allow derived classes to enable their specific channels
+	OnEnableRequiredTraceChannelsInternal();
+}
+
+void FRuntimeModule::OnRecordingStartingInternal()
+{
+}
+
+void FRuntimeModule::OnRecordingStartedInternal()
+{
+}
+
+void FRuntimeModule::OnRecordingStartFailedInternal(const FText& FailureReason)
+{
+}
+
+void FRuntimeModule::OnRecordingStoppedInternal()
+{
+}
+
+void FRuntimeModule::OnEnableRequiredTraceChannelsInternal()
+{
+}
+
+} // UE::TraceBasedDebuggers
+
+#undef LOCTEXT_NAMESPACE
\ No newline at end of file
diff --git a/Engine/Source/Runtime/Experimental/TraceBasedDebuggers/Private/TraceBasedDebuggersModule.cpp b/Engine/Source/Runtime/Experimental/TraceBasedDebuggers/Private/TraceBasedDebuggersModule.cpp
new file mode 100644
--- /dev/null
+++ b/Engine/Source/Runtime/Experimental/TraceBasedDebuggers/Private/TraceBasedDebuggersModule.cpp
@@ -0,0 +1,5 @@
+// Copyright Epic Games, Inc. All Rights Reserved.
+
+#include "Modules/ModuleManager.h"
+
+IMPLEMENT_MODULE(FDefaultModuleImpl, TraceBasedDebuggers);
\ No newline at end of file
diff --git a/Engine/Source/Runtime/Experimental/TraceBasedDebuggers/Public/EngineEditorBridge.h b/Engine/Source/Runtime/Experimental/TraceBasedDebuggers/Public/EngineEditorBridge.h
new file mode 100644
--- /dev/null
+++ b/Engine/Source/Runtime/Experimental/TraceBasedDebuggers/Public/EngineEditorBridge.h
@@ -0,0 +1,112 @@
+// Copyright Epic Games, Inc. All Rights Reserved.
+
+#pragma once
+
+#if WITH_TRACE_BASED_DEBUGGERS
+
+#include "RemoteSessionsManager.h"
+#include "Templates/SharedPointer.h"
+
+class UGameInstance;
+
+#define UE_API TRACEBASEDDEBUGGERS_API
+
+#ifndef WITH_TRACE_BASED_DEBUGGERS_EXTERNAL_MESSAGING
+#define WITH_TRACE_BASED_DEBUGGERS_EXTERNAL_MESSAGING 0
+#endif
+
+namespace UE::TraceBasedDebuggers
+{
+
+/**
+ * Object that bridges the gap between a given debugger runtime module and the Engine & associated editor.
+ * As the debugger runtime module often does not have access to the engine module, this object reacts to events
+ * and performs necessary operations the runtime module cannot do directly.
+ */
+class FEngineEditorBridge
+{
+public:
+
+	FEngineEditorBridge() = delete;
+	UE_API explicit FEngineEditorBridge(const FLogCategoryBase& LogCategory);
+
+	virtual ~FEngineEditorBridge() = default;
+
+	UE_API void Initialize();
+	UE_API void TearDown();
+
+	UE_API static TSharedPtr<FRemoteSessionsManager> GetSessionsManager();
+
+	/** 
+	 * Derived classes can use this method to register a remote sessions handler that adds support
+	 * for more message types. 
+	 * Method will ensure if a handler is already registered.
+	 * The handler will be automatically unregistered from the session manager when the Engine-Editor bridge 
+	 * gets teared down.
+	 */
+	UE_API void RegisterRemoteSessionsHandler(TSharedRef<IRemoteSessionsHandler> Handler);
+
+	TSharedPtr<IDataRelayTransport> GetTraceRelayTransportInstance()
+	{
+		return RelayTraceDataTransportInstance;
+	}
+
+	bool IsInitialized() const
+	{
+		return bIsInitialized;
+	}
+
+	FSimpleMulticastDelegate& OnEngineEditorBridgeInitialized()
+	{
+		return InitializedDelegate;
+	}
+
+	UE_API void SetDataRelayTransportInstance(const TSharedPtr<IDataRelayTransport>& InExternalTraceRelayInstance);
+
+protected:
+	UE_API void HandleRecordingStarted();
+	UE_API void HandleRecordingStopped();
+	UE_API void BroadcastSessionStatus() const;
+
+	virtual void OnInitializeInternal()
+	{
+	}
+
+	virtual void OnTearDownInternal()
+	{
+	}
+
+	virtual void OnRecordingStartedInternal()
+	{
+	}
+
+	virtual void OnRecordingStoppedInternal()
+	{
+	}
+
+	virtual void BuildRecordingStatusInternal(FRecordingStatusMessage& OutStatusMessage) const
+	{
+	}
+
+private:
+	TSharedPtr<IDataRelayTransport> RelayTraceDataTransportInstance;
+
+	static TWeakPtr<FRemoteSessionsManager> SharedRemoteSessionsManager;
+	TSharedRef<FRemoteSessionsManager> RemoteSessionsManager;
+	/** Handler created by derived classes to handle more message types */
+	TSharedPtr<IRemoteSessionsHandler> RemoteSessionsHandler;
+
+	const FLogCategoryBase& LogCategory;
+
+	FTSTicker::FDelegateHandle RecordingStatusUpdateHandle;
+
+	FSimpleMulticastDelegate InitializedDelegate;
+
+	bool bIsInitialized = false;
+};
+
+} // UE::TraceBasedDebuggers
+
+#undef UE_API
+
+#endif // WITH_TRACE_BASED_DEBUGGERS
diff --git a/Engine/Source/Runtime/Experimental/TraceBasedDebuggers/Public/RelayTraceDataWriter.h b/Engine/Source/Runtime/Experimental/TraceBasedDebuggers/Public/RelayTraceDataWriter.h
new file mode 100644
--- /dev/null
+++ b/Engine/Source/Runtime/Experimental/TraceBasedDebuggers/Public/RelayTraceDataWriter.h
@@ -0,0 +1,125 @@
+// Copyright Epic Games, Inc. All Rights Reserved.
+
+#pragma once
+
+#if WITH_TRACE_BASED_DEBUGGERS
+#include <atomic>
+#include "Containers/Array.h"
+#include "Containers/SpscQueue.h"
+#include "Delegates/Delegate.h"
+#include "Serialization/BitWriter.h"
+#include "Templates/UniquePtr.h"
+
+#define UE_API TRACEBASEDDEBUGGERS_API
+
+class FBufferArchive;
+
+namespace UE::TraceBasedDebuggers
+{
+/**
+ * Writer object that chunks trace data into smaller bunches and queues them so a relay instance
+ * can dequeue and send them over the network later
+ */
+struct FRelayTraceDataWriter
+{
+	FRelayTraceDataWriter() = delete;
+	UE_API explicit FRelayTraceDataWriter(const FLogCategoryBase& LogCategory);
+
+	UE_API ~FRelayTraceDataWriter();
+
+	/**
+	 * Sets max number of bytes a queued bunch can have
+	 * @param MaxBytes Greater than zero number of bytes to set as limit
+	 */
+	UE_API void SetMaxBytesPerBunch(uint32 MaxBytes);
+
+	/**
+	 * Sets max number of bytes this writer can buffer before overflowing, triggering a close
+	 * @param MaxBytes number of bytes to set as limit
+	 */
+	UE_API void SetMaxAllowedPendingBytes(uint64 MaxBytes);
+
+	UE_API static bool WriteHelper(UPTRINT WriterHandle, const void* Data, uint32 Size);
+
+	static void CloseHelper(UPTRINT WriterHandle)
+	{
+		if (!WriterHandle)
+		{
+			return;
+		}
+
+		reinterpret_cast<FRelayTraceDataWriter*>(WriterHandle)->Close();
+	}
+
+	/**
+	 * Dequeues a bunch ready to send over the network, if any
+	 */
+	UE_API TOptional<TUniquePtr<FBufferArchive>> DequeuePendingBunch();
+
+	/**
+	 * Returns the size of the next bunch in the queue
+	 */
+	UE_API int32 NextPendingBunchSize() const;
+
+	/**
+	 * Returns true if we have any data pending to be sent over the network
+	 */
+	bool HasPendingBunches() const
+	{
+		return PendingBunchesCount > 0;
+	}
+
+	/**
+	 * returns the number of bytes overall that are in the queue waiting to be sent
+	 */
+	UE_API int64 GetQueuedBytesNum() const;
+
+	/**
+	 * Closes this Writer. After this is called, no new data will be processed, but any pending data can still be dequeued and sent over
+	 */
+	UE_API void Close();
+
+	/**
+	 * Returns true if this writer is closed and no new data is being processed
+	 */
+	bool IsClosed() const
+	{
+		return bClosed;
+	}
+
+	DECLARE_DELEGATE(FNewDataAvailableDelegate)
+	FNewDataAvailableDelegate& OnNewDataAvailable()
+	{
+		return NewDataAvailableDelegate;
+	}
+
+	FSimpleDelegate& OnDataOverflow()
+	{
+		return DataOverflowDelegate;
+	}
+
+private:
+
+	FNewDataAvailableDelegate NewDataAvailableDelegate;
+	FNewDataAvailableDelegate DataOverflowDelegate;
+
+	TUniquePtr<FBufferArchive> CreateBunch();
+
+	bool Write_Internal(const void* Data, uint32 Size);
+
+	bool CanWriteData(uint32 NewDataSize) const;
+
+	const FLogCategoryBase& LogCategory;
+	TSpscQueue<TUniquePtr<FBufferArchive>> PendingBunchesToRelay;
+	TUniquePtr<FBufferArchive> CurrentPendingBunch;
+	std::atomic<int32> PendingBunchesCount = 0;
+	std::atomic<bool> bClosed = false;
+	uint32 MaxPendingBytesPerBunch = 0;
+	std::atomic<int64> QueuedBytesNum = 0;
+	uint64 MaxAllowedPendingBytes = 1073741824;
+};
+
+} // UE::TraceBasedDebuggers
+
+#undef UE_API
+#endif // WITH_TRACE_BASED_DEBUGGERS
diff --git a/Engine/Source/Runtime/Experimental/TraceBasedDebuggers/Public/RemoteSessionsManager.h b/Engine/Source/Runtime/Experimental/TraceBasedDebuggers/Public/RemoteSessionsManager.h
new file mode 100644
--- /dev/null
+++ b/Engine/Source/Runtime/Experimental/TraceBasedDebuggers/Public/RemoteSessionsManager.h
@@ -0,0 +1,613 @@
+// Copyright Epic Games, Inc. All Rights Reserved.
+
+#pragma once
+
+#include "Containers/Ticker.h"
+#include "IMessageContext.h"
+#include "MessageEndpoint.h"
+#include "Misc/DateTime.h"
+#include "Misc/Guid.h"
+#include "Templates/SharedPointer.h"
+#include "TraceDataRelayTransport.h"
+#include "StructUtils/InstancedStruct.h"
+
+#include "RemoteSessionsManager.generated.h"
+
+#define UE_API TRACEBASEDDEBUGGERS_API
+
+class IMessageBus;
+class IMessageContext;
+
+UE_API DECLARE_LOG_CATEGORY_EXTERN(LogTraceBasedDebuggers, Log, All);
+
+namespace UE::TraceBasedDebuggers
+{
+
+struct FTraceConnectionDetails;
+
+USTRUCT(MinimalAPI)
+struct FSessionPing
+{
+	GENERATED_BODY()
+
+	UPROPERTY()
+	FGuid ControllerInstanceId;
+};
+
+USTRUCT(MinimalAPI)
+struct FSessionPong
+{
+	GENERATED_BODY()
+
+	UPROPERTY()
+	FGuid InstanceId;
+
+	UPROPERTY()
+	FGuid SessionId;
+
+	UPROPERTY()
+	FString SessionName;
+
+	UPROPERTY()
+	uint8 BuildTargetType = static_cast<uint8>(EBuildTargetType::Unknown);
+};
+
+USTRUCT(MinimalAPI)
+struct FStopRecordingCommandMessage
+{
+	GENERATED_BODY()
+};
+
+USTRUCT(MinimalAPI)
+struct FRecordingStatusMessage
+{
+	GENERATED_BODY()
+
+	UPROPERTY()
+	FGuid InstanceId = FGuid();
+
+	UPROPERTY()
+	FGuid DebuggerId = FGuid();
+
+	UPROPERTY()
+	bool bIsRecording = false;
+
+	UPROPERTY()
+	float ElapsedTime = false;
+
+	UPROPERTY()
+	int64 BufferedDataBytesSize = 0;
+};
+
+USTRUCT(MinimalAPI)
+struct FTraceConnectionDetailsMessage
+{
+	GENERATED_BODY()
+
+	UPROPERTY()
+	FGuid InstanceId = FGuid();
+
+	UPROPERTY()
+	FTraceConnectionDetails TraceDetails;
+};
+
+USTRUCT(MinimalAPI)
+struct FFullSessionInfoRequestMessage
+{
+	GENERATED_BODY()
+};
+
+USTRUCT(MinimalAPI)
+struct FFullSessionInfoResponseMessage
+{
+	GENERATED_BODY()
+
+	template <typename DataType>
+	DataType* GetDebuggerData()
+	{
+		for (FInstancedStruct& InstancedDebuggerData : DebuggerSpecificData)
+		{
+			if (DataType* DebuggerData = InstancedDebuggerData.GetMutablePtr<DataType>())
+			{
+				return DebuggerData;
+			}
+		}
+		return nullptr;
+	}
+
+	template <typename DataType>
+	const DataType* GetDebuggerData() const
+	{
+		for (const FInstancedStruct& InstancedDebuggerData : DebuggerSpecificData)
+		{
+			if (const DataType* DebuggerData = InstancedDebuggerData.GetPtr<DataType>())
+			{
+				return DebuggerData;
+			}
+		}
+		return nullptr;
+	}
+
+	UPROPERTY()
+	FGuid InstanceId;
+
+	UPROPERTY()
+	FGuid DebuggerId = FGuid();
+
+	UPROPERTY()
+	TArray<FInstancedStruct> DebuggerSpecificData;
+
+	UPROPERTY()
+	bool bIsRecording = false;
+};
+
+USTRUCT(MinimalAPI)
+struct FStartRecordingCommandMessage
+{
+	GENERATED_BODY()
+
+	UPROPERTY()
+	FString Target;
+
+	UPROPERTY()
+	ERecordingMode RecordingMode = ERecordingMode::Invalid;
+
+	UPROPERTY()
+	ETraceTransportMode TransportMode = ETraceTransportMode::Invalid;
+};
+
+UENUM()
+enum class ERemoteSessionAttributes
+{
+	None = 0,
+	SupportsDataChannelChange = 1 << 0,
+	CanExpire = 1 << 1,
+	IsMultiSessionWrapper = 1 << 2,
+};
+ENUM_CLASS_FLAGS(ERemoteSessionAttributes)
+
+UENUM()
+enum class ERemoteSessionReadyState : uint8
+{
+	/** The session is ready to execute commands */
+	Ready,
+	/** We are executing a command in the session we expect to take a while without hearing anything from the target. */
+	Busy
+};
+
+/**
+ * Session object that contains all the information needed to communicate with a remote instance, and the state of that instance
+ */
+struct FSessionInfo
+{
+	explicit FSessionInfo() : InstanceId(FGuid()),
+		SessionTypeAttributes(ERemoteSessionAttributes::CanExpire | ERemoteSessionAttributes::SupportsDataChannelChange)
+	{
+	}
+
+	virtual ~FSessionInfo() = default;
+
+protected:
+
+	explicit FSessionInfo(ERemoteSessionAttributes InSessionTypeAttributes) : SessionTypeAttributes(InSessionTypeAttributes)
+	{
+	}
+
+public:
+
+	template <typename DataType>
+	DataType* GetDebuggerData()
+	{
+		for (FInstancedStruct& InstancedDebuggerData : DebuggerSpecificSessionData)
+		{
+			if (DataType* DebuggerData = InstancedDebuggerData.GetMutablePtr<DataType>())
+			{
+				return DebuggerData;
+			}
+		}
+		return nullptr;
+	}
+
+	/** Append or update data associated to a given debugger type */
+	template <typename DataType>
+	void SetDebuggerData(DataType&& NewSessionData)
+	{
+		if (DataType* ChaosData = GetDebuggerData<DataType>())
+		{
+			*ChaosData = MoveTemp(NewSessionData);
+		}
+		else
+		{
+			DebuggerSpecificSessionData.Emplace(FInstancedStruct::Make(MoveTemp(NewSessionData)));
+		}
+	}
+
+	FGuid InstanceId;
+	FString SessionName;
+	FMessageAddress Address;
+	FDateTime LastPingTime;
+	EBuildTargetType BuildTargetType = EBuildTargetType::Unknown;
+	ERemoteSessionReadyState ReadyState = ERemoteSessionReadyState::Ready;
+
+	FRecordingStatusMessage LastKnownRecordingState;
+
+	UE_API const FTraceConnectionDetails& GetTraceConnectionDetails();
+
+	/** @return Whether any debugger, regardless of its type, is currently recording */
+	UE_API virtual bool IsAnyDebuggerRecording() const;
+
+	/** 
+	 * Indicates whether a given debugger type is currently recording.
+	 * @param DebuggerTypeId The identifier of the debugger type
+	 * @return Whether the provided debugger type id is currently recording
+	 */
+	UE_API virtual bool IsRecording(const FGuid& DebuggerTypeId) const;
+
+	UE_API virtual uint64 GetBufferedBytesNum() const;
+
+	ERemoteSessionAttributes GetSessionTypeAttributes() const
+	{
+		return SessionTypeAttributes;
+	}
+
+	UE_API void SetReceivedBytesPerSecond(uint64 InNewBytesPerSecond);
+	UE_API uint64 GetReceivedBytesPerSecond() const;
+
+protected:
+
+	TArray<FInstancedStruct> DebuggerSpecificSessionData;
+
+	FTraceConnectionDetails LastKnownConnectionDetails;
+
+	uint64 ReceivedBytesPerSecond = 0;
+
+	const ERemoteSessionAttributes SessionTypeAttributes;
+
+	friend struct FRemoteSessionsManager;
+};
+
+/**
+ * Session object that is able to control and provide information about multiple session objects.
+ * Used so the UI can use the same API to control multiple sessions like it does for single sessions.
+ */
+struct FMultiSessionInfo : FSessionInfo
+{
+	explicit FMultiSessionInfo() : FSessionInfo(ERemoteSessionAttributes::IsMultiSessionWrapper)
+	{
+	}
+
+	virtual ~FMultiSessionInfo() override = default;
+
+	UE_API virtual bool IsRecording(const FGuid& DebuggerTypeId) const override;
+	UE_API virtual bool IsAnyDebuggerRecording() const override;
+
+	UE_API virtual uint64 GetBufferedBytesNum() const override;
+
+	template<typename TCallback>
+	void EnumerateInnerSessions(const TCallback& Callback) const
+	{
+		for (const TPair<FGuid, TWeakPtr<FSessionInfo>>& InnerSessionWithID : InnerSessionsByInstanceID)
+		{
+			if (const TSharedPtr<FSessionInfo> SessionPtr = InnerSessionWithID.Value.Pin())
+			{
+				if (!Callback(SessionPtr.ToSharedRef()))
+				{
+					return;
+				}
+			}
+		}
+	}
+
+	TMap<FGuid, TWeakPtr<FSessionInfo>> InnerSessionsByInstanceID;
+};
+
+DECLARE_MULTICAST_DELEGATE_OneParam(FRecordingStateChangeDelegate, const TSharedPtr<FSessionInfo>&);
+DECLARE_MULTICAST_DELEGATE_OneParam(FSessionLifetimeDelegate, const TSharedPtr<FSessionInfo>&);
+
+struct IRemoteSessionsHandler : TSharedFromThis<IRemoteSessionsHandler>
+{
+	virtual ~IRemoteSessionsHandler() = default;
+
+	/** Callback to allow registered sessions handlers to add additional message handlers to the message endpoint builder */
+	virtual void OnCreatingMessageEndpoint(const TSharedPtr<FRemoteSessionsManager>&, FMessageEndpointBuilder&) = 0;
+
+	/** Callback to allow registered sessions handlers to subscribe to specific message types */
+	virtual void OnSubscribingMessageTypes(const TSharedPtr<UE::TraceBasedDebuggers::FRemoteSessionsManager>&, FMessageEndpoint&) = 0;
+
+	/** Callback to allow registered sessions handlers to build the response for a session information request */
+	virtual void OnBuildingFullSessionInfoResponseMessage(FFullSessionInfoResponseMessage&) = 0;
+
+	/** Callback to allow registered sessions handlers to process the response for a session information request */
+	virtual void OnHandlingFullSessionInfoResponseMessage(const FFullSessionInfoResponseMessage&, const TSharedPtr<FSessionInfo>&) = 0;
+};
+
+/** Object that is able to discover, issue and execute commands back and forth between the trace-based debuggers and client/server/editor instances */
+struct FRemoteSessionsManager : TSharedFromThis<FRemoteSessionsManager>
+{
+	UE_API FRemoteSessionsManager();
+	UE_API virtual ~FRemoteSessionsManager();
+
+	UE_API void ReInitializeMessagingSystem(const TSharedPtr<IMessageBus>& InMessageBus);
+
+	FSimpleMulticastDelegate& OnSessionsUpdated()
+	{
+		return SessionsUpdatedDelegate;
+	}
+
+	FSessionLifetimeDelegate& OnSessionDiscovered()
+	{
+		return SessionDiscoveredDelegate;
+	}
+
+	FSessionLifetimeDelegate& OnSessionExpired()
+	{
+		return SessionExpiredDelegate;
+	}
+	UE_API void RegisterExternalHandler(TSharedRef<IRemoteSessionsHandler> Handler);
+	UE_API void UnregisterExternalHandler(TSharedRef<IRemoteSessionsHandler> Handler);
+
+	typedef TFunction<void(const UScriptStruct*)> FVisitorFunction;
+
+	/**
+	 * Enumerate messages types allowed by this module.
+	 * @param InVisitor
+	 */
+	UE_API void EnumerateMessageTypes(const FVisitorFunction& InVisitor);
+
+	/**
+	 * Starts sending ping requests that potential targets (client, servers or other editors) can answer to and report themselves as available
+	 */
+	UE_API void StartSessionDiscovery();
+
+	/**
+	 * Stop sending ping requests to potential targets (client, servers or other editors).
+	 */
+	UE_API void StopSessionDiscovery();
+
+	/**
+	 * Issues a command to the provided address to change the state of a data channel
+	 * 	@param InDestinationAddress Message bus address that will execute the command
+	 * 	@param InMessage New data channel state to apply in the receiving instance
+	 */
+	template <typename T>
+	void SendCommand(const FMessageAddress& InDestinationAddress, const T& InMessage)
+	{
+		if (!MessageEndpoint)
+		{
+			UE_LOG(LogTraceBasedDebuggers, Error, TEXT("[%hs] Failed to send command | Invalid endpoint."), __FUNCTION__);
+			return;
+		}
+
+		MessageEndpoint->Send(
+			FMessageEndpoint::MakeMessage<T>(InMessage),
+			T::StaticStruct(),
+			EMessageFlags::Reliable,
+			nullptr,
+			TArrayBuilder<FMessageAddress>().Add(InDestinationAddress),
+			FTimespan::Zero(),
+			FDateTime::MaxValue());
+	}
+
+	/**
+	 * Returns the session info object for the provided ID
+	 * @param Id SessionID
+	 */
+	UE_API TWeakPtr<FSessionInfo> GetSessionInfo(FGuid Id);
+
+	/**
+	 * Returns the first available session info
+	 */
+	UE_API TWeakPtr<FSessionInfo> GetFirstAvailableSingleSessionInfo();
+
+	/**
+	 * Iterates through all active and valid sessions, and executes the provided callback to it.
+	 * if the callbacks returns false, the iteration will stop
+	* */
+	template <typename CallbackType>
+	void EnumerateActiveSessions(const CallbackType& Callback)
+	{
+		for (const TPair<FGuid, TSharedPtr<FSessionInfo>>& ActiveSession : ActiveSessionsById)
+		{
+			if (ActiveSession.Value)
+			{
+				bool bContinue = Callback(ActiveSession.Value.ToSharedRef());
+
+				if (!bContinue)
+				{
+					return;
+				}
+			}
+		}
+	}
+
+	/**
+	 * Broadcast to the network a given message
+	 * @param InMessage message to broadcast
+	 */
+	template <typename T>
+	void PublishMessage(const T& InMessage)
+	{
+		if (ensure(MessageEndpoint))
+		{
+			MessageEndpoint->Publish(FMessageEndpoint::MakeMessage<T>(InMessage), EMessageScope::Network);
+		}
+	}
+
+	/**
+	 * Delegate that broadcast when a recording was started in a session (either local or remote)
+	 */
+	FRecordingStateChangeDelegate& OnSessionRecordingStarted()
+	{
+		return RecordingStartedDelegate;
+	}
+
+	/**
+	 * Delegate that broadcast when a recording stops in a session (either local or remote)
+	 */
+	FRecordingStateChangeDelegate& OnSessionRecordingStopped()
+	{
+		return RecordingStoppedDelegate;
+	}
+
+	UE_API static const FString AllRemoteSessionsTargetName;
+	UE_API static const FString AllRemoteServersTargetName;
+	UE_API static const FString AllRemoteClientsTargetName;
+	UE_API static const FString AllSessionsTargetName;
+	UE_API static const FString CustomSessionsTargetName;
+	UE_API static const FString LocalEditorSessionName;
+	UE_API static const FName MessageBusEndpointName;
+	UE_API static const FGuid AllRemoteSessionsWrapperGUID;
+	UE_API static const FGuid AllRemoteServersWrapperGUID;
+	UE_API static const FGuid AllRemoteClientsWrapperGUID;
+	UE_API static const FGuid AllSessionsWrapperGUID;
+	UE_API static const FGuid CustomSessionsWrapperGUID;
+	UE_API static const FGuid InvalidSessionGUID;
+	UE_API static const FGuid LocalEditorSessionID;
+
+	DECLARE_MULTICAST_DELEGATE_TwoParams(FMessagingInitialized, TSharedPtr<IMessageBus> NewMessageBus, TSharedPtr<FMessageEndpoint> NewEndpoint);
+	FMessagingInitialized& OnMessagingInitialized()
+	{
+		return MessagingInitializedDelegate;
+	}
+
+	TWeakPtr<IMessageBus> GetMessageBusInstance()
+	{
+		return MessageBusPtr;
+	}
+
+	TWeakPtr<FMessageEndpoint> GetMessageBusEndPoint()
+	{
+		return MessageEndpoint;
+	}
+
+	/**
+	 * Returns true if this instance has controller capabilities (is either an editor or standalone application, which is also an editor)
+	 */
+	static constexpr bool IsController()
+	{
+#if WITH_EDITOR
+		return true;
+#else
+		return false;
+#endif
+	}
+
+	/**
+	 * Adds a message type to the list of supported message types for this session system.
+	 * @param ScriptStruct Message type info
+	 */
+	UE_API void RegisterExternalSupportedMessageType(const UScriptStruct* ScriptStruct);
+
+	/**
+	 * Adds a message type to the list of supported message types for this session system.
+	 */
+	template <typename TDataType>
+	void RegisterExternalSupportedMessageType()
+	{
+		SupportedMessageTypes.Emplace(TBaseStructure<TDataType>::Get());
+	}
+
+private:
+	void Initialize();
+	void Shutdown();
+
+	void RegisterBuiltInMessageTypes();
+
+	void InitializeMessagingSystem(const TSharedPtr<IMessageBus>& InMessageBus);
+	void ShutdownMessagingSystem();
+
+	/**
+	 * Sends a request to obtain the full session information to the provided message bus address
+	 * @param InDestinationAddress
+	 */
+	void SendFullSessionStateRequestCommand(const FMessageAddress& InDestinationAddress);
+
+	/**
+	 * Registers a session object that is able to control multiple session instances
+	 * @param InSessionInfoRef Multi Session object
+	 */
+	void RegisterSessionInMultiSessionWrapper(const TSharedRef<FSessionInfo>& InSessionInfoRef);
+
+	/**
+	 * Deregisters a session object that is able to control multiple session instances
+	 * @param InSessionInfoRef Multi Session object
+	 */
+	void DeRegisterSessionInMultiSessionWrapper(const TSharedRef<FSessionInfo>& InSessionInfoRef);
+
+	/**
+	 * Creates a session object that is able to control multiple other session objects.
+	 * @param InstanceId Instance ID for this object
+	 * @param SessionName Session name that will be used in the UI
+	 * @return
+	 */
+	TSharedPtr<FSessionInfo> CreatedWrapperSessionInfo(FGuid InstanceId, FStringView SessionName);
+
+	/**
+	 * Creates the message bus endpoint this session manager will use
+	 */
+	TSharedPtr<FMessageEndpoint> CreateEndPoint(const TSharedRef<IMessageBus>& InMessageBus);
+
+	bool Tick(float DeltaTime);
+
+	/**
+	 * Usually used by the controller. Broadcast to the network this controller exists.
+	 */
+	void SendPing();
+
+	/**
+	 * Broadcast to the network a small subset of this instance information, in to a received session ping.
+	 * @param InMessage Ping message that was received
+	 */
+	void SendPong(const FSessionPing& InMessage);
+
+	void HandleSessionPongMessage(const FSessionPong& InMessage, const TSharedRef<IMessageContext>& InContext);
+	void HandleSessionPingMessage(const FSessionPing& InMessage, const TSharedRef<IMessageContext>& InContext);
+	void HandleRecordingStatusUpdateMessage(const FRecordingStatusMessage& InMessage, const TSharedRef<IMessageContext>& InContext);
+	void HandleConnectionDetailsUpdateMessage(const FTraceConnectionDetailsMessage& InMessage, const TSharedRef<IMessageContext>& InContext);
+	void HandleFullSessionStateRequestMessage(const FFullSessionInfoRequestMessage& InMessage, const TSharedRef<IMessageContext>& InContext);
+	void HandleFullSessionStateResponseMessage(const FFullSessionInfoResponseMessage& InMessage, const TSharedRef<IMessageContext>& InContext);
+
+	enum class ERemoveSessionOptions : uint8
+	{
+		None = 0,
+		ForceRemoveAll = 1 << 0
+	};
+	FRIEND_ENUM_CLASS_FLAGS(ERemoveSessionOptions)
+
+	void RemoveExpiredSessions(ERemoveSessionOptions Options = ERemoveSessionOptions::None);
+
+	void ProcessPendingMessagesForSession(const FSessionPong& InMessage, const TSharedRef<FSessionInfo>& InSessionInfoPtr);
+
+	TArray<TSharedRef<IRemoteSessionsHandler>> Handlers;
+
+	/** Holds the time at which the last ping was sent. */
+	FDateTime LastPingTime;
+
+	/** Holds a pointer to the message bus. */
+	TWeakPtr<IMessageBus> MessageBusPtr;
+
+	/** Holds the messaging endpoint. */
+	TSharedPtr<FMessageEndpoint> MessageEndpoint;
+
+	TMap<FGuid, TSharedPtr<FSessionInfo>> ActiveSessionsById;
+	TMap<FGuid, FRecordingStatusMessage> PendingRecordingStatusMessages;
+	TMap<FGuid, FTraceConnectionDetails> PendingRecordingConnectionDetailsMessages;
+
+	FSimpleMulticastDelegate SessionsUpdatedDelegate;
+	FSessionLifetimeDelegate SessionDiscoveredDelegate;
+	FSessionLifetimeDelegate SessionExpiredDelegate;
+
+	FRecordingStateChangeDelegate RecordingStartedDelegate;
+	FRecordingStateChangeDelegate RecordingStoppedDelegate;
+
+	FTSTicker::FDelegateHandle TickHandle;
+
+	FMessagingInitialized MessagingInitializedDelegate;
+
+	TArray<const UScriptStruct*> SupportedMessageTypes;
+
+	bool bInitialized = false;
+};
+
+} // namespace UE::TraceBasedDebuggers
+#undef UE_API
\ No newline at end of file
diff --git a/Engine/Source/Runtime/Experimental/TraceBasedDebuggers/Public/TraceBasedDebuggerRuntime.h b/Engine/Source/Runtime/Experimental/TraceBasedDebuggers/Public/TraceBasedDebuggerRuntime.h
new file mode 100644
--- /dev/null
+++ b/Engine/Source/Runtime/Experimental/TraceBasedDebuggers/Public/TraceBasedDebuggerRuntime.h
@@ -0,0 +1,230 @@
+// Copyright Epic Games, Inc. All Rights Reserved.
+
+#pragma once
+
+#include "Containers/Ticker.h"
+#include "ProfilingDebugging/TraceAuxiliary.h"
+#include "TraceDataRelayTransport.h"
+
+#define UE_API TRACEBASEDDEBUGGERS_API
+
+namespace UE::TraceBasedDebuggers
+{
+struct FRelayTraceDataWriter;
+struct FStartRecordingCommandMessage;
+
+DECLARE_MULTICAST_DELEGATE(FRecordingStateChangedDelegate)
+DECLARE_MULTICAST_DELEGATE_OneParam(FRecordingStartFailedDelegate, const FText&)
+DECLARE_DELEGATE_OneParam(FRelayTraceDataDelegate, FRelayTraceDataWriter&)
+DECLARE_DELEGATE_RetVal(FTraceConnectionDetails, FExternalTraceConnectionStatusDelegate)
+DECLARE_DELEGATE_RetVal(uint64, FExternalRelayBufferQueueSizeDelegate)
+
+class FRuntimeModule
+{
+protected:
+	FRuntimeModule() = delete;
+	UE_API explicit FRuntimeModule(const FLogCategoryBase& InLogCategory, const FGuid& InDebuggerTypeID);
+
+	UE_API virtual ~FRuntimeModule();
+
+	/**
+	 * Parses the command line arguments for <AutoRecordCommand>, if set, then start recording if the command is found.
+	 * @return Whether the recording was started.
+	 */
+	UE_API bool TryRecordOnStartup();
+
+public:
+	/**
+	 * Starts a recording by starting a Trace session. It will stop any existing trace session.
+	 * @param Args : Arguments array usually provided by the commandline. Used to determine if we want to record to file or a local trace server.
+	 * Available options are:
+	 * - "" (default): trace to file using GenerateRecordingFileName
+	 * - "File": trace to file using GenerateRecordingFileName
+	 * - "-tracefile [FILENAME]": trace to file using either the optional [FILENAME] or GenerateRecordingFileName
+	 * - "Server [ADDRESS]": trace server using either the optional [ADDRESS] or local host (i.e., "127.0.0.1")
+	 * - "-tracehost [ADDRESS]": trace server using either the optional [ADDRESS] or local host (i.e., "127.0.0.1")
+	 * @see GenerateRecordingFileName
+	 */
+	UE_API void StartRecording(TConstArrayView<FString> Args);
+
+	/**
+	 * Starts a recording by starting a Trace session. It will stop any existing trace session
+	 * @param InRecordingStartCommand : Used to determine if we want to record to file or a local trace server among other settings
+	 */
+
+	UE_API void StartRecording(const FStartRecordingCommandMessage& InRecordingStartCommand);
+
+	/* Stops an active recording */
+	UE_API void StopRecording();
+
+	/** Returns true if we are currently recording trace data for the current debugger */
+	bool IsRecording() const
+	{
+		return bIsRecording;
+	}
+
+	/** @return Unique identifier for the current debugger type */
+	const FGuid& GetDebuggerId() const
+	{
+		return DebuggerGuid;
+	}
+
+	/** 
+	 * Finds a valid file name for a new file for Trace recording.
+	 * The file name is built using the following format:
+	 * <DebuggerName>-<ProjectName>-<BuildTargetType>-<DateTime>
+	 * where <DebuggerName> is provided by the derived classes.
+	 */
+	UE_API void GenerateRecordingFileName(FString& OutFileName) const;
+
+	/** Return the current connection details for the underlying trace session */
+	UE_API FTraceConnectionDetails GetCurrentTraceConnectionDetails() const;
+
+	/** Returns the accumulated recording time in seconds since the recording started */
+	float GetAccumulatedRecordingTime() const
+	{
+		return AccumulatedRecordingTime;
+	}
+
+	/** Returns the accumulated current send buffer size used by the relay transport system - This should be used for logging or UI feedback only */
+	UE_API int64 GetBufferedDataBytesSize() const;
+
+	FSimpleMulticastDelegate& OnTraceConnectionDetailsUpdated()
+	{
+		return TraceConnectionDetailsUpdatedDelegate;
+	}
+
+	/**
+	 * Registers an external handler that will be in charge of relaying the traced data to a remote instance if the relay transport mode is selected
+	 * @param InExternalRelayExecutor Delegate that will be executed to relay data to a remote debugger instance, periodically
+	 */
+	UE_API void RegisterExternalRelayExecutor(const FRelayTraceDataDelegate& InExternalRelayExecutor);
+
+	/**
+	 * Registers an external handler that will be in charge handling what happens when the Relay writer overflows
+	 * @param InExternalRelayOverflowHandler Delegate that will be executed when the Relay writer overflows
+	 */
+	UE_API void RegisterExternalRelayOverflowHandler(const FSimpleDelegate& InExternalRelayOverflowHandler);
+
+	/**
+	 * Unbinds the delegate to an external relay overflow handler
+	 */
+	UE_API void UnregisterExternalRelayOverflowHandler();
+
+	/**
+	 * Unbinds the delegate to an external relay executor
+	 */
+	UE_API void UnregisterCurrentExternalRelayExecutor();
+
+	/**
+	 * Registers an external handler that will be in charge of providing the trace status if the recording mode is set to Relay
+	 * @param InExternalTraceStatusCallback Delegate that will be executed to get the trace status, periodically
+	 */
+	UE_API void RegisterExternalTraceStatusProvider(const FExternalTraceConnectionStatusDelegate& InExternalTraceStatusCallback);
+
+	/**
+	 * Unbinds the delegate to the current external trace status provider
+	 */
+	UE_API void UnregisterExternalTraceStatusProvider();
+
+	/**
+	 *  Registers an external handler that will provide information about any external send buffer the current relay transport implementation might be using, if any.
+	 *  This is used for UI metrics. To handle buffer overflows, use RegisterExternalRelayOverflowHandler
+	 * @param InExternalRelayBufferSizeCallback
+	 */
+	UE_API void RegisterExternalRelayBufferSizeProvider(const FExternalRelayBufferQueueSizeDelegate& InExternalRelayBufferSizeCallback);
+
+	/**
+	 * Unbinds the delegate to the current external Relay buffer size provider
+	 */
+	UE_API void UnregisterExternalRelayBufferSizeProvider();
+
+	FRelayTraceDataWriter* GetRelayTraceWriterInstance() const
+	{
+		return RelayWriter.Get();
+	}
+
+protected:
+	UE_API virtual void OnRecordingStartingInternal();
+	UE_API virtual void OnRecordingStartedInternal();
+	UE_API virtual void OnRecordingStartFailedInternal(const FText& FailureReason);
+	UE_API virtual void OnRecordingStoppedInternal();
+
+	UE_API virtual void OnEnableRequiredTraceChannelsInternal();
+
+	/**
+	 * Prefix provided by the derived classes used by GenerateRecordingFileName 
+	 * using the following format "<DebuggerName>-<ProjectName>-<BuildTargetName>-<DateTime>.utrace"
+	 */
+	FString DebuggerName;
+
+	/**
+	 * Command line provided by the derived classes used during initialization to start recording or not.
+	 */
+	FString AutoRecordCommand;
+
+	/**
+	 * Command line that derived classes can override and used during initialization to specify the address of the host.
+	 * e.g., -<HostAddressCommand>=x.x.x.x where 'HostAddressCommand' is set to 'DebuggerHost' by default.
+	 */
+	FString HostAddressCommand = TEXT("DebuggerHost");
+
+	const FLogCategoryBase& LogCategory;
+
+private:
+
+	void HandleRelayOverflow();
+
+	void RelayTraceData() const;
+
+	void SetupRelayDataPumpDelegates();
+	void ClearRelayDataPumpDelegates();
+
+	/** Stops the current Trace session */
+	void StopTrace();
+
+	/** Queues a full Capture of the simulation on the next frame */
+	bool RecordingTimerTick(float DeltaTime);
+
+	void HandleTraceConnectionEstablished();
+
+	/**
+	 * Used to handle stop requests to the active trace session that were not done by us
+	 * That is a possible scenario because Trace is shared by other In-Editor tools
+	 */
+	void HandleTraceStopRequest(FTraceAuxiliary::EConnectionType TraceType, const FString& TraceDestination);
+
+	bool WaitForTraceSessionDisconnect();
+
+	void EnableRequiredTraceChannels();
+	void SaveAndDisabledCurrentEnabledTraceChannels();
+	void RestoreTraceChannelsToPreRecordingState();
+
+	FTSTicker::FDelegateHandle RecordingTimerHandle;
+	FTSTicker::FDelegateHandle RelayExecutorTickerHandle;
+
+	FDelegateHandle OnConnectionDelegateHandle;
+	FDelegateHandle OnTraceStoppedDelegateHandle;
+
+	TUniquePtr<FRelayTraceDataWriter> RelayWriter;
+
+	FSimpleMulticastDelegate TraceConnectionDetailsUpdatedDelegate;
+	FSimpleDelegate ExternalRelayOverflowHandlerDelegate;
+	FRelayTraceDataDelegate ExternalRelayExecutorDelegate;
+	FExternalRelayBufferQueueSizeDelegate ExternalRelayBufferSizeProviderDelegate;
+	FExternalTraceConnectionStatusDelegate ExternalTraceStatusDelegate;
+
+	TMap<FString, bool> OriginalTraceChannelsState;
+
+	FGuid DebuggerGuid;
+
+	float AccumulatedRecordingTime = 0.0f;
+
+	ETraceTransportMode CurrentTransportMode = ETraceTransportMode::Invalid;
+	bool bIsRecording = false;
+	bool bRequestedStop = false;
+};
+
+} // UE::TraceBasedDebuggers
+
+#undef UE_API
\ No newline at end of file
diff --git a/Engine/Source/Runtime/Experimental/TraceBasedDebuggers/Public/TraceDataRelayTransport.h b/Engine/Source/Runtime/Experimental/TraceBasedDebuggers/Public/TraceDataRelayTransport.h
new file mode 100644
--- /dev/null
+++ b/Engine/Source/Runtime/Experimental/TraceBasedDebuggers/Public/TraceDataRelayTransport.h
@@ -0,0 +1,137 @@
+// Copyright Epic Games, Inc. All Rights Reserved.
+
+#pragma once
+
+#include "TraceDataRelayTransport.generated.h"
+
+struct FGuid;
+
+namespace UE::TraceBasedDebuggers
+{
+
+UENUM()
+enum class ERecordingMode : uint8
+{
+	Invalid,
+	Live,
+	File
+};
+
+/** Available transport modes for the Trace Data produced by the trace-based debuggers */
+UENUM()
+enum class ETraceTransportMode : uint8
+{
+	Invalid UMETA(Hidden),
+	/** The data will be traced directly to a file in the local file system */
+	FileSystem UMETA(Hidden),
+	/** Data will be traced to the selected trace store server */
+	TraceServer UMETA(Hidden),
+	/** Data will be traced directly to the editor, accessible via sockets */
+	Direct,
+	/** Data will be traced via the trace relay system using a custom transport, like normal UE Networking or a direct socket connection */
+	Relay
+};
+
+/** Structure containing the info necessary to connect to a trace session and start receiving data */
+USTRUCT()
+struct FTraceConnectionDetails
+{
+	GENERATED_BODY()
+
+	FTraceConnectionDetails() = default;
+
+	void MarkAsValid()
+	{
+		bIsValid = true;
+	}
+
+	bool IsValid() const
+	{
+		return bIsValid;
+	}
+
+	/** ID used to find the active trace session in the trace system */
+	UPROPERTY()
+	FGuid TraceGuid;
+
+	/** ID from the trace-based debugger remote session from which this connection details come from */
+	UPROPERTY()
+	FGuid SessionGuid;
+
+	/** Address of the trace server or editor in direct trace mode that will receive the trace data */
+	UPROPERTY()
+	FString TraceTarget;
+
+	/*** Port number used for the trace connection, if any*/
+	UPROPERTY()
+	uint16 Port = 0;
+
+	/** How the data is being transported to the editor (Trace Server, Custom Relay, Direct, etc.)*/
+	UPROPERTY()
+	ETraceTransportMode TransportMode = ETraceTransportMode::Invalid;
+
+	/** Encoded SSL certificate for the trace connection, if any */
+	UPROPERTY()
+	TArray<uint8> CertAuth;
+
+	UPROPERTY()
+	bool bIsValid = false;
+};
+
+
+/** Structure containing the details necessary to connect to a relay instance */
+USTRUCT()
+struct FRelayConnectionInfo
+{
+	GENERATED_BODY()
+
+	/** Relay IP address */
+	UPROPERTY()
+	FString Address;
+
+	/** Relay port */
+	UPROPERTY()
+	uint16 Port = 0;
+
+	/** Encoded SSL certificate */
+	UPROPERTY()
+	TArray<uint8> CertificateAuthority;
+};
+
+#if WITH_TRACE_BASED_DEBUGGERS
+
+struct FRelayTraceDataWriter;
+
+/** Possible connection attempt results */
+enum class EConnectionAttemptResult : uint8
+{
+	NotStarted,
+	InProgress,
+	Success,
+	Failed
+};
+
+/**
+ * Interface for any trace-based debugger relay transport implementation
+ */
+class IDataRelayTransport
+{
+public:
+	virtual ~IDataRelayTransport() = default;
+
+	virtual void Initialize() = 0;
+	virtual void Shutdown() = 0;
+
+	DECLARE_DELEGATE_OneParam(FProcessReceivedRelayDataDelegate, const TConstArrayView<uint8> DataPacket);
+	virtual void RegisterRelayDataReceiverForSessionID(FGuid RemoteSessionID, const FProcessReceivedRelayDataDelegate& InProcessDataDelegate) = 0;
+	virtual void UnregisterRelayDataReceiverForSessionID(FGuid RemoteSessionID) = 0;
+
+	virtual void RelayTraceDataFromWriter(FRelayTraceDataWriter& InRelayTraceDataWriter) = 0;
+
+	virtual FRelayConnectionInfo GetConnectionInfo() = 0;
+	virtual EConnectionAttemptResult ConnectToRelay(FGuid RemoteSessionID, const FRelayConnectionInfo ConnectionInfo) = 0;
+	virtual EConnectionAttemptResult GetConnectionAttemptResult(FGuid RemoteSessionID) = 0;
+};
+
+#endif // WITH_TRACE_BASED_DEBUGGERS
+} // namespace UE::TraceBasedDebuggers
\ No newline at end of file
diff --git a/Engine/Source/Runtime/Experimental/TraceBasedDebuggers/TraceBasedDebuggers.Build.cs b/Engine/Source/Runtime/Experimental/TraceBasedDebuggers/TraceBasedDebuggers.Build.cs
new file mode 100644
--- /dev/null
+++ b/Engine/Source/Runtime/Experimental/TraceBasedDebuggers/TraceBasedDebuggers.Build.cs
@@ -0,0 +1,33 @@
+// Copyright Epic Games, Inc. All Rights Reserved.
+
+namespace UnrealBuildTool.Rules
+{
+	public class TraceBasedDebuggers : ModuleRules
+	{
+		public TraceBasedDebuggers(ReadOnlyTargetRules Target) : base(Target)
+		{
+			PublicDependencyModuleNames.AddRange(
+				new string[]
+				{
+					"Core",
+					"CoreUObject",
+					"Messaging",
+					"MessagingCommon",
+					"TraceLog"
+				}
+				);
+
+			bAllowUETypesInNamespaces = true;
+			CppCompileWarningSettings.UnsafeTypeCastWarningLevel = WarningLevel.Error;
+
+			if (AreTraceBasedDebuggersSupported(Target))
+			{
+				PublicDefinitions.Add("WITH_TRACE_BASED_DEBUGGERS=1");
+			}
+			else
+			{
+				PublicDefinitions.Add("WITH_TRACE_BASED_DEBUGGERS=0");
+			}
+		}
+	}
+}
