Fix non-unity errors
#rnx

diff --git a/Engine/Plugins/ChaosVD/Source/ChaosVD/Private/Trace/DataProcessors/ChaosVDCollisionChannelsInfoDataProcessor.cpp b/Engine/Plugins/ChaosVD/Source/ChaosVD/Private/Trace/DataProcessors/ChaosVDCollisionChannelsInfoDataProcessor.cpp
--- a/Engine/Plugins/ChaosVD/Source/ChaosVD/Private/Trace/DataProcessors/ChaosVDCollisionChannelsInfoDataProcessor.cpp
+++ b/Engine/Plugins/ChaosVD/Source/ChaosVD/Private/Trace/DataProcessors/ChaosVDCollisionChannelsInfoDataProcessor.cpp
@@ -3,6 +3,7 @@
 #include "ChaosVDCollisionChannelsInfoDataProcessor.h"
 
 #include "ChaosVDModule.h"
+#include "Containers/Ticker.h"
 
 
 FChaosVDCollisionChannelsInfoDataProcessor::FChaosVDCollisionChannelsInfoDataProcessor() : FChaosVDDataProcessorBase(FChaosVDCollisionChannelsInfoContainer::WrapperTypeName)
diff --git a/Engine/Source/Runtime/Experimental/Chaos/Public/GeometryCollection/ManagedArrayTypes.h b/Engine/Source/Runtime/Experimental/Chaos/Public/GeometryCollection/ManagedArrayTypes.h
--- a/Engine/Source/Runtime/Experimental/Chaos/Public/GeometryCollection/ManagedArrayTypes.h
+++ b/Engine/Source/Runtime/Experimental/Chaos/Public/GeometryCollection/ManagedArrayTypes.h
@@ -11,6 +11,7 @@
 #include "Math/Vector.h"
 #include "Math/Vector2D.h"
 #include "Math/Box.h"
+#include "UObject/SoftObjectPath.h"
 
 #include "Chaos/ImplicitObject.h"
 #include "Chaos/BVHParticles.h"
