[Backout] - CL49127187
#rnx
Original CL Desc
-----------------------------------------------------------------
[TraceBasedDebuggers] refactored and moved TraceSessionsManager from ChaosVD
Main changes:
- settings are now Attributes to assign/bind when creating the manager
- sessions manager creation now requires explicit references since this is no longer setup as a group of related singletons:
- debugger trace runtime module
- engine-editor bridge
- trace module (TraceServices::IModule modular feature)
- FChaosVDTraceSessionDescriptor replaced by UE::TraceBasedDebuggers::FTraceSessionDescriptor
- moved DataStream related types to (FSessionStreamStatsUpdater, IRemoteSessionDataStream, FDirectSocketStream, FRelayDataStream, etc.) under namespace UE::TraceBasedDebuggers

#Synced-CL 49027723

#Synced-CL 49127187

diff --git a/Engine/Plugins/ChaosVD/Source/ChaosVD/ChaosVD.Build.cs b/Engine/Plugins/ChaosVD/Source/ChaosVD/ChaosVD.Build.cs
--- a/Engine/Plugins/ChaosVD/Source/ChaosVD/ChaosVD.Build.cs
+++ b/Engine/Plugins/ChaosVD/Source/ChaosVD/ChaosVD.Build.cs
@@ -34,7 +34,6 @@ public class ChaosVD : ModuleRules
 				"WorkspaceMenuStructure",
 				"TraceAnalysis", 
 				"TraceBasedDebuggers",
-				"TraceBasedDebuggersAnalysis",
 				"TraceInsights",
 				"TraceLog", 
 				"GeometryFramework",
diff --git a/Engine/Plugins/ChaosVD/Source/ChaosVD/Private/ChaosVDEngine.cpp b/Engine/Plugins/ChaosVD/Source/ChaosVD/Private/ChaosVDEngine.cpp
--- a/Engine/Plugins/ChaosVD/Source/ChaosVD/Private/ChaosVDEngine.cpp
+++ b/Engine/Plugins/ChaosVD/Source/ChaosVD/Private/ChaosVDEngine.cpp
@@ -11,20 +11,10 @@
 #include "RemoteSessionsManager.h"
 #include "Settings/ChaosVDMiscSettings.h"
 #include "Trace/ChaosVDCombinedTraceFile.h"
-#include "TraceSessionsManager.h"
 #include "Trace/ChaosVDTraceManager.h"
 
 #include UE_INLINE_GENERATED_CPP_BY_NAME(ChaosVDEngine)
 
-FChaosVDEngine::FChaosVDEngine()
-{
-	InstanceGUID = FGuid::NewGuid();
-}
-
-FChaosVDEngine::~FChaosVDEngine()
-{
-}
-
 void FChaosVDEngine::Initialize()
 {
 	if (bIsInitialized)
@@ -39,10 +29,9 @@ void FChaosVDEngine::Initialize()
 	PlaybackController = MakeShared<FChaosVDPlaybackController>(CurrentScene);
 	PlaybackController->Initialize();
 
-	using namespace UE::TraceBasedDebuggers;
-	if (const TSharedPtr<FRemoteSessionsManager> RemoteSessionManager = FChaosVDEngineEditorBridge::Get().GetSessionsManager())
+	if (TSharedPtr<UE::TraceBasedDebuggers::FRemoteSessionsManager> RemoteSessionManager = FChaosVDEngineEditorBridge::Get().GetSessionsManager())
 	{
-		LiveSessionStoppedDelegateHandle = RemoteSessionManager->OnSessionRecordingStopped().AddLambda([WeakThis = AsWeak()](TWeakPtr<FSessionInfo> Session)
+		LiveSessionStoppedDelegateHandle = RemoteSessionManager->OnSessionRecordingStopped().AddLambda([WeakThis = AsWeak()](TWeakPtr<UE::TraceBasedDebuggers::FSessionInfo> Session)
 		{
 			if (const TSharedPtr<FChaosVDEngine> CVDEngine = WeakThis.Pin())
 			{
@@ -66,12 +55,11 @@ void FChaosVDEngine::Initialize()
 
 void FChaosVDEngine::CloseActiveTraceSessions()
 {
-	using namespace UE::TraceBasedDebuggers;
-	if (const TSharedPtr<FTraceSessionsManager> CVDTraceManager = FChaosVDModule::Get().GetTraceSessionsManager())
+	if (const TSharedPtr<FChaosVDTraceManager> CVDTraceManager = FChaosVDModule::Get().GetTraceManager())
 	{
-		for (FTraceSessionDescriptor& SessionDescriptor : CurrentSessionDescriptors)
+		for (FChaosVDTraceSessionDescriptor& SessionDescriptor : CurrentSessionDescriptors)
 		{
-			OnTraceSessionClosed().Broadcast(SessionDescriptor);
+			OnSessionClosed().Broadcast(SessionDescriptor);
 
 			CVDTraceManager->CloseSession(SessionDescriptor.SessionName);
 			SessionDescriptor.bIsLiveSession = false;
@@ -88,20 +76,19 @@ void FChaosVDEngine::CloseActiveTraceSessions()
 
 void FChaosVDEngine::CloseSessionByRemoteSessionID(FGuid RemoteSessionID)
 {
-	using namespace UE::TraceBasedDebuggers;
-	FTraceSessionDescriptor DescriptorFoundCopy = GetTraceSessionDescriptor(RemoteSessionID);
+	FChaosVDTraceSessionDescriptor DescriptorFoundCopy = GetSessionDescriptor(RemoteSessionID);
 
 	if (!DescriptorFoundCopy.IsValid())
 	{
 		return;
 	}
 
-	if (const TSharedPtr<FTraceSessionsManager> CVDTraceManager = FChaosVDModule::Get().GetTraceSessionsManager())
+	if (const TSharedPtr<FChaosVDTraceManager> CVDTraceManager = FChaosVDModule::Get().GetTraceManager())
 	{
 		CVDTraceManager->CloseSession(DescriptorFoundCopy.SessionName);
 	}
 
-	OnTraceSessionClosed().Broadcast(DescriptorFoundCopy);
+	OnSessionClosed().Broadcast(DescriptorFoundCopy);
 	CurrentSessionDescriptors.Remove(DescriptorFoundCopy);
 
 	if (PlaybackController && CurrentSessionDescriptors.IsEmpty())
@@ -112,10 +99,9 @@ void FChaosVDEngine::CloseSessionByRemoteSessionID(FGuid RemoteSessionID)
 
 void FChaosVDEngine::StopActiveTraceSessions()
 {
-	using namespace UE::TraceBasedDebuggers;
-	if (const TSharedPtr<FTraceSessionsManager> CVDTraceManager = FChaosVDModule::Get().GetTraceSessionsManager())
+	if (const TSharedPtr<FChaosVDTraceManager> CVDTraceManager = FChaosVDModule::Get().GetTraceManager())
 	{
-		for (FTraceSessionDescriptor& SessionDescriptor : CurrentSessionDescriptors)
+		for (FChaosVDTraceSessionDescriptor& SessionDescriptor : CurrentSessionDescriptors)
 		{
 			CVDTraceManager->StopSession(SessionDescriptor.SessionName);
 		}
@@ -136,7 +122,7 @@ void FChaosVDEngine::DeInitialize()
 
 	CloseActiveTraceSessions();
 
-	if (const TSharedPtr<UE::TraceBasedDebuggers::FRemoteSessionsManager> RemoteSessionManager = FChaosVDEngineEditorBridge::Get().GetSessionsManager())
+	if (TSharedPtr<UE::TraceBasedDebuggers::FRemoteSessionsManager> RemoteSessionManager = FChaosVDEngineEditorBridge::Get().GetSessionsManager())
 	{
 		RemoteSessionManager->OnSessionRecordingStopped().Remove(LiveSessionStoppedDelegateHandle);
 	}
@@ -239,19 +225,10 @@ bool FChaosVDEngine::HasAnyLiveSessionActive() const
 	return PlaybackController && PlaybackController->IsPlayingLiveSession();
 }
 
-PRAGMA_DISABLE_DEPRECATION_WARNINGS
-const FChaosVDTraceSessionDescriptor& FChaosVDEngine::GetSessionDescriptor(FGuid) const
+const FChaosVDTraceSessionDescriptor& FChaosVDEngine::GetSessionDescriptor(FGuid RemoteSessionID) const
 {
 	static FChaosVDTraceSessionDescriptor InvalidSessionDescriptor;
-	return InvalidSessionDescriptor;
-}
-PRAGMA_ENABLE_DEPRECATION_WARNINGS
-
-const UE::TraceBasedDebuggers::FTraceSessionDescriptor& FChaosVDEngine::GetTraceSessionDescriptor(FGuid RemoteSessionID) const
-{
-	using namespace UE::TraceBasedDebuggers;
-	static FTraceSessionDescriptor InvalidSessionDescriptor;
-	const FTraceSessionDescriptor* DescriptorFound = CurrentSessionDescriptors.FindByPredicate([RemoteSessionID](const FTraceSessionDescriptor& InDescriptor)
+	const FChaosVDTraceSessionDescriptor* DescriptorFound = CurrentSessionDescriptors.FindByPredicate([RemoteSessionID](const FChaosVDTraceSessionDescriptor& InDescriptor)
 	{
 		return RemoteSessionID == InDescriptor.RemoteSessionID;
 	});
@@ -259,17 +236,14 @@ const UE::TraceBasedDebuggers::FTraceSessionDescriptor& FChaosVDEngine::GetTrace
 	return DescriptorFound ? *DescriptorFound : InvalidSessionDescriptor;
 }
 
-TArrayView<UE::TraceBasedDebuggers::FTraceSessionDescriptor> FChaosVDEngine::GetTraceSessionDescriptors()
-{
-	return CurrentSessionDescriptors;
-}
-
 void FChaosVDEngine::LoadRecording_Internal(const TFunction<FString(const TSharedPtr<FChaosVDRecording>&)>& LoadCallback, EChaosVDLoadRecordedDataMode LoadingMode)
 {
+	FChaosVDTraceSessionDescriptor NewSessionFromFileDescriptor;
+
 	TSharedPtr<FChaosVDRecording> ExistingRecordingInstance = nullptr;
 	if (LoadingMode == EChaosVDLoadRecordedDataMode::MultiSource && !CurrentSessionDescriptors.IsEmpty())
 	{
-		if (const TSharedPtr<const TraceServices::IAnalysisSession> TraceSession = FChaosVDModule::Get().GetTraceSessionsManager()->GetSession(CurrentSessionDescriptors[0].SessionName))
+		if (const TSharedPtr<const TraceServices::IAnalysisSession> TraceSession = FChaosVDModule::Get().GetTraceManager()->GetSession(CurrentSessionDescriptors[0].SessionName))
 		{
 			if (const FChaosVDTraceProvider* ChaosVDProvider = TraceSession->ReadProvider<FChaosVDTraceProvider>(FChaosVDTraceProvider::ProviderName))
 			{
@@ -277,10 +251,9 @@ void FChaosVDEngine::LoadRecording_Internal(const TFunction<FString(const TShare
 			}
 		}
 	}
-
-	using namespace UE::TraceBasedDebuggers;
-	FTraceSessionDescriptor NewSessionFromFileDescriptor;
+	
 	NewSessionFromFileDescriptor.SessionName = LoadCallback(ExistingRecordingInstance);
+	
 	NewSessionFromFileDescriptor.bIsLiveSession = false;
 
 	if (NewSessionFromFileDescriptor.IsValid())
@@ -302,16 +275,12 @@ void FChaosVDEngine::LoadRecording(const FString& FilePath, EChaosVDLoadRecorded
 		return;
 	}
 
-	auto LoadSingleFilePathFromTraceManager = [&FilePath](const TSharedPtr<FChaosVDRecording>& ExistingRecordingInstance)
+	auto LoadSingleFileFromTraceManager = [&FilePath](const TSharedPtr<FChaosVDRecording>& ExistingRecordingInstance)
 	{
-		return FChaosVDModule::Get().GetTraceSessionsManager()->LoadTraceFile(FilePath, /*PreStartAnalysisCallback*/
-			[ExistingRecordingInstance]
-			{
-				FChaosVDTraceManagerThreadContext::Get().PendingExternalRecordingWeakPtr = ExistingRecordingInstance;
-			});
+		return FChaosVDModule::Get().GetTraceManager()->LoadTraceFile(FilePath, ExistingRecordingInstance);
 	};
-
-	LoadRecording_Internal(LoadSingleFilePathFromTraceManager, LoadingMode);
+	
+	LoadRecording_Internal(LoadSingleFileFromTraceManager, LoadingMode);
 	UpdateRecentFilesList(FilePath);
 }
 
@@ -340,11 +309,7 @@ void FChaosVDEngine::LoadCombinedMultiRecording(const FString& FilePath)
 		FString SessionName = FString::Format(TEXT("{0}-{1}"), Args);
 		auto LoadSingleFileHandleFromTraceManager = [&Handle, &SessionName](const TSharedPtr<FChaosVDRecording>& ExistingRecordingInstance)
 		{
-			return FChaosVDModule::Get().GetTraceSessionsManager()->LoadTraceFile(MoveTemp(Handle), SessionName,  /*PreStartAnalysisCallback*/
-				[ExistingRecordingInstance]
-				{
-					FChaosVDTraceManagerThreadContext::Get().PendingExternalRecordingWeakPtr = ExistingRecordingInstance;
-				});
+			return FChaosVDModule::Get().GetTraceManager()->LoadTraceFile(MoveTemp(Handle), SessionName, ExistingRecordingInstance);
 		};
 
 		LoadRecording_Internal(LoadSingleFileHandleFromTraceManager, EChaosVDLoadRecordedDataMode::MultiSource);
@@ -361,8 +326,7 @@ bool FChaosVDEngine::ConnectToLiveSession(uint32 SessionID, const FString& InSes
 		CloseActiveTraceSessions();
 	}
 	
-	using namespace UE::TraceBasedDebuggers;
-	FTraceSessionDescriptor NewSessionFromFileDescriptor;
+	FChaosVDTraceSessionDescriptor NewSessionFromFileDescriptor;
 
 	NewSessionFromFileDescriptor.bIsLiveSession = true;
 	
@@ -372,13 +336,7 @@ bool FChaosVDEngine::ConnectToLiveSession(uint32 SessionID, const FString& InSes
 		ExistingRecordingInstance = PlaybackController->GetCurrentRecording().Pin();
 	}
 
-	auto PreStartAnalysisCallback = [ExistingRecordingInstance]
-		{
-			FChaosVDTraceManagerThreadContext::Get().PendingExternalRecordingWeakPtr = ExistingRecordingInstance;
-		};
-
-	NewSessionFromFileDescriptor.SessionName =
-		FChaosVDModule::Get().GetTraceSessionsManager()->ConnectToLiveSession(InSessionAddress, SessionID, PreStartAnalysisCallback);
+	NewSessionFromFileDescriptor.SessionName = FChaosVDModule::Get().GetTraceManager()->ConnectToLiveSession(InSessionAddress, SessionID, ExistingRecordingInstance);
 
 	bool bSuccess = false;
 
@@ -394,7 +352,7 @@ bool FChaosVDEngine::ConnectToLiveSession(uint32 SessionID, const FString& InSes
 
 bool FChaosVDEngine::ConnectToLiveSession_Direct(EChaosVDLoadRecordedDataMode LoadingMode)
 {
-	UE_LOG(LogChaosVDEditor, Error, TEXT("[%hs] This method is no longer supported. Please use the overload that takes a Remote Session ID."), __func__);
+	UE_LOG(LogChaosVDEditor, Error, TEXT("[%hs] This menthod is no longer supported. Please use the overload that takes a Remote Session ID."), __func__);
 	return false;
 }
 
@@ -405,25 +363,18 @@ bool FChaosVDEngine::ConnectToLiveSession_Direct(FGuid RemoteSessionID, EChaosVD
 		CloseActiveTraceSessions();
 	}
 
-	using namespace UE::TraceBasedDebuggers;
-	FTraceSessionDescriptor NewSessionFromFileDescriptor;
+	FChaosVDTraceSessionDescriptor NewSessionFromFileDescriptor;
 
 	NewSessionFromFileDescriptor.RemoteSessionID = RemoteSessionID;
 	NewSessionFromFileDescriptor.bIsLiveSession = true;
-
+	
 	TSharedPtr<FChaosVDRecording> ExistingRecordingInstance = nullptr;
 	if (LoadingMode == EChaosVDLoadRecordedDataMode::MultiSource && PlaybackController)
 	{
 		ExistingRecordingInstance = PlaybackController->GetCurrentRecording().Pin();
 	}
 
-	auto PreStartAnalysisCallback = [ExistingRecordingInstance]
-		{
-			FChaosVDTraceManagerThreadContext::Get().PendingExternalRecordingWeakPtr = ExistingRecordingInstance;
-		};
-
-	NewSessionFromFileDescriptor.SessionName =
-		FChaosVDModule::Get().GetTraceSessionsManager()->ConnectToLiveSession_Direct(RemoteSessionID, NewSessionFromFileDescriptor.SessionPort, PreStartAnalysisCallback);
+	NewSessionFromFileDescriptor.SessionName = FChaosVDModule::Get().GetTraceManager()->ConnectToLiveSession_Direct(RemoteSessionID, NewSessionFromFileDescriptor.SessionPort, ExistingRecordingInstance);
 
 	bool bSuccess = false;
 
@@ -443,25 +394,18 @@ bool FChaosVDEngine::ConnectToLiveSession_Relay(FGuid RemoteSessionID, EChaosVDL
 		CloseActiveTraceSessions();
 	}
 
-	using namespace UE::TraceBasedDebuggers;
-	FTraceSessionDescriptor NewSessionFromFileDescriptor;
+	FChaosVDTraceSessionDescriptor NewSessionFromFileDescriptor;
 
 	NewSessionFromFileDescriptor.bIsLiveSession = true;
-
+	
 	TSharedPtr<FChaosVDRecording> ExistingRecordingInstance = nullptr;
 	if (LoadingMode == EChaosVDLoadRecordedDataMode::MultiSource && PlaybackController)
 	{
 		ExistingRecordingInstance = PlaybackController->GetCurrentRecording().Pin();
 	}
 
-	auto PreStartAnalysisCallback = [ExistingRecordingInstance]
-		{
-			FChaosVDTraceManagerThreadContext::Get().PendingExternalRecordingWeakPtr = ExistingRecordingInstance;
-		};
-
 	NewSessionFromFileDescriptor.RemoteSessionID = RemoteSessionID;
-	NewSessionFromFileDescriptor.SessionName =
-		FChaosVDModule::Get().GetTraceSessionsManager()->ConnectToLiveSession_Relay(RemoteSessionID, PreStartAnalysisCallback);
+	NewSessionFromFileDescriptor.SessionName = FChaosVDModule::Get().GetTraceManager()->ConnectToLiveSession_Relay(RemoteSessionID, ExistingRecordingInstance);
 
 	bool bSuccess = false;
 
@@ -474,7 +418,7 @@ bool FChaosVDEngine::ConnectToLiveSession_Relay(FGuid RemoteSessionID, EChaosVDL
 	return bSuccess;
 }
 
-void FChaosVDEngine::OpenSession(const UE::TraceBasedDebuggers::FTraceSessionDescriptor& SessionDescriptor, EChaosVDLoadRecordedDataMode LoadingMode)
+void FChaosVDEngine::OpenSession(const FChaosVDTraceSessionDescriptor& SessionDescriptor, EChaosVDLoadRecordedDataMode LoadingMode)
 {
 	CurrentSessionDescriptors.Emplace(SessionDescriptor);
 
@@ -483,7 +427,7 @@ void FChaosVDEngine::OpenSession(const UE::TraceBasedDebuggers::FTraceSessionDes
 		PlaybackController->LoadChaosVDRecordingFromTraceSession(SessionDescriptor);
 	}
 
-	OnTraceSessionOpened().Broadcast(SessionDescriptor);
+	OnSessionOpened().Broadcast(SessionDescriptor);
 }
 
 bool FChaosVDEngine::SaveOpenSessionToCombinedFile(const FString& InTargetFilePath)
@@ -498,8 +442,7 @@ bool FChaosVDEngine::SaveOpenSessionToCombinedFile(const FString& InTargetFilePa
 	TArray<TUniquePtr<IFileHandle>> FilePathsToCombine;
 	FilePathsToCombine.Reserve(CurrentSessionDescriptors.Num());
 
-	using namespace UE::TraceBasedDebuggers;
-	for (FTraceSessionDescriptor& SessionDescriptor : CurrentSessionDescriptors)
+	for (FChaosVDTraceSessionDescriptor& SessionDescriptor : CurrentSessionDescriptors)
 	{
 		if (SessionDescriptor.SessionName.EndsWith(TEXT("cvdmulti")))
 		{
@@ -539,8 +482,7 @@ bool FChaosVDEngine::CanCombineOpenSessions() const
 {
 	if (CurrentSessionDescriptors.Num() > 1)
 	{
-		using namespace UE::TraceBasedDebuggers;
-		for (const FTraceSessionDescriptor& Session : CurrentSessionDescriptors)
+		for (const FChaosVDTraceSessionDescriptor& Session : CurrentSessionDescriptors)
 		{
 			if (Session.SessionName.EndsWith(TEXT("cvdmulti")))
 			{
diff --git a/Engine/Plugins/ChaosVD/Source/ChaosVD/Private/ChaosVDModule.cpp b/Engine/Plugins/ChaosVD/Source/ChaosVD/Private/ChaosVDModule.cpp
--- a/Engine/Plugins/ChaosVD/Source/ChaosVD/Private/ChaosVDModule.cpp
+++ b/Engine/Plugins/ChaosVD/Source/ChaosVD/Private/ChaosVDModule.cpp
@@ -1,31 +1,28 @@
-﻿// Copyright Epic Games, Inc. All Rights Reserved.
+// Copyright Epic Games, Inc. All Rights Reserved.
 
 #include "ChaosVDModule.h"
 
-#include "Chaos/ChaosVDEngineEditorBridge.h"
 #include "ChaosVDCommands.h"
 #include "ChaosVDEngine.h"
 #include "ChaosVDPlaybackController.h"
-#include "ChaosVDRuntimeModule.h"
 #include "ChaosVDSceneCompositionReport.h"
 #include "ChaosVDSceneParticleCustomization.h"
 #include "ChaosVDSettingsManager.h"
 #include "ChaosVDStyle.h"
 #include "ChaosVDTabsIDs.h"
 #include "DesktopPlatformModule.h"
-#include "ExtensionsSystem/ChaosVDExtensionsManager.h"
-#include "Framework/Application/SlateApplication.h"
-#include "HAL/FileManager.h"
 #include "IDesktopPlatform.h"
 #include "Misc/App.h"
 #include "Misc/Guid.h"
-#include "Settings/ChaosVDGeneralSettings.h"
-#include "Trace/ChaosVDTraceModule.h"
-#include "TraceSessionsManager.h"
+#include "Trace/ChaosVDTraceManager.h"
 #include "Widgets/Docking/SDockTab.h"
 #include "Widgets/SChaosVDMainTab.h"
 #include "WorkspaceMenuStructure.h"
 #include "WorkspaceMenuStructureModule.h"
+#include "DetailsCustomizations/ChaosVDShapeDataCustomization.h"
+#include "ExtensionsSystem/ChaosVDExtensionsManager.h"
+#include "Framework/Application/SlateApplication.h"
+#include "HAL/FileManager.h"
 
 #define LOCTEXT_NAMESPACE "ChaosVisualDebugger"
 
@@ -48,9 +45,9 @@ FChaosVDModule& FChaosVDModule::Get()
 }
 
 void FChaosVDModule::StartupModule()
-{
+{	
 	FChaosVDStyle::Initialize();
-
+	
 	FChaosVDCommands::Register();
 
 	FGlobalTabmanager::Get()->RegisterNomadTabSpawner(FChaosVDTabID::ChaosVisualDebuggerTab, FOnSpawnTab::CreateRaw(this, &FChaosVDModule::SpawnMainTab))
@@ -59,33 +56,7 @@ void FChaosVDModule::StartupModule()
 								.SetIcon(FSlateIcon(FChaosVDStyle::GetStyleSetName(), "ChaosVisualDebugger"))
 								.SetGroup(WorkspaceMenu::GetMenuStructure().GetDeveloperToolsDebugCategory());
 
-	ChaosVDTraceManager = MakeShared<UE::TraceBasedDebuggers::FTraceSessionsManager>(
-		FChaosVDRuntimeModule::Get()
-		, FChaosVDEngineEditorBridge::Get()
-		, MakeShared<FChaosVDTraceModule>()
-		, /*DefaultSavePath*/TEXT("ChaosVisualDebugger"));
-
-	// Set options using lambdas so the most recent values from the settings are used
-	// (i.e., users can change settings during the session while the trace manager is created only once)
-	ChaosVDTraceManager->SetSaveRecordingsToDisk(MakeAttributeLambda([]
-		{
-			if (const UChaosVDGeneralSettings* Settings = FChaosVDSettingsManager::Get().GetSettingsObject<UChaosVDGeneralSettings>())
-			{
-				return Settings->bSaveRecordingsToDisk;
-			}
-			return false;
-		}));
-
-	ChaosVDTraceManager->SetSavePathOverride(MakeAttributeLambda([]
-		{
-			if (const UChaosVDGeneralSettings* Settings = FChaosVDSettingsManager::Get().GetSettingsObject<UChaosVDGeneralSettings>())
-			{
-				return Settings->DefaultSavePathOverride;
-			}
-			return FFilePath{};
-		}));
-
-	UE::Trace::ToggleChannel(TEXT("ChaosVD"), /*bEnabled*/true);
+	ChaosVDTraceManager = MakeShared<FChaosVDTraceManager>();
 
 	if (IsStandaloneChaosVisualDebugger())
 	{
diff --git a/Engine/Plugins/ChaosVD/Source/ChaosVD/Private/ChaosVDPlaybackController.cpp b/Engine/Plugins/ChaosVD/Source/ChaosVD/Private/ChaosVDPlaybackController.cpp
--- a/Engine/Plugins/ChaosVD/Source/ChaosVD/Private/ChaosVDPlaybackController.cpp
+++ b/Engine/Plugins/ChaosVD/Source/ChaosVD/Private/ChaosVDPlaybackController.cpp
@@ -12,8 +12,8 @@
 #include "ExtensionsSystem/ChaosVDExtensionsManager.h"
 #include "Misc/MessageDialog.h"
 #include "Settings/ChaosVDGeneralSettings.h"
+#include "Trace/ChaosVDTraceManager.h"
 #include "Trace/ChaosVDTraceProvider.h"
-#include "TraceSessionsManager.h"
 #include "TraceServices/Model/AnalysisSession.h"
 #include "Widgets/SChaosVDTimelineWidget.h"
 
@@ -86,7 +86,7 @@ void FChaosVDPlaybackController::DeInitialize()
 	bInitialized = false;
 }
 
-bool FChaosVDPlaybackController::LoadChaosVDRecordingFromTraceSession(const UE::TraceBasedDebuggers::FTraceSessionDescriptor& InSessionDescriptor)
+bool FChaosVDPlaybackController::LoadChaosVDRecordingFromTraceSession(const FChaosVDTraceSessionDescriptor& InSessionDescriptor)
 {
 	if (!ensure(InSessionDescriptor.IsValid()))
 	{
@@ -98,7 +98,7 @@ bool FChaosVDPlaybackController::LoadChaosVDRecordingFromTraceSession(const UE::
 		UnloadCurrentRecording();
 	}
 
-	if (const TSharedPtr<const TraceServices::IAnalysisSession> TraceSession = FChaosVDModule::Get().GetTraceSessionsManager()->GetSession(InSessionDescriptor.SessionName))
+	if (const TSharedPtr<const TraceServices::IAnalysisSession> TraceSession = FChaosVDModule::Get().GetTraceManager()->GetSession(InSessionDescriptor.SessionName))
 	{
 		if (const FChaosVDTraceProvider* ChaosVDProvider = TraceSession->ReadProvider<FChaosVDTraceProvider>(FChaosVDTraceProvider::ProviderName))
 		{
diff --git a/Engine/Plugins/ChaosVD/Source/ChaosVD/Private/Trace/ChaosVDTraceManager.cpp b/Engine/Plugins/ChaosVD/Source/ChaosVD/Private/Trace/ChaosVDTraceManager.cpp
new file mode 100644
--- /dev/null
+++ b/Engine/Plugins/ChaosVD/Source/ChaosVD/Private/Trace/ChaosVDTraceManager.cpp
@@ -0,0 +1,809 @@
+// Copyright Epic Games, Inc. All Rights Reserved.
+
+#include "Trace/ChaosVDTraceManager.h"
+
+#include "Chaos/ChaosVDEngineEditorBridge.h"
+#include "Chaos/Serialization/SerializedDataBuffer.h"
+#include "ChaosVDEngine.h"
+#include "ChaosVDModule.h"
+#include "ChaosVDRuntimeModule.h"
+#include "ChaosVDSettingsManager.h"
+#include "Features/IModularFeatures.h"
+#include "HAL/PlatformFileManager.h"
+#include "Misc/Paths.h"
+#include "Modules/ModuleManager.h"
+#include "RemoteSessionsManager.h"
+#include "SessionInfo.h"
+#include "Settings/ChaosVDGeneralSettings.h"
+#include "Trace/ChaosVDTraceModule.h"
+#include "Trace/StoreClient.h"
+#include "TraceServices/AnalysisService.h"
+#include "TraceServices/ITraceServicesModule.h"
+#include "TraceServices/Model/AnalysisSession.h"
+
+namespace Chaos::VD
+{
+	FAsyncProgressNotification::~FAsyncProgressNotification()
+	{
+		if (NotificationHandle.IsValid())
+		{
+			// FSlateNotificationManager is not thread safe, so we need to make sure this runs in the game thread
+			if (IsInGameThread())
+			{
+				FSlateNotificationManager::Get().CancelProgressNotification(NotificationHandle);
+			}
+			else
+			{
+				constexpr float Delay = 0.0f;
+                FTSTicker::GetCoreTicker().AddTicker(TEXT("CancelCVDAsyncNotification"), Delay, [HandleCopy = NotificationHandle](float DeltaTime)
+                {
+                	FSlateNotificationManager::Get().CancelProgressNotification(HandleCopy);
+                	return false;
+                });
+			}
+		}
+	}
+
+	void FAsyncProgressNotification::EnterProgress(int32 InCurrentProgress)
+	{
+		CurrentProgress = InCurrentProgress;
+	}
+
+	bool FAsyncProgressNotification::Tick(float DeltaTime)
+	{
+		constexpr int32 MaxProgress = 100;
+		constexpr float TimeBetweenUpdatesSeconds = 0.5f;
+
+		ElapsedTimeSinceLastUpdate += DeltaTime;
+		if (ElapsedTimeSinceLastUpdate < TimeBetweenUpdatesSeconds)
+		{
+			return true;
+		}
+
+		ElapsedTimeSinceLastUpdate = 0.0f;
+		
+		if (CurrentProgress >= MaxProgress)
+		{
+			if (NotificationHandle.IsValid())
+			{
+				FSlateNotificationManager::Get().UpdateProgressNotification(NotificationHandle, MaxProgress, MaxProgress, NotificationTitle);
+			}
+			NotificationHandle = FProgressNotificationHandle();
+
+			// We are done and we can stop ticking
+			return false;
+		}
+
+		if (!NotificationHandle.IsValid())
+		{
+			NotificationHandle = FSlateNotificationManager::Get().StartProgressNotification(NotificationTitle, MaxProgress);
+		}
+		else
+		{
+			FSlateNotificationManager::Get().UpdateProgressNotification(NotificationHandle, CurrentProgress, MaxProgress, NotificationTitle);
+		}
+
+		return true;
+	}
+
+	FSessionStreamStatsUpdater::~FSessionStreamStatsUpdater()
+	{
+	}
+
+	bool FSessionStreamStatsUpdater::Tick(float DeltaTime)
+	{
+		ElapsedTimeSinceLastUpdate += DeltaTime;
+
+		constexpr double IntervalBetweenUpdatesSeconds = 1.0;
+		if (ElapsedTimeSinceLastUpdate < IntervalBetweenUpdatesSeconds)
+		{
+			return true;
+		}
+
+		if (!ensure(DataStreamInstance))
+		{
+			return false;
+		}
+
+		ElapsedTimeSinceLastUpdate = 0.0;
+
+		using namespace UE::TraceBasedDebuggers;
+		TSharedPtr<FRemoteSessionsManager> RemoteSessionManager = FChaosVDEngineEditorBridge::Get().GetSessionsManager();
+		TSharedPtr<FSessionInfo> SessionInfo = RemoteSessionManager ? RemoteSessionManager->GetSessionInfo(DataStreamInstance->GetOwningRemoteSessionID()).Pin() : nullptr;
+		if (!SessionInfo)
+		{
+			return false;
+		}
+
+		SessionInfo->SetReceivedBytesPerSecond(DataStreamInstance->UpdateAndGetBytesReadSinceLastMeasurement());
+	
+		return true;
+	}
+
+	int32 FDirectSocketStream::Read(void* Data, uint32 Size)
+	{
+		int32 BytesRead =  UE::Trace::FDirectSocketStream::Read(Data, Size);
+
+		if (BytesRead > 0)
+		{
+			BytesReadSinceLastMeasurement += BytesRead;
+		}
+	
+		return BytesRead;
+	}
+
+	FGuid FDirectSocketStream::GetOwningRemoteSessionID() const
+	{
+		return RemoteSessionID;
+	}
+
+	uint64 FDirectSocketStream::UpdateAndGetBytesReadSinceLastMeasurement()
+	{
+		uint64 BytesRead = BytesReadSinceLastMeasurement;
+
+		BytesReadSinceLastMeasurement = 0;
+	
+		return BytesRead;
+	}
+
+	FRelayDataStream::FRelayDataStream(FGuid InRemoteSessionID)
+		: CancelEvent(FPlatformProcess::GetSynchEventFromPool(false)),
+		WaitForDataEvent(FPlatformProcess::GetSynchEventFromPool(false)),
+		RemoteSessionID(InRemoteSessionID),
+		FileWriter(nullptr),
+		BytesReadSinceLastMeasurement(0),
+		StatsUpdater(this)
+	{
+#if WITH_CHAOS_VISUAL_DEBUGGER
+		using namespace UE::TraceBasedDebuggers;
+		TSharedPtr<IDataRelayTransport> RelayTransportInstance = FChaosVDEngineEditorBridge::Get().GetTraceRelayTransportInstance();
+		if (ensure(RelayTransportInstance))
+		{
+			RelayTransportInstance->RegisterRelayDataReceiverForSessionID(RemoteSessionID, IDataRelayTransport::FProcessReceivedRelayDataDelegate::CreateRaw(this, &FRelayDataStream::EnqueueRelayedData));
+		}
+#endif
+	}
+
+	FRelayDataStream::FRelayDataStream(FGuid InRemoteSessionID, TUniquePtr<FArchiveFileWriterGeneric>&& InFileWriter)
+		: CancelEvent(FPlatformProcess::GetSynchEventFromPool(false)),
+		WaitForDataEvent(FPlatformProcess::GetSynchEventFromPool(false)),
+		RemoteSessionID(InRemoteSessionID),
+		FileWriter(MoveTemp(InFileWriter)),
+		BytesReadSinceLastMeasurement(0),
+		StatsUpdater(this)
+	{
+#if WITH_CHAOS_VISUAL_DEBUGGER
+		using namespace UE::TraceBasedDebuggers;
+		TSharedPtr<IDataRelayTransport> RelayTransportInstance = FChaosVDEngineEditorBridge::Get().GetTraceRelayTransportInstance();
+		if (ensure(RelayTransportInstance))
+		{
+			RelayTransportInstance->RegisterRelayDataReceiverForSessionID(RemoteSessionID, IDataRelayTransport::FProcessReceivedRelayDataDelegate::CreateRaw(this, &FRelayDataStream::EnqueueRelayedData));
+		}
+#endif
+	}
+
+	FRelayDataStream::~FRelayDataStream()
+	{
+		FPlatformProcess::ReturnSynchEventToPool(CancelEvent);
+		FPlatformProcess::ReturnSynchEventToPool(WaitForDataEvent);
+
+#if WITH_CHAOS_VISUAL_DEBUGGER
+		using namespace UE::TraceBasedDebuggers;
+		TSharedPtr<IDataRelayTransport> RelayTransportInstance = FChaosVDEngineEditorBridge::Get().GetTraceRelayTransportInstance();
+		if (RelayTransportInstance)
+		{
+			RelayTransportInstance->UnregisterRelayDataReceiverForSessionID(RemoteSessionID);
+		}
+#endif
+	}
+
+	void FRelayDataStream::EnqueueRelayedData(const TConstArrayView<uint8> InTraceDataBuffer)
+	{
+		// An empty data array signals an end of file event from the relay instance
+		if (InTraceDataBuffer.IsEmpty())
+		{
+			Close();
+			return;
+		}
+
+		WaitForDataEvent->Trigger();
+	
+		FSerializedDataBuffer DataBuffer(InTraceDataBuffer);
+		DataQueue.Enqueue(MoveTemp(DataBuffer));
+	}
+
+	void FRelayDataStream::EnqueueRelayedData(TArray<uint8>&& InTraceDataBuffer)
+	{
+		// An empty data array signals an end of file event from the relay instance
+		if (InTraceDataBuffer.IsEmpty())
+		{
+			Close();
+			return;
+		}
+
+		WaitForDataEvent->Trigger();
+	
+		FSerializedDataBuffer DataBuffer(MoveTemp(InTraceDataBuffer));
+		DataQueue.Enqueue(MoveTemp(DataBuffer));
+	}
+
+	FGuid FRelayDataStream::GetOwningRemoteSessionID() const
+	{
+		return RemoteSessionID;
+	}
+
+	uint64 FRelayDataStream::UpdateAndGetBytesReadSinceLastMeasurement()
+	{
+		uint64 BytesRead = BytesReadSinceLastMeasurement;
+		BytesReadSinceLastMeasurement = 0;
+		return BytesRead;
+	}
+
+	int32 FRelayDataStream::Read(void* Dest, uint32 DestSize)
+	{
+		if (DestSize == 0)
+		{
+			return 0;
+		}
+
+		if (Dest == nullptr)
+		{
+			return 0;
+		}
+
+		if (!ensure(CancelEvent))
+		{
+			return 0;
+		}
+
+		uint8* DestBuffer = static_cast<uint8*>(Dest);
+		uint32 BytesCopied = 0;
+
+		// We need to wait for some data to be queued. Returning 
+		// zero bytes read is interpreted as eof.
+		while (!CancelEvent->Wait(0, true))
+		{
+			if (!CurrentPacket)
+			{
+				CurrentPacket = DataQueue.Dequeue();
+				CurrentPacketOffset = 0;
+			}
+
+			// Consume as many queued packets as possible and copy into
+			// the destination buffer
+			while (CurrentPacket && BytesCopied < DestSize)
+			{
+				const TArray<uint8>& SrcBuffer = CurrentPacket->GetDataAsByteArrayRef();
+				if (SrcBuffer.IsEmpty())
+				{
+					// We should never get an empty packet. From this point on we can't trust the integrity of the data
+					// so we need to close the stream
+					Close();
+					break;
+				}
+
+				const uint32 ReceivedBufferSize = static_cast<uint32>(SrcBuffer.Num());
+
+				// Copy as much of the current buffer as possible into the remaining space 
+				// in the destination buffer
+				uint32 BytesToCopy = FMath::Min(ReceivedBufferSize - CurrentPacketOffset, DestSize - BytesCopied);
+				FMemory::Memcpy(&DestBuffer[BytesCopied], SrcBuffer.GetData() + CurrentPacketOffset, BytesToCopy);
+				BytesCopied += BytesToCopy;
+				CurrentPacketOffset += BytesToCopy;
+
+				if (CurrentPacketOffset == ReceivedBufferSize)
+				{
+					// Current packet buffer has been fully copied, goto next
+					CurrentPacket = DataQueue.Dequeue();
+					CurrentPacketOffset = 0;
+				}
+			}
+
+			// If at least some bytes have been written return those. Otherwise, 
+			// we'll wait for some data to arrive.
+			if (BytesCopied > 0)
+			{
+				BytesReadSinceLastMeasurement += BytesCopied;
+
+				if (FileWriter.IsValid() && !FileWriter->IsError())
+				{
+					FileWriter->Serialize(Dest, BytesCopied);
+				}	
+				break;
+			}
+			else
+			{
+				WaitForDataEvent->Reset();
+				WaitForDataEvent->Wait();
+			}
+		}
+
+		return BytesCopied;
+	}
+
+	void FRelayDataStream::Close()
+	{
+		if (CancelEvent)
+		{
+			CancelEvent->Trigger();
+		}
+
+		if (WaitForDataEvent)
+		{
+			WaitForDataEvent->Trigger();
+		}
+
+		if (FileWriter.IsValid())
+		{
+			FileWriter->Close();
+			FileWriter.Reset();
+		}
+
+		bClosed = true;
+	}
+
+	bool FRelayDataStream::WaitUntilReady()
+	{		
+		return DataQueue.Peek() != nullptr; 
+	}
+}
+
+uint16 FChaosVDTraceManager::CurrentDirectModeStartingPort = DefaultStartingPort;
+
+FChaosVDTraceManager::FChaosVDTraceManager() 
+{
+	ChaosVDTraceModule = MakeShared<FChaosVDTraceModule>();
+	IModularFeatures::Get().RegisterModularFeature(TraceServices::ModuleFeatureName, ChaosVDTraceModule.Get());
+
+	FString ChannelNameFString(TEXT("ChaosVD"));
+	UE::Trace::ToggleChannel(ChannelNameFString.GetCharArray().GetData(), true);
+}
+
+FChaosVDTraceManager::~FChaosVDTraceManager()
+{
+	IModularFeatures::Get().UnregisterModularFeature(TraceServices::ModuleFeatureName, ChaosVDTraceModule.Get());
+}
+
+FString FChaosVDTraceManager::LoadTraceFile(const FString& InTraceFilename, const TSharedPtr<FChaosVDRecording>& ExistingRecordingPtr)
+{
+	CloseSession(InTraceFilename);
+
+	IPlatformFile& FileSystem = IPlatformFile::GetPlatformPhysical();
+	
+	constexpr bool bAllowWrite = false;
+	IFileHandle* Handle = FileSystem.OpenRead(*InTraceFilename, bAllowWrite);
+	if (!Handle)
+	{
+		return FString();
+	}
+
+	return LoadTraceFile(TUniquePtr<IFileHandle>(Handle), InTraceFilename, ExistingRecordingPtr);
+}
+
+FString FChaosVDTraceManager::LoadTraceFile(TUniquePtr<IFileHandle>&& InFileHandle, const FString& InTraceSessionName, const TSharedPtr<FChaosVDRecording>& ExistingRecordingPtr)
+{
+	if (InFileHandle->Size() == 0)
+	{
+		return FString();
+	}
+
+	ITraceServicesModule& TraceServicesModule = FModuleManager::LoadModuleChecked<ITraceServicesModule>("TraceServices");
+	if (const TSharedPtr<TraceServices::IAnalysisService> TraceAnalysisService = TraceServicesModule.GetAnalysisService())
+	{
+		SetPendingExternalRecordingToProcess(ExistingRecordingPtr);
+
+		if (const TSharedPtr<const TraceServices::IAnalysisSession> NewSession = TraceAnalysisService->StartAnalysis(~0, *InTraceSessionName, CreateFileDataStream(MoveTemp(InFileHandle))))
+		{
+			AnalysisSessionByName.Add(InTraceSessionName, NewSession.ToSharedRef());
+
+			return NewSession->GetName();
+		}
+	}
+
+	return FString();
+}
+
+const UE::Trace::FStoreClient::FSessionInfo* FChaosVDTraceManager::GetTraceSessionInfo(FStringView InSessionHost, FGuid TraceGuid)
+{
+	if (InSessionHost.IsEmpty())
+	{
+		UE_LOG(LogChaosVDEditor, Error, TEXT("[%s] Failed to connect to trace store. Provided session host is empty"), ANSI_TO_TCHAR(__FUNCTION__));
+		return nullptr;
+	}
+
+	using namespace UE::Trace;
+	const FStoreClient* StoreClient = FStoreClient::Connect(InSessionHost.GetData());
+
+	if (!StoreClient)
+	{
+		UE_LOG(LogChaosVDEditor, Error, TEXT("[%s] Failed to connect to trace store at [%s]"), ANSI_TO_TCHAR(__FUNCTION__), InSessionHost.GetData())
+		return nullptr;
+	}
+
+	return StoreClient->GetSessionInfoByGuid(TraceGuid);
+}
+
+FString FChaosVDTraceManager::GetTraceFileNameFromStoreForSession(FStringView InSessionHost, uint32 SessionID)
+{
+	using namespace UE::Trace;
+
+	if (InSessionHost.IsEmpty())
+	{
+		return FString();
+	}
+	
+	FStoreClient* StoreClient = FStoreClient::Connect(InSessionHost.GetData());
+	if (!StoreClient)
+	{
+		return FString();
+	}
+
+	// Note: The following calls to the store client are scoped because FStoreClient::FTraceInfo and FStoreClient::FStatus share the same buffer
+	// under the hood, therefor although the ptr will still be valid, the data will not after we obtain one after the other
+	FString FileName;
+	{
+		const FStoreClient::FTraceInfo* TraceInfo = StoreClient->GetTraceInfoById(SessionID);
+		if (!TraceInfo)
+		{
+			return FString();
+		}
+
+		const FUtf8StringView Utf8NameView = TraceInfo->GetName();
+		FileName = FString(Utf8NameView);
+		if (!FileName.EndsWith(TEXT(".utrace")))
+		{
+			FileName += TEXT(".utrace");
+		}
+	}
+
+	FString TraceStorePath;
+	{
+		const FStoreClient::FStatus* StoreStatus = StoreClient->GetStatus();
+		if (!StoreStatus)
+		{
+			return FString();
+		}
+
+		TraceStorePath = FString(StoreStatus->GetStoreDir());
+	}
+	
+	FString FullFileName = FPaths::Combine(TraceStorePath, FileName);
+	FPaths::NormalizeFilename(FullFileName);
+
+	return FullFileName;
+}
+
+
+UE::Trace::FStoreClient::FTraceData FChaosVDTraceManager::GetTraceDataStreamFromStore(FStringView InSessionHost, uint32 SessionID)
+{
+	using namespace UE::Trace;
+	FStoreClient* StoreClient = FStoreClient::Connect(InSessionHost.GetData());
+	if (!StoreClient)
+	{
+		return FStoreClient::FTraceData();
+	}
+
+	return StoreClient->ReadTrace(SessionID);
+}
+
+FString FChaosVDTraceManager::GetDefaultSaveDirectoryPath()
+{
+	return FPaths::Combine(FPlatformProcess::UserDir(), TEXT("ChaosVisualDebugger"));
+}
+
+TUniquePtr<UE::Trace::IInDataStream> FChaosVDTraceManager::CreateFileDataStream(TUniquePtr<IFileHandle>&& InFileHandle)
+{
+	check(InFileHandle);
+
+	struct FFileDataStream : public UE::Trace::IInDataStream
+	{
+		explicit FFileDataStream(TUniquePtr<IFileHandle>&& FileHandle) :
+			Handle(MoveTemp(FileHandle)),
+			AsyncProgressNotification(Chaos::VD::FAsyncProgressNotification(NSLOCTEXT("ChaosVisualDebugger", "FileLoadingProgressNotification","Loading File")))
+		{
+			check(Handle);
+			FileSize = Handle->Size();
+			RemainingDataSize = FileSize;
+			constexpr int32 FileFullyLoadedPercentage = 100;
+
+			check(FileSize > 0);
+			ProgressPercentagePerByte = static_cast<double>(FileFullyLoadedPercentage) / static_cast<double>(FileSize);
+		}
+		
+		virtual int32 Read(void* Data, uint32 Size) override
+		{
+			if (RemainingDataSize <= 0)
+			{
+				return 0;
+			}
+
+			if (Size > RemainingDataSize)
+			{
+				Size = static_cast<uint32>(RemainingDataSize);
+			}
+		
+			RemainingDataSize -= Size;
+
+			double ProgressPercentage = static_cast<double>(FileSize - RemainingDataSize) * ProgressPercentagePerByte;
+			AsyncProgressNotification.EnterProgress(FMath::CeilToInt32(ProgressPercentage));
+
+			check(Handle);
+			if (!Handle->Read(static_cast<uint8*>(Data), Size))
+			{
+				return 0;
+			}
+			return Size;
+		}
+
+		virtual void Close() override
+		{
+			Handle.Reset();
+		}
+
+	private:
+		
+		TUniquePtr<IFileHandle> Handle;
+		uint64 RemainingDataSize = 0;
+		uint64 FileSize = 0;
+		double ProgressPercentagePerByte = 0.0;
+	
+		Chaos::VD::FAsyncProgressNotification AsyncProgressNotification;
+	};
+
+	FFileDataStream* FileStream = new FFileDataStream(MoveTemp(InFileHandle));
+
+	TUniquePtr<UE::Trace::IInDataStream> DataStream(FileStream);
+
+	return DataStream;
+}
+
+TUniquePtr<Chaos::VD::FRelayDataStream> FChaosVDTraceManager::CreateRelayDataStream(FGuid RemoteSessionID)
+{
+	TUniquePtr<FArchiveFileWriterGeneric> TraceFileWriter = nullptr;
+
+	if (UChaosVDGeneralSettings* Settings = FChaosVDSettingsManager::Get().GetSettingsObject<UChaosVDGeneralSettings>())
+	{
+		if (Settings->bSaveRecordingsToDisk)
+		{
+			TraceFileWriter = OpenTraceFileForWrite(Settings->DefaultSavePathOverride.FilePath.IsEmpty() ? GetDefaultSaveDirectoryPath() : Settings->DefaultSavePathOverride.FilePath);
+			if (!TraceFileWriter)
+			{
+				UE_LOG(LogChaosVDEditor, Error, TEXT("[%hs] Failed to open file for relay trace stream"), __func__);
+				return nullptr;
+			}
+		}
+	}
+
+	TUniquePtr<Chaos::VD::FRelayDataStream> NewDataStream = MakeUnique<Chaos::VD::FRelayDataStream>(RemoteSessionID, MoveTemp(TraceFileWriter));
+	
+	return MoveTemp(NewDataStream);
+}
+
+void FChaosVDTraceManager::SetPendingExternalRecordingToProcess(const TSharedPtr<FChaosVDRecording>& ExistingRecordingPtr)
+{
+	TWeakPtr<FChaosVDRecording>& PendingRecording = FChaosVDTraceManagerThreadContext::Get().PendingExternalRecordingWeakPtr;
+	ensureMsgf(!PendingRecording.Pin(), TEXT("Attempted to start a secondary trace session before a pending recording instance was processed"));
+	PendingRecording = ExistingRecordingPtr;
+}
+
+bool FChaosVDTraceManager::ConnectToLiveSession_Internal(uint32 SessionID, const TSharedPtr<FChaosVDRecording>& ExistingRecordingPtr, const FString& InRequestedSessionName, UE::Trace::FStoreClient::FTraceData&& InTraceDataStream)
+{
+	if (!InTraceDataStream.IsValid())
+	{
+		return false;
+	}
+
+	ITraceServicesModule& TraceServicesModule = FModuleManager::LoadModuleChecked<ITraceServicesModule>("TraceServices");
+	if (const TSharedPtr<TraceServices::IAnalysisService> TraceAnalysisService = TraceServicesModule.GetAnalysisService())
+	{
+		// Close this session in case we were already analysing it
+		CloseSession(InRequestedSessionName);
+		
+		SetPendingExternalRecordingToProcess(ExistingRecordingPtr);
+	
+		if (const TSharedPtr<const TraceServices::IAnalysisSession> NewSession = TraceAnalysisService->StartAnalysis(SessionID, *InRequestedSessionName, MoveTemp(InTraceDataStream)))
+		{
+			AnalysisSessionByName.Add(InRequestedSessionName, NewSession.ToSharedRef());
+			return true;
+		}
+	}
+
+	return false;
+}
+
+
+FString FChaosVDTraceManager::ConnectToLiveSession(FStringView InSessionHost, uint32 SessionID, const TSharedPtr<FChaosVDRecording>& ExistingRecordingPtr)
+{
+	if (InSessionHost.IsEmpty())
+	{
+		return FString();
+	}
+
+	using namespace UE::Trace;
+	FString LiveSessionName = FString::Printf(TEXT("LiveSession[%.*s - %u]"), InSessionHost.Len(), InSessionHost.GetData(), SessionID);
+
+	FStoreClient::FTraceData TraceData = GetTraceDataStreamFromStore(InSessionHost, SessionID);
+	if (!TraceData.IsValid())
+	{
+		UE_LOG(LogChaosVDEditor, Error, TEXT("[%hs] Failed to get trace data stream from the trace store | Host [%s] | Session ID [%u]"), __func__, InSessionHost.GetData(), SessionID);
+		return FString();
+	}
+
+	if (ConnectToLiveSession_Internal(SessionID, ExistingRecordingPtr, LiveSessionName, MoveTemp(TraceData)))
+	{
+		 return LiveSessionName;
+	}
+
+	UE_LOG(LogChaosVDEditor, Error, TEXT("[%hs] Failed connect to session | Host [%s] | Session ID [%u]"), __func__, InSessionHost.GetData(), SessionID);
+
+	return FString();
+}
+
+TUniquePtr<FArchiveFileWriterGeneric> FChaosVDTraceManager::OpenTraceFileForWrite(const FString& InDirectoryToSavePath)
+{
+	IPlatformFile& PlatformFile = FPlatformFileManager::Get().GetPlatformFile();
+
+	if (InDirectoryToSavePath.IsEmpty())
+	{
+		UE_LOG(LogChaosVDEditor, Error, TEXT("[%hs] Failed open file for write | No save directory was provided"), __func__);
+		return nullptr;
+	}
+
+	PlatformFile.CreateDirectory(*InDirectoryToSavePath);
+	if (!PlatformFile.DirectoryExists(*InDirectoryToSavePath))
+	{
+		UE_LOG(LogChaosVDEditor, Error, TEXT("[%hs] Failed open file for write | Failed to create or access provided save directory [%s]"), __func__, *InDirectoryToSavePath);
+		return nullptr;
+	}
+
+	TUniquePtr<FArchiveFileWriterGeneric> FileWriter = nullptr;
+
+#if WITH_CHAOS_VISUAL_DEBUGGER
+
+	FChaosVDRuntimeModule& CVDRuntimeModule = FChaosVDRuntimeModule::Get();
+	
+	constexpr int32 MaxAttempts = 10;
+	int32 CurrentAttempts = 0;
+	while (CurrentAttempts < MaxAttempts)
+	{
+		FString OutFileName;
+		CVDRuntimeModule.GenerateRecordingFileName(OutFileName);
+		OutFileName = FPaths::Combine(InDirectoryToSavePath, OutFileName);
+
+		if (IFileHandle* FileHandle = PlatformFile.OpenWrite(*OutFileName))
+		{
+			FileWriter = MakeUnique<FArchiveFileWriterGeneric>(FileHandle, *OutFileName, FileHandle->Tell());
+			break;
+		}
+	
+		CurrentAttempts++;
+
+		FPlatformProcess::Sleep(0.1f);
+	}
+
+#endif
+
+	return MoveTemp(FileWriter);
+}
+
+FString FChaosVDTraceManager::ConnectToLiveSession_Direct(FGuid RemoteSessionID, uint16& OutSessionPort, const TSharedPtr<FChaosVDRecording>& ExistingRecordingPtr)
+{
+	TUniquePtr<FArchiveFileWriterGeneric> TraceFileWriter = nullptr;
+
+	if (UChaosVDGeneralSettings* Settings = FChaosVDSettingsManager::Get().GetSettingsObject<UChaosVDGeneralSettings>())
+	{
+		if (Settings->bSaveRecordingsToDisk)
+		{
+			TraceFileWriter = OpenTraceFileForWrite(Settings->DefaultSavePathOverride.FilePath.IsEmpty() ? GetDefaultSaveDirectoryPath() : Settings->DefaultSavePathOverride.FilePath);
+			if (!TraceFileWriter)
+			{
+				UE_LOG(LogChaosVDEditor, Error, TEXT("[%hs] Failed to open open file for direct trace stream"), __func__);
+				return FString();
+			}
+		}
+	}
+
+	TUniquePtr<Chaos::VD::FDirectSocketStream> DirectSocketStream = MakeUnique<Chaos::VD::FDirectSocketStream>(RemoteSessionID, MoveTemp(TraceFileWriter));
+	OutSessionPort = DirectSocketStream->StartListening(CurrentDirectModeStartingPort);
+
+	if (OutSessionPort == 0)
+	{
+		UE_LOG(LogChaosVDEditor, Error, TEXT("[%hs] Failed to open direct trace stream socket"), __func__);
+		return FString();
+	}
+
+	// If the port was in use, FDirectSocketStream will try multiple ports in sequence, therefore the next port we should try while we have recording active (common case in Multi-session mode)
+	// is the next one from the last success
+	CurrentDirectModeStartingPort = OutSessionPort + 1;
+
+
+	FString DirectSessionName = FString::Printf(TEXT("DirectSession[127.0.0.1:%u]"), OutSessionPort);
+	
+	if (ConnectToLiveSession_Internal(~0u, ExistingRecordingPtr, DirectSessionName, MoveTemp(DirectSocketStream)))
+	{
+		return DirectSessionName;
+	}
+	else
+	{
+		UE_LOG(LogChaosVDEditor, Error, TEXT("[%hs] Failed start trace analysis with direct trace stream socket | [%s]"), __func__, *DirectSessionName);
+		return FString();
+	}
+}
+
+FString FChaosVDTraceManager::ConnectToLiveSession_Relay(FGuid RemoteSessionID, const TSharedPtr<FChaosVDRecording>& ExistingRecordingPtr)
+{
+	FString RelaySessionName = FString::Printf(TEXT("RelaySession[%s]"), *RemoteSessionID.ToString());
+
+	TUniquePtr<Chaos::VD::FRelayDataStream> RelayDataStream = CreateRelayDataStream(RemoteSessionID);
+
+	if (ConnectToLiveSession_Internal(~0u, ExistingRecordingPtr, RelaySessionName, MoveTemp(RelayDataStream)))
+	{
+		return RelaySessionName;
+	}
+	else
+	{
+		UE_LOG(LogChaosVDEditor, Error, TEXT("[%hs] Failed start trace analysis with relay trace stream | [%s]"), __func__, *RelaySessionName);
+		return FString();
+	}
+}
+
+FString FChaosVDTraceManager::GetLocalTraceStoreDirPath()
+{
+	UE::Trace::FStoreClient* StoreClient = UE::Trace::FStoreClient::Connect(TEXT("localhost"));
+
+	if (!StoreClient)
+	{
+		UE_LOG(LogChaosVDEditor, Error, TEXT("[%s] Failed to connect to local Trace Store client"), ANSI_TO_TCHAR(__FUNCTION__));
+		return TEXT("");
+	}
+
+	const UE::Trace::FStoreClient::FStatus* Status = StoreClient->GetStatus();
+	if (!Status)
+	{
+		UE_LOG(LogChaosVDEditor, Error, TEXT("[%s] Failed to to get Trace Store status"), ANSI_TO_TCHAR(__FUNCTION__));
+		return TEXT("");
+	}
+
+	return FString(Status->GetStoreDir());
+}
+
+TSharedPtr<const TraceServices::IAnalysisSession> FChaosVDTraceManager::GetSession(const FString& InSessionName)
+{
+	if (TSharedPtr<const TraceServices::IAnalysisSession>* FoundSession = AnalysisSessionByName.Find(InSessionName))
+	{
+		return *FoundSession;
+	}
+
+	return nullptr;
+}
+
+void FChaosVDTraceManager::CloseSession(const FString& InSessionName)
+{
+	if (const TSharedPtr<const TraceServices::IAnalysisSession>* Session = AnalysisSessionByName.Find(InSessionName))
+	{
+		if (Session->IsValid())
+		{
+			(*Session)->Stop(true);
+		}
+
+		AnalysisSessionByName.Remove(InSessionName);
+
+		if (AnalysisSessionByName.IsEmpty())
+		{
+			// If we don't have any more session left, we can reset back to the original starting port
+			// As all sockets we had in use are closed at this point
+			CurrentDirectModeStartingPort = DefaultStartingPort;
+		}
+	}
+}
+
+void FChaosVDTraceManager::StopSession(const FString& InSessionName)
+{
+	if (const TSharedPtr<const TraceServices::IAnalysisSession>* Session = AnalysisSessionByName.Find(InSessionName))
+	{
+		if (Session->IsValid())
+		{
+			(*Session)->Stop(true);
+		}
+	}
+}
diff --git a/Engine/Plugins/ChaosVD/Source/ChaosVD/Private/Widgets/ChaosVDRecordingControls.cpp b/Engine/Plugins/ChaosVD/Source/ChaosVD/Private/Widgets/ChaosVDRecordingControls.cpp
--- a/Engine/Plugins/ChaosVD/Source/ChaosVD/Private/Widgets/ChaosVDRecordingControls.cpp
+++ b/Engine/Plugins/ChaosVD/Source/ChaosVD/Private/Widgets/ChaosVDRecordingControls.cpp
@@ -1,4 +1,4 @@
-﻿// Copyright Epic Games, Inc. All Rights Reserved.
+// Copyright Epic Games, Inc. All Rights Reserved.
 
 #include "Widgets/ChaosVDRecordingControls.h"
 
@@ -26,8 +26,6 @@
 #include "ToolMenu.h"
 #include "ToolMenus.h"
 #include "TraceDataRelayTransport.h"
-#include "TraceSessionsManager.h"
-#include "Trace/StoreClient.h"
 #include "Widgets/Images/SImage.h"
 #include "Widgets/Input/SButton.h"
 #include "Widgets/Input/SComboButton.h"
@@ -657,7 +655,7 @@ bool FChaosVDRecordingControls::CanChangeLoadingMode() const
 		else if (const TSharedPtr<SChaosVDMainTab> MainTabSharedPtr = MainTabWeakPtr.Pin())
 		{
 			// If nothing is loaded yet, it does not make sense change the loading mode
-			return !MainTabSharedPtr->GetChaosVDEngineInstance()->GetTraceSessionDescriptors().IsEmpty();
+			return !MainTabSharedPtr->GetChaosVDEngineInstance()->GetCurrentSessionDescriptors().IsEmpty();
 		}
 	}
 
@@ -935,7 +933,7 @@ FChaosVDRecordingControls::FTraceStoreConnectionSettings FChaosVDRecordingContro
 
 	const FTraceConnectionDetails& RecordingSessionDetails = InSessionInfoRef->GetTraceConnectionDetails();
 
-	if (const UE::Trace::FStoreClient::FSessionInfo* TraceSessionInfo = FTraceSessionsManager::GetTraceSessionInfo(RecordingSessionDetails.TraceTarget, RecordingSessionDetails.TraceGuid))
+	if (const UE::Trace::FStoreClient::FSessionInfo* TraceSessionInfo = FChaosVDTraceManager::GetTraceSessionInfo(RecordingSessionDetails.TraceTarget, RecordingSessionDetails.TraceGuid))
 	{
 		ConnectionSettings.SessionAddress = RecordingSessionDetails.TraceTarget;
 		ConnectionSettings.TraceID = TraceSessionInfo->GetTraceId();
@@ -1071,7 +1069,7 @@ void FChaosVDRecordingControls::SetSessionRecordingState(bool bIsRecording, cons
 			}
 
 			// If we get here it means we have at least one valid session
-			const FTraceSessionDescriptor& SessionDesc = MainTabSharedPtr->GetChaosVDEngineInstance()->GetTraceSessionDescriptors().Last();
+			const FChaosVDTraceSessionDescriptor& SessionDesc = MainTabSharedPtr->GetChaosVDEngineInstance()->GetCurrentSessionDescriptors().Last();
 
 			RecordingParams.Target.Appendf(TEXT(":%u"), SessionDesc.SessionPort);
 
diff --git a/Engine/Plugins/ChaosVD/Source/ChaosVD/Private/Widgets/SChaosVDMainTab.cpp b/Engine/Plugins/ChaosVD/Source/ChaosVD/Private/Widgets/SChaosVDMainTab.cpp
--- a/Engine/Plugins/ChaosVD/Source/ChaosVD/Private/Widgets/SChaosVDMainTab.cpp
+++ b/Engine/Plugins/ChaosVD/Source/ChaosVD/Private/Widgets/SChaosVDMainTab.cpp
@@ -61,7 +61,7 @@
 #include "Styling/StyleColors.h"
 #include "Styling/ToolBarStyle.h"
 #include "TabSpawers/ChaosVDSceneQueryBrowserTab.h"
-#include "TraceSessionsManager.h"
+#include "Trace/ChaosVDTraceManager.h"
 #include "Visualizers/ChaosVDCharacterGroundConstraintsDataComponentVisualizer.h"
 #include "Visualizers/ChaosVDJointConstraintsDataComponentVisualizer.h"
 #include "Visualizers/ChaosVDParticleDataComponentVisualizer.h"
@@ -394,9 +394,9 @@ void SChaosVDMainTab::UpdateTraceProcessingNotificationUpdate()
 	TSharedRef<FChaosVDEngine> CVDEngineInstance = GetChaosVDEngineInstance();
 
 	// Live sessions should not show the notification pop up. When CVD is connected in live debugging mode the UI already has feedback to inform the user.
-	if (CVDEngineInstance->GetTraceSessionDescriptors().Num() > 0 && !CVDEngineInstance->HasAnyLiveSessionActive())
+	if (CVDEngineInstance->GetCurrentSessionDescriptors().Num() > 0 && !CVDEngineInstance->HasAnyLiveSessionActive())
 	{
-		FChaosVDModule::Get().GetTraceSessionsManager()->EnumerateActiveSessions([&bIsTraceAnalysisInProgress](const TSharedRef<const TraceServices::IAnalysisSession>& InSession)
+		FChaosVDModule::Get().GetTraceManager()->EnumerateActiveSessions([&bIsTraceAnalysisInProgress](const TSharedRef<const TraceServices::IAnalysisSession>& InSession)
 		{
 			// TODO: We should implement a way to gather the actual progress so we can report that in the UI
 			// but the Trace API don't have that data yet.
@@ -727,14 +727,15 @@ FReply SChaosVDMainTab::BrowseAndOpenChaosVDRecording()
 			}
 		case EChaosVDBrowseFileModalResponse::CurrentSaveFolder:
 			{
-				// GetSaveDirPath will return the path to use base on the current ChaosVDGeneralSettings (i.e., using DefaultSavePathOverride if set)
-				const FString FolderPath = FChaosVDModule::Get().GetTraceSessionsManager()->GetSaveDirPath();
+				UChaosVDGeneralSettings* Settings = FChaosVDSettingsManager::Get().GetSettingsObject<UChaosVDGeneralSettings>();
+				const FString FolderPath = Settings && !Settings->DefaultSavePathOverride.FilePath.IsEmpty() ? Settings->DefaultSavePathOverride.FilePath : FChaosVDTraceManager::GetDefaultSaveDirectoryPath();
+
 				BrowseChaosVDRecordingFromFolder(*FolderPath, LoadingMode);
 				break;
 			}
 		case EChaosVDBrowseFileModalResponse::LocalTraceStore:
 			{
-				const FString TraceStorePath = FChaosVDModule::Get().GetTraceSessionsManager()->GetLocalTraceStoreDirPath();
+				const FString TraceStorePath = FChaosVDModule::Get().GetTraceManager()->GetLocalTraceStoreDirPath();
 				if (TraceStorePath.IsEmpty())
 				{
 					UE_LOG(LogChaosVDEditor, Error, TEXT("[%hs] Failed to access Trace Store..."), __func__);
@@ -982,7 +983,7 @@ bool SChaosVDMainTab::ConnectToLiveSession(int32 SessionID, const FString& InSes
 
 bool SChaosVDMainTab::ConnectToLiveSession_Direct(EChaosVDLoadRecordedDataMode LoadingMode) const
 {
-	UE_LOG(LogChaosVDEditor, Error, TEXT("[%hs] This method is no longer supported. Please use the overload that takes a Remote Session ID."), __func__);
+	UE_LOG(LogChaosVDEditor, Error, TEXT("[%hs] This menthod is no longer supported. Please use the overload that takes a Remote Session ID."), __func__);
 	return false;
 }
 
@@ -1242,12 +1243,10 @@ void SChaosVDMainTab::HandleDisconnectSessionClicked()
 	
 		if (bIsAlreadyInLiveSession)
 		{
-			using namespace UE::TraceBasedDebuggers;
-			const TSharedPtr<FTraceSessionsManager> SessionsManager = FChaosVDModule::Get().GetTraceSessionsManager();
-			const TArrayView<FTraceSessionDescriptor> ActiveSessions = GetChaosVDEngineInstance()->GetTraceSessionDescriptors();
-			for (FTraceSessionDescriptor& ActiveSession : ActiveSessions)
+			TArrayView<FChaosVDTraceSessionDescriptor> ActiveSessions = GetChaosVDEngineInstance()->GetCurrentSessionDescriptors();
+			for (FChaosVDTraceSessionDescriptor& ActiveSession : ActiveSessions)
 			{
-				SessionsManager->CloseSession(ActiveSession.SessionName);
+				FChaosVDModule::Get().GetTraceManager()->CloseSession(ActiveSession.SessionName);
 				ActiveSession.bIsLiveSession = false;
 			}
 
@@ -1260,7 +1259,7 @@ FText SChaosVDMainTab::GetDisconnectButtonText() const
 {
 	if (GetChaosVDEngineInstance()->HasAnyLiveSessionActive())
 	{
-		if (GetChaosVDEngineInstance()->GetTraceSessionDescriptors().Num() > 1)
+		if (GetChaosVDEngineInstance()->GetCurrentSessionDescriptors().Num() > 1)
 		{
 			return LOCTEXT("DisconnectFromMultipleSessions", "Disconnect from all Sessions");
 		}
diff --git a/Engine/Plugins/ChaosVD/Source/ChaosVD/Private/Widgets/SChaosVDRecordedLogBrowser.cpp b/Engine/Plugins/ChaosVD/Source/ChaosVD/Private/Widgets/SChaosVDRecordedLogBrowser.cpp
--- a/Engine/Plugins/ChaosVD/Source/ChaosVD/Private/Widgets/SChaosVDRecordedLogBrowser.cpp
+++ b/Engine/Plugins/ChaosVD/Source/ChaosVD/Private/Widgets/SChaosVDRecordedLogBrowser.cpp
@@ -3,15 +3,14 @@
 #include "SChaosVDRecordedLogBrowser.h"
 
 #include "ChaosVDEngine.h"
-#include "ChaosVDModule.h"
 #include "ChaosVDPlaybackController.h"
 #include "ChaosVDScene.h"
-#include "Components/VerticalBox.h"
-#include "Misc/TextFilterExpressionEvaluator.h"
 #include "SChaosVDLogBrowserToolbar.h"
 #include "SChaosVDRecordedLogView.h"
+#include "SlateOptMacros.h"
 #include "ToolMenus.h"
-#include "TraceSessionsManager.h"
+#include "Components/VerticalBox.h"
+#include "Misc/TextFilterExpressionEvaluator.h"
 #include "TraceServices/Model/AnalysisSession.h"
 #include "TraceServices/Model/Log.h"
 
@@ -157,7 +156,7 @@ void SChaosVDRecordedLogBrowser::UpdateBrowserContents()
 
 	if (!CurrentSessionName.IsEmpty())
 	{
-		Session = FChaosVDModule::Get().GetTraceSessionsManager()->GetSession(CurrentSessionName);
+		Session = FChaosVDModule::Get().GetTraceManager()->GetSession(CurrentSessionName);
 	}
 	
 	if (!Session)
@@ -460,7 +459,7 @@ void SChaosVDRecordedLogBrowser::HandleItemFocused(const TSharedPtr<FChaosVDLogV
 
 	TSharedPtr<FChaosVDEngine> CVDEngineInstance = EngineInstanceWeakPtr.Pin();
 	TSharedPtr<const TraceServices::IAnalysisSession> Session = CVDEngineInstance && !CurrentSessionName.IsEmpty() ?
-																FChaosVDModule::Get().GetTraceSessionsManager()->GetSession(CurrentSessionName) :
+																FChaosVDModule::Get().GetTraceManager()->GetSession(CurrentSessionName) :
 																nullptr;
 
 	if (!Session)
@@ -611,7 +610,7 @@ void SChaosVDRecordedLogBrowser::ProcessPendingTrackUpdates()
 	}
 
 	TSharedPtr<FChaosVDEngine> CVDEngineInstance = EngineInstanceWeakPtr.Pin();
-	TSharedPtr<UE::TraceBasedDebuggers::FTraceSessionsManager> TraceManager = FChaosVDModule::Get().GetTraceSessionsManager();
+	TSharedPtr<FChaosVDTraceManager> TraceManager = FChaosVDModule::Get().GetTraceManager();
 	TSharedPtr<const TraceServices::IAnalysisSession> Session = CVDEngineInstance && TraceManager && !CurrentSessionName.IsEmpty() ?
 																TraceManager->GetSession(CurrentSessionName) :
 																nullptr;
diff --git a/Engine/Plugins/ChaosVD/Source/ChaosVD/Private/Widgets/SChaosVDRecordedLogBrowser.h b/Engine/Plugins/ChaosVD/Source/ChaosVD/Private/Widgets/SChaosVDRecordedLogBrowser.h
--- a/Engine/Plugins/ChaosVD/Source/ChaosVD/Private/Widgets/SChaosVDRecordedLogBrowser.h
+++ b/Engine/Plugins/ChaosVD/Source/ChaosVD/Private/Widgets/SChaosVDRecordedLogBrowser.h
@@ -2,10 +2,11 @@
 
 #pragma once
 
-#include "Async/AsyncWork.h"
+#include "CoreMinimal.h"
 #include "ChaosVDPlaybackControllerInstigator.h"
 #include "ChaosVDPlaybackControllerObserver.h"
 #include "ContentBrowserDataMenuContexts.h"
+#include "Async/AsyncWork.h"
 #include "Framework/Docking/TabManager.h"
 #include "Misc/TextFilterExpressionEvaluator.h"
 #include "TraceServices/Model/Log.h"
@@ -13,6 +14,7 @@
 
 #include "SChaosVDRecordedLogBrowser.generated.h"
 
+struct FChaosVDTraceSessionDescriptor;
 class FTabManager;
 class FChaosVDEngine;
 class SChaosVDLogBrowserToolbar;
diff --git a/Engine/Plugins/ChaosVD/Source/ChaosVD/Private/Widgets/SChaosVDRecordedLogInstances.cpp b/Engine/Plugins/ChaosVD/Source/ChaosVD/Private/Widgets/SChaosVDRecordedLogInstances.cpp
--- a/Engine/Plugins/ChaosVD/Source/ChaosVD/Private/Widgets/SChaosVDRecordedLogInstances.cpp
+++ b/Engine/Plugins/ChaosVD/Source/ChaosVD/Private/Widgets/SChaosVDRecordedLogInstances.cpp
@@ -2,14 +2,14 @@
 
 #include "SChaosVDRecordedLogInstances.h"
 
-#include "Application/SlateApplicationBase.h"
 #include "ChaosVDEngine.h"
 #include "ChaosVDModule.h"
 #include "ChaosVDStyle.h"
 #include "SChaosVDRecordedLogBrowser.h"
-#include "TraceSessionsManager.h"
-#include "Widgets/Docking/SDockTab.h"
+#include "Application/SlateApplicationBase.h"
+#include "Trace/ChaosVDTraceManager.h"
 #include "Widgets/SToolTip.h"
+#include "Widgets/Docking/SDockTab.h"
 #include "Widgets/Text/STextBlock.h"
 
 #define LOCTEXT_NAMESPACE "ChaosVisualDebugger"
@@ -20,8 +20,8 @@ SChaosVDRecordedLogInstances::~SChaosVDRecordedLogInstances()
 
 	if (TSharedPtr<FChaosVDEngine> EngineInstance = EngineInstanceWeakPtr.Pin())
 	{
-		EngineInstance->OnTraceSessionOpened().RemoveAll(this);
-		EngineInstance->OnTraceSessionClosed().RemoveAll(this);
+		EngineInstance->OnSessionOpened().RemoveAll(this);
+		EngineInstance->OnSessionClosed().RemoveAll(this);
 	}
 }
 
@@ -43,21 +43,20 @@ void SChaosVDRecordedLogInstances::Construct(const FArguments& InArgs, const TSh
 		]
 	];
 
-	InEngineInstance->OnTraceSessionClosed().AddSP(this, &SChaosVDRecordedLogInstances::HandleSessionClosed);
-	InEngineInstance->OnTraceSessionOpened().AddSP(this, &SChaosVDRecordedLogInstances::HandleSessionOpened);
+	InEngineInstance->OnSessionClosed().AddSP(this, &SChaosVDRecordedLogInstances::HandleSessionClosed);
+	InEngineInstance->OnSessionOpened().AddSP(this, &SChaosVDRecordedLogInstances::HandleSessionOpened);
 
-	using namespace UE::TraceBasedDebuggers;
-	TArrayView<FTraceSessionDescriptor> SessionDescriptors = InEngineInstance->GetTraceSessionDescriptors();
+	TArrayView<FChaosVDTraceSessionDescriptor> SessionDescriptors = InEngineInstance->GetCurrentSessionDescriptors();
 	if (!SessionDescriptors.IsEmpty())
 	{
-		for (const FTraceSessionDescriptor& SessionDescriptor : SessionDescriptors)
+		for (const FChaosVDTraceSessionDescriptor& SessionDescriptor : SessionDescriptors)
 		{
 			HandleSessionOpened(SessionDescriptor);
 		}
 	}
 }
 
-void SChaosVDRecordedLogInstances::HandleSessionOpened(const UE::TraceBasedDebuggers::FTraceSessionDescriptor& InSessionDescriptor)
+void SChaosVDRecordedLogInstances::HandleSessionOpened(const FChaosVDTraceSessionDescriptor& InSessionDescriptor)
 {
 	FName TabID = GetAvailableTabID();
 
@@ -66,7 +65,7 @@ void SChaosVDRecordedLogInstances::HandleSessionOpened(const UE::TraceBasedDebug
 	InstancesTabManager->TryInvokeTab(TabID);
 }
 
-void SChaosVDRecordedLogInstances::HandleSessionClosed(const UE::TraceBasedDebuggers::FTraceSessionDescriptor& InSessionDescriptor)
+void SChaosVDRecordedLogInstances::HandleSessionClosed(const FChaosVDTraceSessionDescriptor& InSessionDescriptor)
 {
 	if (TWeakPtr<SDockTab>* ActiveTabForSessionPtrPtr = ActiveTabsBySessionName.Find(InSessionDescriptor.SessionName))
 	{
diff --git a/Engine/Plugins/ChaosVD/Source/ChaosVD/Private/Widgets/SChaosVDRecordedLogInstances.h b/Engine/Plugins/ChaosVD/Source/ChaosVD/Private/Widgets/SChaosVDRecordedLogInstances.h
--- a/Engine/Plugins/ChaosVD/Source/ChaosVD/Private/Widgets/SChaosVDRecordedLogInstances.h
+++ b/Engine/Plugins/ChaosVD/Source/ChaosVD/Private/Widgets/SChaosVDRecordedLogInstances.h
@@ -1,16 +1,12 @@
-﻿// Copyright Epic Games, Inc. All Rights Reserved.
+// Copyright Epic Games, Inc. All Rights Reserved.
 
 #pragma once
 
 #include "Framework/Docking/TabManager.h"
 #include "Widgets/SCompoundWidget.h"
 
-namespace UE::TraceBasedDebuggers
-{
-struct FTraceSessionDescriptor;
-}
-
 class FSpawnTabArgs;
+struct FChaosVDTraceSessionDescriptor;
 class FChaosVDEngine;
 class SDockTab;
 
@@ -31,8 +27,8 @@ public:
 
 private:
 
-	void HandleSessionOpened(const UE::TraceBasedDebuggers::FTraceSessionDescriptor& InSessionDescriptor);
-	void HandleSessionClosed(const UE::TraceBasedDebuggers::FTraceSessionDescriptor& InSessionDescriptor);
+	void HandleSessionOpened(const FChaosVDTraceSessionDescriptor& InSessionDescriptor);
+	void HandleSessionClosed(const FChaosVDTraceSessionDescriptor& InSessionDescriptor);
 
 	virtual TSharedRef<SDockTab> HandleTabSpawnRequest(const FSpawnTabArgs& Args, FName TabSessionID);
 
diff --git a/Engine/Plugins/ChaosVD/Source/ChaosVD/Public/ChaosVDEngine.h b/Engine/Plugins/ChaosVD/Source/ChaosVD/Public/ChaosVDEngine.h
--- a/Engine/Plugins/ChaosVD/Source/ChaosVD/Public/ChaosVDEngine.h
+++ b/Engine/Plugins/ChaosVD/Source/ChaosVD/Public/ChaosVDEngine.h
@@ -2,25 +2,22 @@
 
 #pragma once
 
+#include "Templates/SharedPointer.h"
 #include "Containers/Ticker.h"
 #include "Misc/Guid.h"
-#include "Templates/SharedPointer.h"
+#include "Trace/ChaosVDTraceManager.h"
+
 #include "ChaosVDEngine.generated.h"
 
 namespace Chaos::VisualDebugger
 {
-struct FChaosVDOptionalDataChannel;
-}
-
-namespace UE::TraceBasedDebuggers
-{
-struct FTraceSessionDescriptor;
+	struct FChaosVDOptionalDataChannel;
 }
 
+class FChaosVDTraceManager;
 class FChaosVDPlaybackController;
 class FChaosVDScene;
 class FChaosVisualDebuggerMainUI;
-struct FChaosVDRecording;
 
 /** Enumeration of the available modes controlling how data is loaded into CVD */
 UENUM()
@@ -32,19 +29,16 @@ enum class EChaosVDLoadRecordedDataMode : uint8
 	MultiSource
 };
 
-DECLARE_MULTICAST_DELEGATE_OneParam(FTraceSessionStateChangedDelegate, const UE::TraceBasedDebuggers::FTraceSessionDescriptor& InSessionDescriptor)
-
-PRAGMA_DISABLE_DEPRECATION_WARNINGS
-UE_DEPRECATED(5.8, "Use FTraceSessionStateChangedDelegate taking UE::TraceBasedDebuggers::FTraceSessionDescriptor in parameter")
-DECLARE_MULTICAST_DELEGATE_OneParam(FSessionStateChangedDelegate, const struct FChaosVDTraceSessionDescriptor& InSessionDescriptor)
-PRAGMA_ENABLE_DEPRECATION_WARNINGS
+DECLARE_MULTICAST_DELEGATE_OneParam(FSessionStateChangedDelegate, const FChaosVDTraceSessionDescriptor& InSessionDescriptor)
 
 /** Core Implementation of the visual debugger - Owns the systems that are not UI */
 class FChaosVDEngine : public FTSTickerObjectBase, public TSharedFromThis<FChaosVDEngine>
 {
 public:
-	CHAOSVD_API FChaosVDEngine();
-	CHAOSVD_API virtual ~FChaosVDEngine();
+	FChaosVDEngine()
+	{
+		InstanceGUID = FGuid::NewGuid();
+	}
 
 	CHAOSVD_API void Initialize();
 	CHAOSVD_API void CloseActiveTraceSessions();
@@ -71,69 +65,41 @@ public:
 		return PlaybackController;
 	}
 
+	TArrayView<FChaosVDTraceSessionDescriptor> GetCurrentSessionDescriptors()
+	{
+		return CurrentSessionDescriptors;
+	}
+
 	bool HasAnyLiveSessionActive() const;
 
 	CHAOSVD_API void LoadRecording(const FString& FilePath, EChaosVDLoadRecordedDataMode LoadingMode = EChaosVDLoadRecordedDataMode::MultiSource);
 	void LoadCombinedMultiRecording(const FString& FilePath);
 
 	bool ConnectToLiveSession(uint32 SessionID, const FString& InSessionAddress, EChaosVDLoadRecordedDataMode LoadingMode = EChaosVDLoadRecordedDataMode::MultiSource);
+
+	UE_DEPRECATED(5.8, "Please use the version that takes a remote Session ID")
+	bool ConnectToLiveSession_Direct(EChaosVDLoadRecordedDataMode LoadingMode = EChaosVDLoadRecordedDataMode::MultiSource);
+	
 	bool ConnectToLiveSession_Direct(FGuid RemoteSessionID, EChaosVDLoadRecordedDataMode LoadingMode = EChaosVDLoadRecordedDataMode::MultiSource);
 	bool ConnectToLiveSession_Relay(FGuid RemoteSessionID, EChaosVDLoadRecordedDataMode LoadingMode = EChaosVDLoadRecordedDataMode::MultiSource);
 
-	void OpenSession(const UE::TraceBasedDebuggers::FTraceSessionDescriptor& SessionDescriptor, EChaosVDLoadRecordedDataMode LoadingMode = EChaosVDLoadRecordedDataMode::MultiSource);
+	void OpenSession(const FChaosVDTraceSessionDescriptor& SessionDescriptor, EChaosVDLoadRecordedDataMode LoadingMode = EChaosVDLoadRecordedDataMode::MultiSource);
 
 	bool SaveOpenSessionToCombinedFile(const FString& InTargetFilePath = FString());
 
 	bool CanCombineOpenSessions() const;
-
-PRAGMA_DISABLE_DEPRECATION_WARNINGS
-
-	UE_DEPRECATED(5.8, "Use GetTraceSessionDescriptors")
-	TArrayView<FChaosVDTraceSessionDescriptor> GetCurrentSessionDescriptors()
-	{
-		return {};
-	}
-
-	UE_DEPRECATED(5.8, "Please use the version that takes a remote Session ID")
-	bool ConnectToLiveSession_Direct(EChaosVDLoadRecordedDataMode LoadingMode = EChaosVDLoadRecordedDataMode::MultiSource);
-
-	UE_DEPRECATED(5.8, "Use the version taking UE::TraceBasedDebuggers::FTraceSessionDescriptor in parameter")
-	void OpenSession(const FChaosVDTraceSessionDescriptor& SessionDescriptor, EChaosVDLoadRecordedDataMode LoadingMode = EChaosVDLoadRecordedDataMode::MultiSource)
-	{
-	}
-
-	UE_DEPRECATED(5.8, "Use OnTraceSessionOpened instead")
+	
 	FSessionStateChangedDelegate& OnSessionOpened()
-	{
-		static FSessionStateChangedDelegate SessionOpenedDelegate;
-		return SessionOpenedDelegate;
-	}
-
-	UE_DEPRECATED(5.8, "Use OnTraceSessionClosed instead")
-	FSessionStateChangedDelegate& OnSessionClosed()
-	{
-		static FSessionStateChangedDelegate SessionClosedDelegate;
-		return SessionClosedDelegate;
-	}
-
-	UE_DEPRECATED(5.8, "Use GetTraceSessionDescriptor")
-	const FChaosVDTraceSessionDescriptor& GetSessionDescriptor(FGuid RemoteSessionID) const;
-
-PRAGMA_ENABLE_DEPRECATION_WARNINGS
-
-	FTraceSessionStateChangedDelegate& OnTraceSessionOpened()
 	{
 		return OnSessionOpenedDelegate;
 	}
 
-	FTraceSessionStateChangedDelegate& OnTraceSessionClosed()
+	FSessionStateChangedDelegate& OnSessionClosed()
 	{
 		return OnSessionClosedDelegate;
 	}
 
-	CHAOSVD_API const UE::TraceBasedDebuggers::FTraceSessionDescriptor& GetTraceSessionDescriptor(FGuid RemoteSessionID) const;
-
-	CHAOSVD_API TArrayView<UE::TraceBasedDebuggers::FTraceSessionDescriptor> GetTraceSessionDescriptors();
+	const FChaosVDTraceSessionDescriptor& GetSessionDescriptor(FGuid RemoteSessionID) const;
 
 private:
 
@@ -146,7 +112,7 @@ private:
 
 	FGuid InstanceGUID;
 
-	TArray<UE::TraceBasedDebuggers::FTraceSessionDescriptor> CurrentSessionDescriptors;
+	TArray<FChaosVDTraceSessionDescriptor> CurrentSessionDescriptors;
 
 	TSharedPtr<FChaosVDScene> CurrentScene;
 	TSharedPtr<FChaosVDPlaybackController> PlaybackController;
@@ -157,6 +123,6 @@ private:
 
 	FDelegateHandle DataChannelStateUpdatedHandle;
 	
-	FTraceSessionStateChangedDelegate OnSessionOpenedDelegate;
-	FTraceSessionStateChangedDelegate OnSessionClosedDelegate;
-};
\ No newline at end of file
+	FSessionStateChangedDelegate OnSessionOpenedDelegate;
+	FSessionStateChangedDelegate OnSessionClosedDelegate;
+};
diff --git a/Engine/Plugins/ChaosVD/Source/ChaosVD/Public/ChaosVDModule.h b/Engine/Plugins/ChaosVD/Source/ChaosVD/Public/ChaosVDModule.h
--- a/Engine/Plugins/ChaosVD/Source/ChaosVD/Public/ChaosVDModule.h
+++ b/Engine/Plugins/ChaosVD/Source/ChaosVD/Public/ChaosVDModule.h
@@ -9,17 +9,13 @@
 #include "Misc/Guid.h"
 
 class SWidget;
+class FChaosVDTraceManager;
 class FChaosVDEngine;
 class SChaosVDMainTab;
 class SDockTab;
 class FSpawnTabArgs;
 struct FGuid;
 
-namespace UE::TraceBasedDebuggers
-{
-class FTraceSessionsManager;
-}
-
 /** Main module class for the Chaos Visual Debugger editor */
 DECLARE_LOG_CATEGORY_EXTERN(LogChaosVDEditor, Log, VeryVerbose);
 class FChaosVDModule : public IModuleInterface
@@ -35,16 +31,7 @@ public:
 	virtual void ShutdownModule() override;
 
 	/** Returns the Trace manager instance used by CVD */
-	TSharedPtr<UE::TraceBasedDebuggers::FTraceSessionsManager>& GetTraceSessionsManager()
-	{
-		return ChaosVDTraceManager;
-	}
-
-	UE_DEPRECATED(5.8, "Use GetTraceSessionsManager")
-	TSharedPtr<class FChaosVDTraceManager> GetTraceManager()
-	{
-		return nullptr;
-	}
+	TSharedPtr<FChaosVDTraceManager>& GetTraceManager() { return ChaosVDTraceManager; }
 
 	/** Spawns a new instance of CVD */
 	void SpawnCVDTab();
@@ -74,7 +61,7 @@ private:
 
 	TMap<FGuid, TWeakPtr<SDockTab>> ActiveCVDTabs;
 
-	TSharedPtr<UE::TraceBasedDebuggers::FTraceSessionsManager> ChaosVDTraceManager;
+	TSharedPtr<FChaosVDTraceManager> ChaosVDTraceManager;
 
 	static FString ChaosVisualDebuggerProgramName;
 
diff --git a/Engine/Plugins/ChaosVD/Source/ChaosVD/Public/ChaosVDPlaybackController.h b/Engine/Plugins/ChaosVD/Source/ChaosVD/Public/ChaosVDPlaybackController.h
--- a/Engine/Plugins/ChaosVD/Source/ChaosVD/Public/ChaosVDPlaybackController.h
+++ b/Engine/Plugins/ChaosVD/Source/ChaosVD/Public/ChaosVDPlaybackController.h
@@ -1,4 +1,4 @@
-﻿// Copyright Epic Games, Inc. All Rights Reserved.
+// Copyright Epic Games, Inc. All Rights Reserved.
 
 #pragma once
 
@@ -13,17 +13,13 @@
 struct FChaosVDPlaybackEngineSnapshot;
 class IChaosVDPlaybackControllerInstigator;
 class UChaosVDCoreSettings;
+struct FChaosVDTraceSessionDescriptor;
 struct FChaosVDTrackInfo;
 class FChaosVDScene;
 struct FChaosVDRecording;
 class FChaosVDPlaybackController;
 class FString;
 
-namespace UE::TraceBasedDebuggers
-{
-struct FTraceSessionDescriptor;
-}
-
 enum class EChaosVDPlaybackButtonsID : uint8;
 
 DECLARE_MULTICAST_DELEGATE_OneParam(FChaosVDInitialFirstFramePlayed, TWeakPtr<FChaosVDPlaybackController>)
@@ -128,12 +124,7 @@ public:
 	void DeInitialize();
 
 	/** Loads a recording using a CVD Trace Session Descriptor */
-	bool LoadChaosVDRecordingFromTraceSession(const UE::TraceBasedDebuggers::FTraceSessionDescriptor& InSessionDescriptor);
-
-PRAGMA_DISABLE_DEPRECATION_WARNINGS
-	UE_DEPRECATED(5.8, "Use the version taking UE::TraceBasedDebuggers::FTraceSessionDescriptor in parameter")
-	bool LoadChaosVDRecordingFromTraceSession(const struct FChaosVDTraceSessionDescriptor& InSessionDescriptor);
-PRAGMA_ENABLE_DEPRECATION_WARNINGS
+	bool LoadChaosVDRecordingFromTraceSession(const FChaosVDTraceSessionDescriptor& InSessionDescriptor);
 
 	/** Unloads the currently loaded recording
 	 * @param UnloadOptions Options flags to change the steps performed during the unload
diff --git a/Engine/Plugins/ChaosVD/Source/ChaosVD/Public/Trace/ChaosVDTraceManager.h b/Engine/Plugins/ChaosVD/Source/ChaosVD/Public/Trace/ChaosVDTraceManager.h
--- a/Engine/Plugins/ChaosVD/Source/ChaosVD/Public/Trace/ChaosVDTraceManager.h
+++ b/Engine/Plugins/ChaosVD/Source/ChaosVD/Public/Trace/ChaosVDTraceManager.h
@@ -2,15 +2,13 @@
 
 #pragma once
 
-#include "Chaos/Serialization/SerializedDataBuffer.h"
 #include "ChaosVDModule.h"
+#include "Chaos/Serialization/SerializedDataBuffer.h"
 #include "Containers/SpscQueue.h"
 #include "Containers/StringFwd.h"
 #include "Framework/Notifications/NotificationManager.h"
 #include "Templates/SharedPointer.h"
 #include "Trace/StoreClient.h"
-#include "TraceDataStreamTypes.h"
-#include "TraceSessionsManager.h"
 
 struct FChaosVDRelayTraceDataMessage;
 enum class EChaosVDLoadRecordedDataMode : uint8;
@@ -19,69 +17,148 @@ class FChaosVDTraceModule;
 
 namespace TraceServices { class IAnalysisSession; }
 
-/**
- * Objects that allows us to use TLS to temporarily store and access a ptr to an existing instance.
- * This is temporary to work around the lack of an API method we need in the trace API, and will be removed in the future.
- * Either when we add that to the API, or find another way to pass an existing CVD recording to the trace provider before analysis starts
- */
-UE_EXPERIMENTAL(5.8, "This class is temporary and will be removed when a method in the trace API allows us to pass an existing recording to the trace provider before analysis starts.")
-class FChaosVDTraceManagerThreadContext : public TThreadSingleton<FChaosVDTraceManagerThreadContext>
-{
-public:
-
-	TWeakPtr<FChaosVDRecording> PendingExternalRecordingWeakPtr;
-};
-
 namespace Chaos::VD
 {
-PRAGMA_DISABLE_DEPRECATION_WARNINGS
-	struct UE_DEPRECATED(5.8, "This struct will no longer be exposed publicly") FAsyncProgressNotification : FTSTickerObjectBase
+	/** Object that allows showing a progress bar in the status bar of the editor,
+	 * where the progress can be updated from any thread */
+	struct FAsyncProgressNotification : FTSTickerObjectBase
 	{
-		FAsyncProgressNotification(const FText&) {}
-		void EnterProgress(int32 InCurrentProgress) {}
+		FAsyncProgressNotification(const FText& InMessage) : NotificationTitle(InMessage)
+		{
+		}
+
+		virtual ~FAsyncProgressNotification() override;
+
+		/** Updates the current progress
+		 * @param InCurrentProgress Value between 0 and 100, representing the progress we want to display
+		 */
+		void EnterProgress(int32 InCurrentProgress);
+
+	private:
+		virtual bool Tick(float DeltaTime) override;
+
+		float ElapsedTimeSinceLastUpdate = 0.0f;
+
+		FProgressNotificationHandle NotificationHandle;
+		std::atomic<int32> CurrentProgress = 0;
+		FText NotificationTitle;
 	};
 
-	class UE_DEPRECATED(5.8, "Use UE::TraceBasedDebuggers::IRemoteSessionDataStream instead") IRemoteSessionDataStream : public UE::TraceBasedDebuggers::IRemoteSessionDataStream
+	class IRemoteSessionDataStream
 	{
+	public:
+		virtual ~IRemoteSessionDataStream() = default;
+		virtual FGuid GetOwningRemoteSessionID() const = 0;
+		virtual uint64 UpdateAndGetBytesReadSinceLastMeasurement() = 0;
 	};
 
-	class UE_DEPRECATED(5.8, "Use UE::TraceBasedDebuggers::FSessionStreamStatsUpdater instead") FSessionStreamStatsUpdater : FTSTickerObjectBase
+	class FSessionStreamStatsUpdater : FTSTickerObjectBase
 	{
 	public:
-		explicit FSessionStreamStatsUpdater(IRemoteSessionDataStream* const DataStreamInstance) {}
+		virtual ~FSessionStreamStatsUpdater() override;
+
+		explicit FSessionStreamStatsUpdater(IRemoteSessionDataStream* const DataStreamInstance)
+			: DataStreamInstance(DataStreamInstance)
+		{
+		}
+
+	private:
+		virtual bool Tick(float DeltaTime) override;
+	
+		IRemoteSessionDataStream* DataStreamInstance;
+
+		double ElapsedTimeSinceLastUpdate = 0.0;
 	};
 
-	class UE_DEPRECATED(5.8, "Use UE::TraceBasedDebuggers::FDirectSocketStream instead") FDirectSocketStream : public UE::Trace::FDirectSocketStream, public IRemoteSessionDataStream
+	class FDirectSocketStream : public UE::Trace::FDirectSocketStream, public IRemoteSessionDataStream
 	{
 	public:
+
 		FDirectSocketStream() = delete;
-		explicit FDirectSocketStream(FGuid) {}
-		explicit FDirectSocketStream(FGuid, TUniquePtr<FArchiveFileWriterGeneric>&& InFileWriter)
-			: UE::Trace::FDirectSocketStream(MoveTemp(InFileWriter))
+
+		explicit FDirectSocketStream(FGuid InRemoteSessionID) : BytesReadSinceLastMeasurement(0), RemoteSessionID(InRemoteSessionID), StatsUpdater(this)
+		{
+		}
+
+		explicit FDirectSocketStream(FGuid InRemoteSessionID, TUniquePtr<FArchiveFileWriterGeneric>&& InFileWriter)
+			: UE::Trace::FDirectSocketStream(MoveTemp(InFileWriter)),
+			BytesReadSinceLastMeasurement(0),
+			RemoteSessionID(InRemoteSessionID),
+			StatsUpdater(this)
 		{
 		}
 
-		virtual int32 Read(void* Data, uint32 Size) override { return 0; }
-		virtual FGuid GetOwningRemoteSessionID() const override { return FGuid{}; }
-		virtual uint64 UpdateAndGetBytesReadSinceLastMeasurement() override { return 0; }
+		virtual int32 Read(void* Data, uint32 Size) override;
+
+		virtual FGuid GetOwningRemoteSessionID() const override;
+		virtual uint64 UpdateAndGetBytesReadSinceLastMeasurement() override;
+	
+	private:
+		std::atomic<uint64> BytesReadSinceLastMeasurement;
+		FGuid RemoteSessionID;
+
+		FSessionStreamStatsUpdater StatsUpdater;
 	};
 
-	class UE_DEPRECATED(5.8, "Use UE::TraceBasedDebuggers::FRelayDataStream instead") FRelayDataStream : public UE::Trace::IInDataStream,  public IRemoteSessionDataStream
+	/** A trace data stream object that listens for data coming from a relay stream
+	 * via CVD's Engine Editor Bridge object */
+	class FRelayDataStream : public UE::Trace::IInDataStream,  public IRemoteSessionDataStream
 	{
 	public:
-		explicit FRelayDataStream(FGuid InRemoteSessionID) {}
-		explicit FRelayDataStream(FGuid InRemoteSessionID, TUniquePtr<FArchiveFileWriterGeneric>&& InFileWriter) {}
+		explicit FRelayDataStream(FGuid InRemoteSessionID);
+		explicit FRelayDataStream(FGuid InRemoteSessionID, TUniquePtr<FArchiveFileWriterGeneric>&& InFileWriter);
+
+		virtual ~FRelayDataStream() override;
+
+		void EnqueueRelayedData(const TConstArrayView<uint8> InTraceDataBuffer);
+		void EnqueueRelayedData(TArray<uint8>&& InTraceDataBuffer);
+
+		virtual FGuid GetOwningRemoteSessionID() const override;
+		virtual uint64 UpdateAndGetBytesReadSinceLastMeasurement() override;
+
+	private:
+		virtual int32 Read(void* Dest, uint32 DestSize) override;
+		virtual void Close() override;
+		virtual bool WaitUntilReady() override;
+
+		TSpscQueue<FSerializedDataBuffer> DataQueue;
+		FEvent*	CancelEvent = nullptr;
+		FEvent*	WaitForDataEvent = nullptr;
+		TOptional<FSerializedDataBuffer> CurrentPacket;
+		uint32 CurrentPacketOffset = 0;
+		bool bClosed = false;
 
-		void EnqueueRelayedData(const TConstArrayView<uint8> InTraceDataBuffer) {}
-		void EnqueueRelayedData(TArray<uint8>&& InTraceDataBuffer) {}
+		FGuid RemoteSessionID = FGuid();
 
-		virtual FGuid GetOwningRemoteSessionID() const override { return FGuid{}; }
-		virtual uint64 UpdateAndGetBytesReadSinceLastMeasurement() override { return 0; }
+		TUniquePtr<FArchiveFileWriterGeneric> FileWriter;
+
+		std::atomic<uint64> BytesReadSinceLastMeasurement;
+		FSessionStreamStatsUpdater StatsUpdater;
 	};
-} // Chaos::VD
+}
 
-struct UE_DEPRECATED(5.8, "Use UE::TraceBasedDebuggers::FTraceSessionDescriptor instead") FChaosVDTraceSessionDescriptor : UE::TraceBasedDebuggers::FTraceSessionDescriptor
+/**
+ * Structure containing info about a trace session used by Chaos Visual Debugger
+ */
+struct FChaosVDTraceSessionDescriptor
 {
+
+	/** The trace session name. Used to get the trace analysis instance for this CVD session*/
+	FString SessionName;
+	/** Port number used to establish a connection to this session. The value is 0 if no specific port was used*/
+	uint16 SessionPort = 0;
+	/** Session ID of this CVD session in the Remote Session management system*/
+	FGuid RemoteSessionID;
+
+	/** True if this is a session recorded live (not loaded from disk)*/
+	bool bIsLiveSession = false;
+
+	/** Returns true if valid. For a session to be valued needs at minimum have a Trace Session name */
+	bool IsValid() const
+	{
+		return !SessionName.IsEmpty();
+	}
+
 	friend bool operator==(const FChaosVDTraceSessionDescriptor& Lhs, const FChaosVDTraceSessionDescriptor& RHS)
 	{
 		return Lhs.SessionName == RHS.SessionName
@@ -94,20 +171,40 @@ struct UE_DEPRECATED(5.8, "Use UE::TraceBasedDebuggers::FTraceSessionDescriptor
 	}
 };
 
-class UE_DEPRECATED(5.8, "Use UE::TraceBasedDebuggers::FTraceSessionsManager instead") FChaosVDTraceManager : public UE::TraceBasedDebuggers::FTraceSessionsManager
+/** Objects that allows us to use TLS to temporarily store and access a ptr to an existing instance.
+ * This is temporary to workaround the lack of an API method we need in the trace API, and will be removed in the future.
+ * either when we add that to the API, or find another way to pass an existing CVD recording to the trace provider before anaislis starts
+ */
+class FChaosVDTraceManagerThreadContext : public TThreadSingleton<FChaosVDTraceManagerThreadContext>
+{
+public:
+
+	TWeakPtr<FChaosVDRecording> PendingExternalRecordingWeakPtr;
+};
+
+/** Manager class used by Chaos VD to interact/control UE Trace systems */
+class FChaosVDTraceManager
 {
 public:
 	FChaosVDTraceManager();
 	~FChaosVDTraceManager();
 
+	/** Load a trace file and starts analyzing it
+	 * @param InTraceFilename File Name including (Path Included) of the Trace file to load
+	 */
 	CHAOSVD_API FString LoadTraceFile(const FString& InTraceFilename, const TSharedPtr<FChaosVDRecording>& ExistingRecordingPtr = nullptr);
 
 	FString LoadTraceFile(TUniquePtr<IFileHandle>&& InFileHandle, const FString& InTraceSessionName, const TSharedPtr<FChaosVDRecording>& ExistingRecordingPtr);
+	
+	/**
+	 * Connects to a live Trace Session and starts analyzing it.
+	 * @param InSessionHost Trace Store Address for this session
+	 * @param SessionID Trace ID in the Trace Store provided as host
+	 * @param ExistingRecordingPtr Ptr to an existing recording instance to append the new CVD data to, if any
+	 * @return Session Name
+	 */
+	FString ConnectToLiveSession(FStringView InSessionHost, uint32 SessionID, const TSharedPtr<FChaosVDRecording>& ExistingRecordingPtr = nullptr);
 
-	FString ConnectToLiveSession(FStringView InSessionHost, uint32 SessionID, const TSharedPtr<FChaosVDRecording>& ExistingRecordingPtr = nullptr)
-	{
-		return FString();
-	}
 
 	UE_DEPRECATED(5.8, "Please use the version that takes a Remote Session ID")
 	FString ConnectToLiveSession_Direct( uint16& OutSessionPort, const TSharedPtr<FChaosVDRecording>& ExistingRecordingPtr = nullptr)
@@ -115,23 +212,123 @@ public:
 		return FString();
 	};
 
-	FString ConnectToLiveSession_Direct(FGuid RemoteSessionID, uint16& OutSessionPort, const TSharedPtr<FChaosVDRecording>& ExistingRecordingPtr = nullptr)
+	/**
+	 * Opens up a direct trace stream and waits for data
+	 * @param RemoteSessionID Session ID to which we want to connect to
+	 * @param OutSessionPort Number of the port we open for this session
+	 * @param ExistingRecordingPtr Ptr to an existing recording instance to append the new CVD data to, if any
+	 * @return Session Name
+	 */
+	FString ConnectToLiveSession_Direct(FGuid RemoteSessionID, uint16& OutSessionPort, const TSharedPtr<FChaosVDRecording>& ExistingRecordingPtr = nullptr);
+
+	/**
+	 * Opens up a relay trace stream and waits for data
+	 * @param RemoteSessionID Session ID to which we want to connect to
+	 * @param ExistingRecordingPtr Ptr to an existing recording instance to append the new CVD data to, if any
+	 * @return Session Name
+	 */
+	FString ConnectToLiveSession_Relay(FGuid RemoteSessionID, const TSharedPtr<FChaosVDRecording>& ExistingRecordingPtr = nullptr);
+
+	/** Returns the path to the local trace store */
+	FString GetLocalTraceStoreDirPath();
+
+	/** Returns a ptr to the session registered with the provided session name. Null if no session is found */
+	CHAOSVD_API TSharedPtr<const TraceServices::IAnalysisSession> GetSession(const FString& InSessionName);
+
+	/** Stops and de-registers a trace session registered with the provided session name */
+	CHAOSVD_API void CloseSession(const FString& InSessionName);
+
+	/** Stops a trace session registered with the provided session name */
+	void StopSession(const FString& InSessionName);
+
+	template<typename TVisitor>
+	static void EnumerateActiveSessions(FStringView InSessionHost, TVisitor Callback);
+
+	static const UE::Trace::FStoreClient::FSessionInfo* GetTraceSessionInfo(FStringView InSessionHost, FGuid TraceGuid);
+
+	template<typename TVisitor>
+	void EnumerateActiveSessions(TVisitor Callback);
+
+	/** Access the trace store at the provided host address, and returns the file name for the trace file
+	 * @param InSessionHost Address of the Trace Store
+	 * @param SessionID Session id of the trace data from which get the file name from
+	 * @return Full Filename with path to the trace file if exist. Returns an empty string if a trace file cannot be found
+	 */
+	FString GetTraceFileNameFromStoreForSession(FStringView InSessionHost, uint32 SessionID);
+
+	/** Access the trace store at the provided host address, and returns a trace data stream ready to use
+	 * @param InSessionHost Address of the Trace Store
+	 * @param SessionID Session id of the trace data
+	 */
+	UE::Trace::FStoreClient::FTraceData GetTraceDataStreamFromStore(FStringView InSessionHost, uint32 SessionID);
+
+	static FString GetDefaultSaveDirectoryPath();
+
+private:
+
+	TUniquePtr<FArchiveFileWriterGeneric> OpenTraceFileForWrite(const FString& InDirectoryToSavePath);
+
+	bool ConnectToLiveSession_Internal(uint32 SessionID, const TSharedPtr<FChaosVDRecording>& ExistingRecordingPtr, const FString& InRequestedSessionName, UE::Trace::FStoreClient::FTraceData&& InTraceDataStream);
+
+	TUniquePtr<UE::Trace::IInDataStream> CreateFileDataStream(TUniquePtr<IFileHandle>&& InFileHandle);
+	TUniquePtr<Chaos::VD::FRelayDataStream> CreateRelayDataStream(FGuid RemoteSessionID);
+
+	/** Temporary workaround method to set an existing recording structure in CVD's trace provider before the trace analysis starts
+	 * in the Trace analysis thread.*/
+	void SetPendingExternalRecordingToProcess(const TSharedPtr<FChaosVDRecording>& ExistingRecordingPtr);
+
+	/** The trace analysis session. */
+	TMap<FString, TSharedPtr<const TraceServices::IAnalysisSession>> AnalysisSessionByName;
+
+	TSharedPtr<FChaosVDTraceModule> ChaosVDTraceModule;
+
+	static uint16 CurrentDirectModeStartingPort;
+	static constexpr uint16 DefaultStartingPort = 1993;
+};
+
+template <typename TCallback>
+void FChaosVDTraceManager::EnumerateActiveSessions(FStringView InSessionHost, TCallback Callback)
+{
+	if (InSessionHost.IsEmpty())
 	{
-		return FString();
+		UE_LOG(LogChaosVDEditor, Error, TEXT("[%s] Failed to connect to trace store. Provided session host is empty"), ANSI_TO_TCHAR(__FUNCTION__));
+		return;
 	}
 
-	FString ConnectToLiveSession_Relay(FGuid RemoteSessionID, const TSharedPtr<FChaosVDRecording>& ExistingRecordingPtr = nullptr)
+	using namespace UE::Trace;
+	const FStoreClient* StoreClient = FStoreClient::Connect(InSessionHost.GetData());
+
+	if (!StoreClient)
 	{
-		return FString();
+		UE_LOG(LogChaosVDEditor, Error, TEXT("[%s] Failed to connect to trace store at [%s]"), ANSI_TO_TCHAR(__FUNCTION__), InSessionHost.GetData())
+		return;
 	}
 
-	template<typename TVisitor>
-	static void EnumerateActiveSessions(FStringView InSessionHost, TVisitor Callback)
+	const uint32 SessionCount = StoreClient->GetSessionCount();
+
+	for (uint32 SessionIndex = 0; SessionIndex < SessionCount; SessionIndex++)
 	{
-		// deprecated
+		if (const FStoreClient::FSessionInfo* SessionInfo = StoreClient->GetSessionInfo(SessionIndex))
+		{
+			if (!Callback(*SessionInfo))
+			{
+				return;
+			}
+		}	
 	}
+}
 
-	static FString GetDefaultSaveDirectoryPath();
-};
-PRAGMA_ENABLE_DEPRECATION_WARNINGS
-
+template <typename TVisitor>
+void FChaosVDTraceManager::EnumerateActiveSessions(TVisitor Callback)
+{
+	for (const TPair<FString, TSharedPtr<const TraceServices::IAnalysisSession>>& SessionWithName : AnalysisSessionByName)
+	{
+		if (SessionWithName.Value)
+		{
+			if (!Callback(SessionWithName.Value.ToSharedRef()))
+			{
+				return;
+			}
+		}
+	}
+}
diff --git a/Engine/Source/Runtime/Experimental/ChaosVDData/Public/ChaosVDRecording.h b/Engine/Source/Runtime/Experimental/ChaosVDData/Public/ChaosVDRecording.h
--- a/Engine/Source/Runtime/Experimental/ChaosVDData/Public/ChaosVDRecording.h
+++ b/Engine/Source/Runtime/Experimental/ChaosVDData/Public/ChaosVDRecording.h
@@ -1,4 +1,4 @@
-﻿// Copyright Epic Games, Inc. All Rights Reserved.
+// Copyright Epic Games, Inc. All Rights Reserved.
 
 #pragma once
 
diff --git a/Engine/Source/Runtime/Experimental/TraceBasedDebuggers/Private/RemoteSessionsManager.cpp b/Engine/Source/Runtime/Experimental/TraceBasedDebuggers/Private/RemoteSessionsManager.cpp
--- a/Engine/Source/Runtime/Experimental/TraceBasedDebuggers/Private/RemoteSessionsManager.cpp
+++ b/Engine/Source/Runtime/Experimental/TraceBasedDebuggers/Private/RemoteSessionsManager.cpp
@@ -16,6 +16,8 @@
 #define TRACE_BASED_DEBUGGERS_WITHOUT_TRACE 0
 #endif
 
+DEFINE_LOG_CATEGORY(LogTraceBasedDebuggers)
+
 namespace UE::TraceBasedDebuggers
 {
 const FGuid FRemoteSessionsManager::AllRemoteSessionsWrapperGUID = FGuid::NewGuid();
@@ -498,10 +500,6 @@ void FRemoteSessionsManager::HandleRecordingStatusUpdateMessage(const FRecording
 			return;
 		}
 
-		// Update recording state before broadcasting so listeners can use the information to only
-		// process request from their debugger Id.
-		SessionInfoPtr->LastKnownRecordingState = Message;
-
 		if (bLastKnownDebuggerIsRecording != Message.bIsRecording)
 		{
 			if (Message.bIsRecording)
@@ -513,6 +511,8 @@ void FRemoteSessionsManager::HandleRecordingStatusUpdateMessage(const FRecording
 				RecordingStoppedDelegate.Broadcast(SessionInfoPtr);
 			}
 		}
+
+		SessionInfoPtr->LastKnownRecordingState = Message;
 	}
 	else
 	{
diff --git a/Engine/Source/Runtime/Experimental/TraceBasedDebuggers/Private/TraceBasedDebuggerTypes.cpp b/Engine/Source/Runtime/Experimental/TraceBasedDebuggers/Private/TraceBasedDebuggerTypes.cpp
deleted file mode 100644
diff --git a/Engine/Source/Runtime/Experimental/TraceBasedDebuggers/Public/RemoteSessionsManager.h b/Engine/Source/Runtime/Experimental/TraceBasedDebuggers/Public/RemoteSessionsManager.h
--- a/Engine/Source/Runtime/Experimental/TraceBasedDebuggers/Public/RemoteSessionsManager.h
+++ b/Engine/Source/Runtime/Experimental/TraceBasedDebuggers/Public/RemoteSessionsManager.h
@@ -5,7 +5,6 @@
 #include "Containers/Ticker.h"
 #include "IMessageContext.h"
 #include "MessageEndpoint.h"
-#include "TraceBasedDebuggerTypes.h"
 #include "Misc/DateTime.h"
 #include "Misc/Guid.h"
 #include "Templates/SharedPointer.h"
@@ -18,6 +17,8 @@
 class IMessageBus;
 class IMessageContext;
 
+UE_API DECLARE_LOG_CATEGORY_EXTERN(LogTraceBasedDebuggers, Log, All);
+
 namespace UE::TraceBasedDebuggers
 {
 struct FFullSessionInfoRequestMessage;
diff --git a/Engine/Source/Runtime/Experimental/TraceBasedDebuggers/Public/TraceBasedDebuggerTypes.h b/Engine/Source/Runtime/Experimental/TraceBasedDebuggers/Public/TraceBasedDebuggerTypes.h
deleted file mode 100644
diff --git a/Engine/Source/Runtime/Experimental/TraceBasedDebuggersAnalysis/Private/TraceBasedDebuggersAnalysis.cpp b/Engine/Source/Runtime/Experimental/TraceBasedDebuggersAnalysis/Private/TraceBasedDebuggersAnalysis.cpp
deleted file mode 100644
diff --git a/Engine/Source/Runtime/Experimental/TraceBasedDebuggersAnalysis/Private/TraceDataStreamTypes.cpp b/Engine/Source/Runtime/Experimental/TraceBasedDebuggersAnalysis/Private/TraceDataStreamTypes.cpp
deleted file mode 100644
diff --git a/Engine/Source/Runtime/Experimental/TraceBasedDebuggersAnalysis/Private/TraceSessionsManager.cpp b/Engine/Source/Runtime/Experimental/TraceBasedDebuggersAnalysis/Private/TraceSessionsManager.cpp
deleted file mode 100644
diff --git a/Engine/Source/Runtime/Experimental/TraceBasedDebuggersAnalysis/Public/TraceDataStreamTypes.h b/Engine/Source/Runtime/Experimental/TraceBasedDebuggersAnalysis/Public/TraceDataStreamTypes.h
deleted file mode 100644
diff --git a/Engine/Source/Runtime/Experimental/TraceBasedDebuggersAnalysis/Public/TraceSessionsManager.h b/Engine/Source/Runtime/Experimental/TraceBasedDebuggersAnalysis/Public/TraceSessionsManager.h
deleted file mode 100644
diff --git a/Engine/Source/Runtime/Experimental/TraceBasedDebuggersAnalysis/TraceBasedDebuggersAnalysis.Build.cs b/Engine/Source/Runtime/Experimental/TraceBasedDebuggersAnalysis/TraceBasedDebuggersAnalysis.Build.cs
deleted file mode 100644
