[TraceBasedDebuggers] refactored and moved RecordingControls and ConnectionStatusBar from ChaosVD
#Synced-CL 49268467

diff --git a/Engine/Plugins/ChaosVD/Source/ChaosVD/Private/Widgets/ChaosVDRecordingControls.cpp b/Engine/Plugins/ChaosVD/Source/ChaosVD/Private/Widgets/ChaosVDRecordingControls.cpp
--- a/Engine/Plugins/ChaosVD/Source/ChaosVD/Private/Widgets/ChaosVDRecordingControls.cpp
+++ b/Engine/Plugins/ChaosVD/Source/ChaosVD/Private/Widgets/ChaosVDRecordingControls.cpp
@@ -1,213 +1,71 @@
-﻿// Copyright Epic Games, Inc. All Rights Reserved.
+// Copyright Epic Games, Inc. All Rights Reserved.
 
 #include "Widgets/ChaosVDRecordingControls.h"
 
-#include "AsyncCompilationHelpers.h"
 #include "Chaos/ChaosVDEngineEditorBridge.h"
 #include "Chaos/ChaosVDRemoteSessionsManager.h"
 #include "ChaosVDEngine.h"
 #include "ChaosVDModule.h"
-#include "ChaosVDScene.h"
 #include "ChaosVDSettingsManager.h"
 #include "ChaosVDStyle.h"
-#include "ChaosVisualDebugger/ChaosVisualDebuggerTrace.h"
-#include "Editor.h"
-#include "Framework/MultiBox/MultiBoxBuilder.h"
-#include "Input/Reply.h"
-#include "Misc/MessageDialog.h"
-#include "Misc/ScopedSlowTask.h"
-#include "RemoteSessionsManager.h"
-#include "SChaosVDNameListPicker.h"
 #include "SEnumCombo.h"
 #include "Settings/ChaosVDGeneralSettings.h"
-#include "SessionInfo.h"
-#include "SocketSubsystem.h"
-#include "StatusBarSubsystem.h"
-#include "ToolMenu.h"
-#include "ToolMenus.h"
-#include "TraceDataRelayTransport.h"
-#include "TraceSessionsManager.h"
-#include "Trace/StoreClient.h"
-#include "Widgets/Images/SImage.h"
-#include "Widgets/Input/SButton.h"
 #include "Widgets/Input/SComboButton.h"
-#include "Widgets/Notifications/SNotificationList.h"
 #include "Widgets/SChaosVDMainTab.h"
-#include "Widgets/Text/STextBlock.h"
-
-#include UE_INLINE_GENERATED_CPP_BY_NAME(ChaosVDRecordingControls)
 
 #define LOCTEXT_NAMESPACE "ChaosVisualDebugger"
 
 class UChaosVDGeneralSettings;
 
-FChaosVDRecordingControls::FChaosVDRecordingControls(const TSharedRef<SChaosVDMainTab>& InMainTabSharedRef)
-{
-	MainTabWeakPtr = InMainTabSharedRef;
-	StatusBarID = InMainTabSharedRef->GetStatusBarName();
-	
-	RecordingAnimation = FCurveSequence();
-	RecordingAnimation.AddCurve(0.f, 1.5f, ECurveEaseFunction::Linear);
-
-	CurrentSelectedSessionId = UE::TraceBasedDebuggers::FRemoteSessionsManager::LocalEditorSessionID;
-}
-
-FChaosVDRecordingControls::~FChaosVDRecordingControls()
-{
-	Deinitialize();
-}
-
-void FChaosVDRecordingControls::Initialize()
+FChaosVDRecordingControls::FChaosVDRecordingControls(const TSharedRef<SChaosVDMainTab>& InMainTab
+	, const TSharedRef<UE::TraceBasedDebuggers::FRemoteSessionsManager>& InRemoteSessionManager)
+	: FRecordingControls(InMainTab->GetStatusBarName(), InRemoteSessionManager, Chaos::VD::DebuggerGuid, LogChaosVDEditor)
 {
-	if (bInitialized)
-	{
-		return;
-	}
+	TraceRelayTransportInstanceWeakPtr = FChaosVDEngineEditorBridge::Get().GetTraceRelayTransportInstance();
 
-	if (const TSharedPtr<UE::TraceBasedDebuggers::FRemoteSessionsManager> RemoteSessionManager = FChaosVDEngineEditorBridge::Get().GetSessionsManager())
-	{
-		RemoteSessionManager->OnSessionRecordingStarted().AddSP(this, &FChaosVDRecordingControls::HandleRecordingStart);
-		RemoteSessionManager->OnSessionRecordingStopped().AddSP(this, &FChaosVDRecordingControls::HandleRecordingStop);
-		RemoteSessionManager->OnSessionExpired().AddSP(this, &FChaosVDRecordingControls::CancelInFlightConnectionAttempt);
-		RemoteSessionManager->OnSessionsUpdated().AddSP(this, &FChaosVDRecordingControls::HandleRemoteSessionsUpdated);
-	}
+	MainTabWeakPtr = InMainTab;
+	CurrentLoadingMode = EChaosVDLoadRecordedDataMode::SingleSource;
 
-	bInitialized = true;
+	StartRecordingBrush = FChaosVDStyle::Get().GetBrush("RecordIcon");
+	StopRecordingBrush = FChaosVDStyle::Get().GetBrush("StopIcon");
+	RecordingSessionIcon = FSlateIcon(FChaosVDStyle::GetStyleSetName(), FName("RecordIcon"), FName("RecordIcon"));
 }
 
-void FChaosVDRecordingControls::Deinitialize()
+void FChaosVDRecordingControls::AddToMenuInternal(FToolMenuSection& InSection)
 {
-	if (!bInitialized)
-	{
-		return;
-	}
-
-	if (const TSharedPtr<UE::TraceBasedDebuggers::FRemoteSessionsManager> RemoteSessionManager = FChaosVDEngineEditorBridge::Get().GetSessionsManager())
-	{
-		RemoteSessionManager->OnSessionRecordingStarted().RemoveAll(this);
-		RemoteSessionManager->OnSessionRecordingStopped().RemoveAll(this);
-		RemoteSessionManager->OnSessionExpired().RemoveAll(this);
-		RemoteSessionManager->OnSessionsUpdated().RemoveAll(this);
-	}
-
-	bInitialized = false;
-}
-
-TSharedRef<SWidget> FChaosVDRecordingControls::GenerateToggleRecordingStateButton(const FText& StartRecordingTooltip)
-{
-	return SNew(SButton)
-			.OnClicked(FOnClicked::CreateRaw(this, &FChaosVDRecordingControls::ToggleRecordingState))
-			.ForegroundColor(FSlateColor::UseForeground())
-			.IsFocusable(false)
-			.IsEnabled_Raw(this, &FChaosVDRecordingControls::IsRecordingToggleButtonEnabled)
-			.VAlign(VAlign_Center)
-			.HAlign(HAlign_Center)
-			.OnHovered_Lambda([this](){ bRecordingButtonHovered = true;})
-			.OnUnhovered_Lambda([this](){ bRecordingButtonHovered = false;})
-			.ToolTipText_Lambda([this, StartRecordingTooltip]()
-			{
-				return IsRecording() ? LOCTEXT("StopRecordButtonDesc", "Stop the current recording ") : StartRecordingTooltip;
-			})
+	FToolMenuEntry& AdvancedRecordingOptionsEntry = InSection.AddEntry(
+		FToolMenuEntry::InitWidget(
+			"AdvancedRecordingOptions",
+			SNew(SBox).Padding(1.0f, 0.0f)
 			[
-				SNew(SHorizontalBox)
-				+SHorizontalBox::Slot()
-				.Padding(FMargin(0, 0, 0, 0))
-				.VAlign(VAlign_Center)
-				.AutoWidth()
-				[
-					SNew(SImage)
-					.Image_Raw(this, &FChaosVDRecordingControls::GetRecordOrStopButton)
-					.ColorAndOpacity_Lambda([this]()
-					{
-						if (IsRecording())
-						{
-							TSharedPtr<SChaosVDMainTab> MainTabPtr = MainTabWeakPtr.Pin();
-							if (MainTabPtr && !RecordingAnimation.IsPlaying())
-							{
-								RecordingAnimation.Play(MainTabPtr->AsShared(), true);
-							}
-
-							const FLinearColor Color = bRecordingButtonHovered ? FLinearColor::Red : FLinearColor::White;
-							return FSlateColor(bRecordingButtonHovered ? Color : Color.CopyWithNewOpacity(0.2f + 0.8f * RecordingAnimation.GetLerp()));
-						}
-
-						RecordingAnimation.Pause();
-						return FSlateColor::UseSubduedForeground();
-					})
-				]
-				+SHorizontalBox::Slot()
-				.Padding(FMargin(4, 0, 0, 0))
-				.VAlign(VAlign_Center)
-				.AutoWidth()
-				[
-					SNew(STextBlock)
-					.Visibility_Lambda([this](){return IsRecording() ? EVisibility::Collapsed : EVisibility::Visible;})
-					.TextStyle(FAppStyle::Get(), "SmallButtonText")
-					.Text(LOCTEXT("StartRecordingButtonLabel", "Start Recording"))
-				]
-			];
-}
-
-FText FChaosVDRecordingControls::GetCurrentSelectedSessionName() const
-{
-	if (const TSharedPtr<UE::TraceBasedDebuggers::FSessionInfo> CurrentSessionPtr = GetCurrentSessionInfo())
-	{
-		return FText::AsCultureInvariant(CurrentSessionPtr->SessionName);
-	}
-
-	static FText InvalidSessionName = LOCTEXT("InvalidSessionLabel", "Select a Session");
-
-	return InvalidSessionName;
-}
+				GenerateAdvancedRecordOptionsButton()
+			],
+			FText::GetEmpty(),
+			true,
+			false
+		));
+	AdvancedRecordingOptionsEntry.InsertPosition.Position = EToolMenuInsertType::Before;
+	AdvancedRecordingOptionsEntry.InsertPosition.Name = "RecordingTime";
 
-TSharedRef<SWidget> FChaosVDRecordingControls::GenerateTargetSessionSelector()
-{
-	return SNew(SHorizontalBox)
-		+SHorizontalBox::Slot()
-		.AutoWidth()
-		[
-			SNew(SComboButton)
-			.MenuPlacement(MenuPlacement_AboveAnchor).ComboButtonStyle(&FAppStyle::Get().GetWidgetStyle<FComboButtonStyle>("ComboButton"))
-			.OnGetMenuContent(this, &FChaosVDRecordingControls::GenerateTargetSessionDropdown)
-			.HasDownArrow(true)
-			.ButtonContent()
+	FToolMenuEntry& DataChannelsEntry = InSection.AddEntry(
+		FToolMenuEntry::InitWidget(
+			"DataChannelsButton",
+			SNew(SBox).Padding(4.0f, 0.0f)
 			[
-				SNew(STextBlock)
-				.Text_Raw(this, &FChaosVDRecordingControls::GetCurrentSelectedSessionName)
-				.TextStyle(&FAppStyle::Get().GetWidgetStyle<FTextBlockStyle>("NormalText"))
-		]];
+				GenerateDataChannelsButton()
+			],
+			FText::GetEmpty(),
+			false,
+			false
+		));
+
+	DataChannelsEntry.InsertPosition.Position = EToolMenuInsertType::Before;
+	DataChannelsEntry.InsertPosition.Name = "RecordingTime";
 }
 
-TSharedRef<SWidget> FChaosVDRecordingControls::GenerateTargetSessionDropdown()
+void FChaosVDRecordingControls::OnGenerateTargetSessionSelectorInternal(FMenuBuilder& MenuBuilder)
 {
 	using namespace UE::TraceBasedDebuggers;
-
-	FMenuBuilder MenuBuilder(true, nullptr);
-
-	MenuBuilder.BeginSection("CVDRecordingWidgetTargets", LOCTEXT("CVDRecordingTargetsMenu", "Available Targets"));
-	{
-		if (const TSharedPtr<FRemoteSessionsManager> RemoteSessionManager = FChaosVDEngineEditorBridge::Get().GetSessionsManager())
-		{
-			RemoteSessionManager->EnumerateActiveSessions([this, &MenuBuilder](const TSharedRef<FSessionInfo>& InSessionInfoRef)
-			{
-				if (EnumHasAnyFlags(InSessionInfoRef->GetSessionTypeAttributes(), ERemoteSessionAttributes::IsMultiSessionWrapper))
-				{
-					return true;
-				}
-
-				const FText SessionNameAsText = FText::AsCultureInvariant(InSessionInfoRef->SessionName);
-				MenuBuilder.AddMenuEntry(
-				SessionNameAsText,
-				FText::Format(LOCTEXT("SingleTargetItemTooltip", "Select {0} session as current target"), SessionNameAsText),
-				GetIconForSession(InSessionInfoRef->InstanceId),
-				FUIAction(FExecuteAction::CreateSP(this, &FChaosVDRecordingControls::SelectTargetSession, InSessionInfoRef->InstanceId), EUIActionRepeatMode::RepeatDisabled));
-				return true;
-			});
-		}
-	}
-
-	MenuBuilder.EndSection();
-
 	MenuBuilder.BeginSection("CVDRecordingWidgetTargetsMulti", LOCTEXT("CVDRecordingMultiTargetsMenu", "Multi Target"));
 
 	const FText AllRemoteTargetsLabel = LOCTEXT("AllRemoteOption", "All Remote");
@@ -253,24 +111,118 @@ TSharedRef<SWidget> FChaosVDRecordingControls::GenerateTargetSessionDropdown()
 		);
 
 	MenuBuilder.EndSection();
+}
 
-	return MenuBuilder.MakeWidget();
+int32 FChaosVDRecordingControls::GetMaxConnectionRetriesInternal()
+{
+	if (const UChaosVDGeneralSettings* Settings = FChaosVDSettingsManager::Get().GetSettingsObject<UChaosVDGeneralSettings>())
+	{
+		return Settings->MaxConnectionRetries;
+	}
+
+	const int32 DefaultValue = FRecordingControls::GetMaxConnectionRetriesInternal();
+	UE_LOG(LogChaosVDEditor, Warning, TEXT("Failed to obtain setting object. Setting the retries attempts to connect to a session to %d as a fallback."), DefaultValue);
+	return DefaultValue;
 }
 
-TSharedRef<SWidget> FChaosVDRecordingControls::GenerateRecordingTimeTextBlock()
+UE::TraceBasedDebuggers::ETraceTransportMode FChaosVDRecordingControls::GetTransportModeOverrideForTargetTypeInternal(const EBuildTargetType InBuildTarget)
 {
-	return SNew(SBox)
-			.VAlign(VAlign_Center)
-			.Visibility_Raw(this, &FChaosVDRecordingControls::GetRecordingTimeTextBlockVisibility)
-			.Padding(12.0f, 0.0f, 0.0f, 0.0f)
-			[
-				SNew(STextBlock)
-				.TextStyle(FAppStyle::Get(), "SmallButtonText")
-				.Text_Raw(this, &FChaosVDRecordingControls::GetRecordingTimeText)
-				.ColorAndOpacity(FColor::White)
-			];
+	if (const UChaosVDGeneralSettings* Settings = FChaosVDSettingsManager::Get().GetSettingsObject<UChaosVDGeneralSettings>())
+	{
+		return Settings->GetTransportModeForTargetType(InBuildTarget);
+	}
+
+	UE::TraceBasedDebuggers::ETraceTransportMode DefaultValue = FRecordingControls::GetTransportModeOverrideForTargetTypeInternal(InBuildTarget);
+	UE_LOG(LogChaosVDEditor, Warning, TEXT("Failed to obtain setting object. Setting the retries attempts to connect to a session to %s as a fallback."), *UEnum::GetValueAsString(DefaultValue));
+	return DefaultValue;
+}
+
+bool FChaosVDRecordingControls::ConnectToLiveSession_RelayInternal(const FGuid& InSessionId)
+{
+	if (const TSharedPtr<SChaosVDMainTab> MainTabSharedPtr = MainTabWeakPtr.Pin())
+	{
+		return MainTabSharedPtr->ConnectToLiveSession_Relay(InSessionId, CurrentLoadingMode);
+	}
+
+	return false;
+}
+
+bool FChaosVDRecordingControls::ConnectToLiveSession_DirectInternal(const FGuid& InSessionId)
+{
+	if (const TSharedPtr<SChaosVDMainTab> MainTabSharedPtr = MainTabWeakPtr.Pin())
+	{
+		return MainTabSharedPtr->ConnectToLiveSession_Direct(InSessionId, CurrentLoadingMode);
+	}
+
+	return false;
+}
+
+TNotNull<TUniquePtr<UE::TraceBasedDebuggers::FStartRecordingCommandMessage>> FChaosVDRecordingControls::BuildStartRecordingParamsInternal()
+{
+	return MakeUnique<FChaosVDStartRecordingCommandMessage>();
+}
+
+void FChaosVDRecordingControls::SendStartRecordingCommandInternal(const FMessageAddress& InAddress
+	, TUniquePtr<UE::TraceBasedDebuggers::FStartRecordingCommandMessage> InStartRecordingCommand)
+{
+	if (const TSharedPtr<UE::TraceBasedDebuggers::FRemoteSessionsManager> SessionsManager = RemoteSessionManagerWeakPtr.Pin())
+	{
+		const FChaosVDStartRecordingCommandMessage* StartRecordingCommand = static_cast<FChaosVDStartRecordingCommandMessage*>(InStartRecordingCommand.Get());
+		SessionsManager->SendCommand(InAddress, *StartRecordingCommand);
+	}
+}
+
+void FChaosVDRecordingControls::SendStopRecordingCommandInternal(const FMessageAddress& InAddress)
+{
+	if (const TSharedPtr<UE::TraceBasedDebuggers::FRemoteSessionsManager> SessionsManager = RemoteSessionManagerWeakPtr.Pin())
+	{
+		SessionsManager->SendCommand(InAddress, FChaosVDStopRecordingCommandMessage{});
+	}
+}
+
+bool FChaosVDRecordingControls::IsRecordingToggleButtonEnabledInternal(const TSharedRef<UE::TraceBasedDebuggers::FSessionInfo>& InSessionInfo) const
+{
+	// If we are recording, don't show the stop button for the mode that is disabled
+	if (EnumHasAnyFlags(InSessionInfo->GetSessionTypeAttributes(), UE::TraceBasedDebuggers::ERemoteSessionAttributes::IsMultiSessionWrapper))
+	{
+		TSharedPtr<UE::TraceBasedDebuggers::FMultiSessionInfo> AsMultiSessionInfo = StaticCastSharedPtr<UE::TraceBasedDebuggers::FMultiSessionInfo>(InSessionInfo.ToSharedPtr());
+
+		return CanSelectMultiSessionTarget(InSessionInfo);
+	}
+
+	return true;
+}
+
+void FChaosVDRecordingControls::OnToggleRecordingStateInternal(const TSharedRef<UE::TraceBasedDebuggers::FSessionInfo>& InSessionInfo)
+{
+	if (EnumHasAnyFlags(InSessionInfo->GetSessionTypeAttributes(), UE::TraceBasedDebuggers::ERemoteSessionAttributes::IsMultiSessionWrapper))
+	{
+		ToggleMultiSessionSessionRecordingState(StaticCastSharedPtr<UE::TraceBasedDebuggers::FMultiSessionInfo>(InSessionInfo.ToSharedPtr()).ToSharedRef());
+	}
+	else
+	{
+		ToggleSingleSessionRecordingState(InSessionInfo);
+	}
 }
 
+void FChaosVDRecordingControls::CloseSessionByRemoteSessionIDInternal(const FGuid& InSessionId)
+{
+	if (const TSharedPtr<SChaosVDMainTab> MainTab = MainTabWeakPtr.Pin())
+	{
+		MainTab->GetChaosVDEngineInstance()->CloseSessionByRemoteSessionID(InSessionId);
+	}
+}
+
+UE::TraceBasedDebuggers::FTraceSessionDescriptor& FChaosVDRecordingControls::GetCurrentSessionDescriptorInternal()
+{
+	if (const TSharedPtr<SChaosVDMainTab> MainTab = MainTabWeakPtr.Pin())
+	{
+		return MainTab->GetChaosVDEngineInstance()->GetTraceSessionDescriptors().Last();
+	}
+
+	static UE::TraceBasedDebuggers::FTraceSessionDescriptor InvalidDescriptor;
+	return InvalidDescriptor;
+}
 
 TSharedRef<SWidget> FChaosVDRecordingControls::GenerateAdvancedRecordOptionsButton()
 {
@@ -332,15 +284,6 @@ TSharedRef<SWidget> FChaosVDRecordingControls::GenerateAdvancedRecordOptionsMenu
 	return MenuBuilder.MakeWidget();
 }
 
-EVisibility FChaosVDRecordingControls::GetRecordingTimeTextBlockVisibility() const
-{
-	using namespace UE::TraceBasedDebuggers;
-	const TSharedPtr<FSessionInfo> SessionInfo = GetCurrentSessionInfo();
-	const bool bIsVisible = SessionInfo && !EnumHasAnyFlags(SessionInfo->GetSessionTypeAttributes(), ERemoteSessionAttributes::IsMultiSessionWrapper) && IsRecording();
-
-	return bIsVisible ? EVisibility::Visible : EVisibility::Collapsed;
-}
-
 TSharedRef<SWidget> FChaosVDRecordingControls::GenerateDataChannelsButton()
 {
 	return SNew(SComboButton)
@@ -460,11 +403,6 @@ bool FChaosVDRecordingControls::CanChangeChannelEnabledState(FString ChannelName
 	return false;
 }
 
-void FChaosVDRecordingControls::SelectTargetSession(FGuid SessionId)
-{
-	CurrentSelectedSessionId = SessionId;
-}
-
 void FChaosVDRecordingControls::GenerateCustomTargetsMenu(FMenuBuilder& MenuBuilder)
 {
 	using namespace UE::TraceBasedDebuggers;
@@ -602,21 +540,6 @@ bool FChaosVDRecordingControls::IsSaveTracesToDiskEnabled()
 	return Settings ? Settings->bSaveRecordingsToDisk : false;
 }
 
-FSlateIcon FChaosVDRecordingControls::GetIconForSession(FGuid SessionId)
-{
-	if (const TSharedPtr<UE::TraceBasedDebuggers::FSessionInfo> SessionInfo = GetSessionInfo(SessionId))
-	{
-		return SessionInfo->IsRecording(Chaos::VD::DebuggerGuid) ? FSlateIcon(FChaosVDStyle::GetStyleSetName(),FName("RecordIcon"), FName("RecordIcon")) : FSlateIcon();
-	}
-
-	return FSlateIcon();
-}
-
-TSharedPtr<UE::TraceBasedDebuggers::FSessionInfo> FChaosVDRecordingControls::GetCurrentSessionInfo() const
-{
-	return GetSessionInfo(CurrentSelectedSessionId);
-}
-
 FChaosVDSessionData* FChaosVDRecordingControls::GetChaosDataForCurrentSession() const
 {
 	if (const TSharedPtr<UE::TraceBasedDebuggers::FSessionInfo> SessionInfo = GetCurrentSessionInfo())
@@ -626,15 +549,6 @@ FChaosVDSessionData* FChaosVDRecordingControls::GetChaosDataForCurrentSession()
 	return nullptr;
 }
 
-TSharedPtr<UE::TraceBasedDebuggers::FSessionInfo> FChaosVDRecordingControls::GetSessionInfo(FGuid Id) const
-{
-	using namespace UE::TraceBasedDebuggers;
-	const TSharedPtr<FRemoteSessionsManager> RemoteSessionManager = FChaosVDEngineEditorBridge::Get().GetSessionsManager();
-	TSharedPtr<FSessionInfo> SessionInfo = RemoteSessionManager ? RemoteSessionManager->GetSessionInfo(Id).Pin() : nullptr;
-
-	return SessionInfo;
-}
-
 bool FChaosVDRecordingControls::HasDataChannelsSupport() const
 {
 	if (const FChaosVDSessionData* SessionData = GetChaosDataForCurrentSession())
@@ -664,287 +578,6 @@ bool FChaosVDRecordingControls::CanChangeLoadingMode() const
 	return false;
 }
 
-const FSlateBrush* FChaosVDRecordingControls::GetRecordOrStopButton() const
-{
-	const FSlateBrush* RecordIconBrush = FChaosVDStyle::Get().GetBrush("RecordIcon");
-	return bRecordingButtonHovered && IsRecording() ? FChaosVDStyle::Get().GetBrush("StopIcon") : RecordIconBrush;
-}
-
-void FChaosVDRecordingControls::HandleRecordingStop(const TSharedPtr<UE::TraceBasedDebuggers::FSessionInfo>&)
-{
-	if (UStatusBarSubsystem* StatusBarSubsystem = GEditor ? GEditor->GetEditorSubsystem<UStatusBarSubsystem>() : nullptr)
-	{
-		StatusBarSubsystem->PopStatusBarMessage(StatusBarID, RecordingMessageHandle);
-
-		const FText LiveSessionEnded = LOCTEXT("RecordingSessionEndedMessage", "Recording session has ended");
-		LiveSessionEndedMessageHandle = StatusBarSubsystem->PushStatusBarMessage(StatusBarID, LiveSessionEnded);
-	}
-}
-
-void FChaosVDRecordingControls::HandleRecordingStart(const TSharedPtr<UE::TraceBasedDebuggers::FSessionInfo>& InSessionInfo)
-{
-	UStatusBarSubsystem* StatusBarSubsystem = GEditor ? GEditor->GetEditorSubsystem<UStatusBarSubsystem>() : nullptr;
-	if (!StatusBarSubsystem)
-	{
-		return;
-	}
-	
-	if (RecordingPathMessageHandle.IsValid())
-	{
-		StatusBarSubsystem->PopStatusBarMessage(StatusBarID, RecordingPathMessageHandle);
-		RecordingPathMessageHandle = FStatusBarMessageHandle();
-	}
-	
-	if (LiveSessionEndedMessageHandle.IsValid())
-	{
-		StatusBarSubsystem->PopStatusBarMessage(StatusBarID, LiveSessionEndedMessageHandle);
-		LiveSessionEndedMessageHandle = FStatusBarMessageHandle();
-	}
-
-	RecordingMessageHandle = StatusBarSubsystem->PushStatusBarMessage(StatusBarID, LOCTEXT("RecordingMessage", "Recording..."));
-}
-
-void FChaosVDRecordingControls::ExecuteAsyncConnectionAttemptTaskWithRetry(FGuid SessionID, int32 RemainingRetries, const TFunction<bool()>& InRecordingStartAttemptCallback, const TFunction<void()>& InRecordingFailedCallback)
-{
-	TSharedPtr<FAsyncConnectionAttemptTask> ConnectionAttemptRequest = MakeShared<FAsyncConnectionAttemptTask>(SessionID, InRecordingStartAttemptCallback, InRecordingFailedCallback, StaticCastWeakPtr<FChaosVDRecordingControls>(AsWeak()));
-	ConnectionAttemptRequest->Start(RemainingRetries, IntervalBetweenAutoplayConnectionAttemptsSeconds);
-}
-
-void FChaosVDRecordingControls::FAsyncConnectionAttemptTask::HandleSuccess()
-{
-	State = EState::Success;
-
-	if (TSharedPtr<FChaosVDRecordingControls> OwnerPtr = Owner.Pin())
-	{
-		OwnerPtr->HandleConnectionAttemptResult(SessionID, EChaosVDLiveConnectionAttemptResult::Success, ProgressNotification);
-	}
-}
-
-void FChaosVDRecordingControls::FAsyncConnectionAttemptTask::HandleFailure()
-{
-	State = EState::Failed;
-
-	if (RecordingFailedCallback)
-	{
-		RecordingFailedCallback();
-	}
-
-	if (TSharedPtr<FChaosVDRecordingControls> OwnerPtr = Owner.Pin())
-	{
-		OwnerPtr->HandleConnectionAttemptResult(SessionID, EChaosVDLiveConnectionAttemptResult::Failed, ProgressNotification);
-	}
-}
-
-bool FChaosVDRecordingControls::FAsyncConnectionAttemptTask::CanExecute() const
-{
-	// The only two states where Execute() can be run are Not Started and InProgress
-	// In Progress is allowed because we might need to execute again as part of the retry process
-	return State == EState::NotStarted || State == EState::InProgress;
-}
-
-void FChaosVDRecordingControls::CancelInFlightConnectionAttempt(const TSharedPtr<UE::TraceBasedDebuggers::FSessionInfo>& SessionInfo)
-{
-	const FGuid SessionID = SessionInfo ? SessionInfo->InstanceId : FGuid{};
-	if (const TSharedPtr<FAsyncConnectionAttemptTask>* ExistingRequestInFlightPtrPtr = InFlightAsyncConnectionRequests.Find(SessionID))
-	{
-		if (const TSharedPtr<FAsyncConnectionAttemptTask> ExistingRequestInFlightPtr = *ExistingRequestInFlightPtrPtr)
-		{
-			ExistingRequestInFlightPtr->Cancel();
-		}
-		else
-		{
-			UE_LOG(LogChaosVDEditor, Verbose, TEXT("[%hs] Attempted to cancel a connection of a Session that is no linger valid | Session ID [%s] ."), __func__, *SessionID.ToString());
-		}
-
-		return;
-	}
-
-	UE_LOG(LogChaosVDEditor, Verbose, TEXT("[%hs] Session ID [%s] not found."), __func__, *SessionID.ToString());
-}
-
-void FChaosVDRecordingControls::HandleRemoteSessionsUpdated()
-{
-	using namespace UE::TraceBasedDebuggers;
-	const TSharedPtr<FSessionInfo> CurrentSessionPtr = GetCurrentSessionInfo();
-	if (!CurrentSessionPtr)
-	{
-		// if the session that is currently selected is not valid, try to find a valid session and select that as default target session
-		if (const TSharedPtr<FRemoteSessionsManager> RemoteSessionManager = FChaosVDEngineEditorBridge::Get().GetSessionsManager())
-		{
-			if (const TSharedPtr<FSessionInfo> FirstValidSession = RemoteSessionManager->GetFirstAvailableSingleSessionInfo().Pin())
-			{
-				SelectTargetSession(FirstValidSession->InstanceId);
-				return;
-			}
-		}
-
-		SelectTargetSession(FRemoteSessionsManager::InvalidSessionGUID);
-	}
-}
-
-FChaosVDRecordingControls::FAsyncConnectionAttemptTask::~FAsyncConnectionAttemptTask()
-{
-	if (RetryDelegateHandle.IsValid())
-	{
-		FTSTicker::RemoveTicker(RetryDelegateHandle);
-		RetryDelegateHandle = FTSTicker::FDelegateHandle();
-	}
-}
-
-void FChaosVDRecordingControls::FAsyncConnectionAttemptTask::Cancel()
-{
-	State = EState::Canceled;
-
-	if (RecordingFailedCallback)
-	{
-		RecordingFailedCallback();
-	}
-
-	if (TSharedPtr<FChaosVDRecordingControls> OwnerPtr = Owner.Pin())
-	{
-		OwnerPtr->HandleConnectionAttemptResult(SessionID, EChaosVDLiveConnectionAttemptResult::Canceled, ProgressNotification);
-	}
-}
-
-void FChaosVDRecordingControls::FAsyncConnectionAttemptTask::ScheduleRetry()
-{
-	if (!CanExecute())
-	{
-		UE_LOG(LogChaosVDEditor, Error, TEXT("[%hs] This connection attempt is in an invalid state | Current State Value [%u]"), __func__, State);
-		return;
-	}
-
-	if (RetryDelegateHandle.IsValid())
-	{
-		UE_LOG(LogChaosVDEditor, Warning, TEXT("[%hs] Attempted to schedule a retry when there is one already scheduled"), __func__);
-		return;
-	}
-
-	if (RemainingRetries <= 0)
-	{
-		UE_LOG(LogChaosVDEditor, Warning, TEXT("[%hs] Failed to connect to live session | attempts exhausted..."), __func__);
-		HandleFailure();
-		return;
-	}
-	
-	RemainingRetries--;
-
-	RetryDelegateHandle = FTSTicker::GetCoreTicker().AddTicker(FTickerDelegate::CreateSPLambda(AsShared(), [this](float DeltaTime)
-	{
-		Execute();
-		return false;
-	}), IntervalBetweenAttemptsSeconds);
-}
-
-void FChaosVDRecordingControls::FAsyncConnectionAttemptTask::Start(int32 InMaxRetriesAttempts, float InIntervalBetweenAttemptsSeconds)
-{
-	if (State != EState::NotStarted)
-	{
-		UE_LOG(LogChaosVDEditor, Error, TEXT("[%hs] This connection attempt is already started or in an invalid state | Current State Value [%u]"), __func__, State);
-		return;
-	}
-
-	TSharedPtr<FChaosVDRecordingControls> OwnerPtr = Owner.Pin();
-	if (!OwnerPtr)
-	{
-		UE_LOG(LogChaosVDEditor, Error, TEXT("[%hs] This connection attempt is does not have a valid owner | Cannot start attempt... "), __func__);
-		return;
-	}
-
-	if (TSharedPtr<FAsyncConnectionAttemptTask>* ExistingRequestInFlight = OwnerPtr->InFlightAsyncConnectionRequests.Find(SessionID))
-	{
-		UE_LOG(LogChaosVDEditor, Warning, TEXT("[%hs] There is a connection attempt in fight for this session. attempting to cancel it | Session ID [%s]"), __func__, *SessionID.ToString());
-		(*ExistingRequestInFlight)->Cancel();
-	}
-
-	OwnerPtr->InFlightAsyncConnectionRequests.Add(SessionID, AsShared());
-
-	State = EState::InProgress;
-
-	RemainingRetries = InMaxRetriesAttempts;
-	IntervalBetweenAttemptsSeconds = InIntervalBetweenAttemptsSeconds;
-	
-	if (!ProgressNotification)
-	{
-		FNotificationInfo Info(LOCTEXT("ConnectingToSessionMessage", "Connecting Session ..."));
-		Info.bFireAndForget = false;
-		Info.FadeOutDuration = 3.0f;
-		Info.ExpireDuration = 0.0f;
-		Info.ButtonDetails.Add(FNotificationButtonInfo(LOCTEXT("CancelConnectionAttemptButton", "Cancel"), LOCTEXT("CancelConnectionAttemptButtonTip", "Cancels the active connection attempt."),
-		FSimpleDelegate::CreateSP(AsShared(), &FAsyncConnectionAttemptTask::Cancel)));
-
-		ProgressNotification = OwnerPtr->PushConnectionAttemptNotification(Info);
-	}
-	
-	// We need to wait at least one tick before attempting to connect to give it time to the trace to be initialized, write to disk, and for the
-	// session manager to hear back from a remote instance
-	ScheduleRetry();
-}
-
-void FChaosVDRecordingControls::FAsyncConnectionAttemptTask::Execute()
-{
-	if (RetryDelegateHandle.IsValid())
-	{
-		FTSTicker::RemoveTicker(RetryDelegateHandle);
-		RetryDelegateHandle.Reset();
-	}
-
-	if (!CanExecute())
-	{
-		UE_LOG(LogChaosVDEditor, Error, TEXT("[%hs] This connection attempt is in an invalid state | Current State Value [%u]"), __func__, State);
-		return;
-	}
-
-	if (!RecordingStartAttemptCallback)
-	{
-		UE_LOG(LogChaosVDEditor, Error, TEXT("[%hs] Failed to connect to session | Invalid callback provided..."), __func__);
-		HandleFailure();
-		return;
-	}
-
-	if (TSharedPtr<FChaosVDRecordingControls> OwnerPtr = Owner.Pin())
-	{
-		 OwnerPtr->UpdateConnectionAttemptNotification(ProgressNotification, RemainingRetries);
-	}
-
-	// CVD needs the trace session name to be able to load a live session. Although the session exist, the session name might not be written right away
-	// Trace files don't really have metadata, it is all part of the same stream, so we need to wait until it is written which might take a few ticks.
-	// Therefore, if it is not ready, try again a few times.
-	if (!RecordingStartAttemptCallback())
-	{
-		UE_LOG(LogChaosVDEditor, Verbose, TEXT("[%hs] Failed to connect to live session | Attempting again in [%f]..."), __func__, IntervalBetweenAttemptsSeconds);
-		ScheduleRetry();
-	}
-	else
-	{
-		HandleSuccess();
-	}
-}
-
-FChaosVDRecordingControls::FTraceStoreConnectionSettings FChaosVDRecordingControls::GetTraceStoreConnectionSettings(const TSharedRef<UE::TraceBasedDebuggers::FSessionInfo>& InSessionInfoRef)
-{
-	using namespace UE::TraceBasedDebuggers;
-	FTraceStoreConnectionSettings ConnectionSettings;
-
-	if (!ensure(!EnumHasAnyFlags(InSessionInfoRef->GetSessionTypeAttributes(), ERemoteSessionAttributes::IsMultiSessionWrapper)))
-	{
-		UE_LOG(LogChaosVDEditor, Error, TEXT("[%s] Cannot not be called with a multi session wrapper."), ANSI_TO_TCHAR(__FUNCTION__));
-
-		return ConnectionSettings;
-	}
-
-	const FTraceConnectionDetails& RecordingSessionDetails = InSessionInfoRef->GetTraceConnectionDetails();
-
-	if (const UE::Trace::FStoreClient::FSessionInfo* TraceSessionInfo = FTraceSessionsManager::GetTraceSessionInfo(RecordingSessionDetails.TraceTarget, RecordingSessionDetails.TraceGuid))
-	{
-		ConnectionSettings.SessionAddress = RecordingSessionDetails.TraceTarget;
-		ConnectionSettings.TraceID = TraceSessionInfo->GetTraceId();
-		return ConnectionSettings;
-	}
-
-	return ConnectionSettings;
-}
-
 void FChaosVDRecordingControls::ToggleMultiSessionSessionRecordingState(const TSharedRef<UE::TraceBasedDebuggers::FMultiSessionInfo>& InSessionInfoRef)
 {
 	const TSharedPtr<SChaosVDMainTab> MainTabSharedPtr = MainTabWeakPtr.Pin();
@@ -968,423 +601,4 @@ void FChaosVDRecordingControls::ToggleMultiSessionSessionRecordingState(const TS
 	});
 }
 
-void FChaosVDRecordingControls::ToggleSingleSessionRecordingState(const TSharedRef<UE::TraceBasedDebuggers::FSessionInfo>& SessionInfoRef)
-{
-	SetSessionRecordingState(!SessionInfoRef->IsRecording(Chaos::VD::DebuggerGuid), SessionInfoRef);
-}
-
-void FChaosVDRecordingControls::SetSessionRecordingState(bool bIsRecording, const TSharedRef<UE::TraceBasedDebuggers::FSessionInfo>& SessionInfoRef)
-{
-	using namespace UE::TraceBasedDebuggers;
-	const TSharedPtr<FRemoteSessionsManager> RemoteSessionManager = FChaosVDEngineEditorBridge::Get().GetSessionsManager();
-	if (!RemoteSessionManager)
-	{
-		UE_LOG(LogChaosVDEditor, Error, TEXT("[%hs] Session Manager is not available"), __func__);
-		return;
-	}
-
-	if (bIsRecording)
-	{
-		const TSharedPtr<SChaosVDMainTab> MainTabSharedPtr = MainTabWeakPtr.Pin();
-		if (!ensure(MainTabSharedPtr))
-		{
-			// This should not be possible, but if we don't have a valid tab instance it will not be possible to properly handle any start recording attempt
-			UE_LOG(LogChaosVDEditor, Error, TEXT("[%hs] Main Tab instance is not available"), __func__);
-			return;
-		}
-
-		TArray<FString, TInlineAllocator<1>> RecordingArgs;
-
-		int32 RemainingRetries = 4;
-		ETraceTransportMode DataTransportOverrideMode = ETraceTransportMode::Invalid;
-
-		if (const UChaosVDGeneralSettings* Settings = FChaosVDSettingsManager::Get().GetSettingsObject<UChaosVDGeneralSettings>())
-		{
-			RemainingRetries = Settings->MaxConnectionRetries;
-			DataTransportOverrideMode = Settings->GetTransportModeForTargetType(SessionInfoRef->BuildTargetType);
-		}
-		else
-		{
-			UE_LOG(LogChaosVDEditor, Warning, TEXT("Failed to obtain setting object. Setting the retries attempts to connect to a session to 4 as a fallback."))
-		}
-
-		// Set the session in a busy state so the auto-expire policy for it is relaxed. Starting a recording might cause the remote instance to stall for a few seconds
-		// if the map loaded is too complex
-		SessionInfoRef->ReadyState = ERemoteSessionReadyState::Busy;
-
-		auto RecordingAttemptFailedCallback = [WeakMainTabPtr = StaticCastWeakPtr<SChaosVDMainTab>(MainTabSharedPtr->AsWeak()), SessionGUID = SessionInfoRef->InstanceId]()
-		{
-			const TSharedPtr<FRemoteSessionsManager> RemoteSessionManagerPtr = FChaosVDEngineEditorBridge::Get().GetSessionsManager();
-			const TSharedPtr<FSessionInfo> SessionInfoPtr = RemoteSessionManagerPtr ? RemoteSessionManagerPtr->GetSessionInfo(SessionGUID).Pin() : nullptr;
-
-			if (!SessionInfoPtr)
-			{
-				return;
-			}
-
-			SessionInfoPtr->ReadyState = ERemoteSessionReadyState::Ready;
-
-			if (const TSharedPtr<SChaosVDMainTab> MainTab = WeakMainTabPtr.Pin())
-			{
-				MainTab->GetChaosVDEngineInstance()->CloseSessionByRemoteSessionID(SessionGUID);
-			}
-		};
-
-		FChaosVDStartRecordingCommandMessage RecordingParams;
-		RecordingParams.RecordingMode = ERecordingMode::Live;
-		RecordingParams.TransportMode = DataTransportOverrideMode != ETraceTransportMode::Invalid ? DataTransportOverrideMode : ETraceTransportMode::Direct;
-
-		check(GLog);
-		bool bOutCanBindAll = false;
-		//TODO: Add a way to specify a local address in case we have multiple adapters?
-		TSharedRef<FInternetAddr> LocalIP = ISocketSubsystem::Get(PLATFORM_SOCKETSUBSYSTEM)->GetLocalHostAddr(*GLog, bOutCanBindAll);
-
-		constexpr bool bAppendPort = false;
-		RecordingParams.Target = LocalIP->ToString(bAppendPort);
-
-#if WITH_CHAOS_VISUAL_DEBUGGER
-		{
-			using namespace Chaos::VisualDebugger;
-			FChaosVDDataChannelsManager::Get().EnumerateDataChannels([&RecordingParams](const TWeakPtr<FChaosVDOptionalDataChannel>& Channel)
-			{
-				if (TSharedPtr<FChaosVDOptionalDataChannel> LockedChannel = Channel.Pin())
-				{
-					if (LockedChannel->IsChannelEnabled())
-					{
-						RecordingParams.DataChannelsEnabledOverrideList.Add(LockedChannel->GetId().ToString());
-					}
-				}
-				return true;
-			});
-		}
-#endif
-
-		if (RecordingParams.TransportMode == ETraceTransportMode::Direct)
-		{
-			// Direct Tracing mode requires opening a trace session in the editor side first before starting the recording
-			// so targets have something to connect to (Usually targets connect to a running trace server)
-			if (!MainTabSharedPtr->ConnectToLiveSession_Direct(SessionInfoRef->InstanceId, CurrentLoadingMode))
-			{
-				UE_LOG(LogChaosVDEditor, Error, TEXT("[%hs] Failed to connect to live session using direct trace mode"), __func__);
-				RecordingAttemptFailedCallback();
-				return;
-			}
-
-			// If we get here it means we have at least one valid session
-			const FTraceSessionDescriptor& SessionDesc = MainTabSharedPtr->GetChaosVDEngineInstance()->GetTraceSessionDescriptors().Last();
-
-			RecordingParams.Target.Appendf(TEXT(":%u"), SessionDesc.SessionPort);
-
-			RemoteSessionManager->SendCommand(SessionInfoRef->Address, RecordingParams);
-
-			ExecuteAsyncConnectionAttemptTaskWithRetry(SessionInfoRef->InstanceId, RemainingRetries, [WeakThis = AsWeak(), SessionInstanceId = SessionInfoRef->InstanceId]()
-			{
-				const TSharedPtr<FRemoteSessionsManager> RemoteSessionManagerPtr = FChaosVDEngineEditorBridge::Get().GetSessionsManager();
-				const TSharedPtr<FSessionInfo> SessionInfoPtr = RemoteSessionManagerPtr ? RemoteSessionManagerPtr->GetSessionInfo(SessionInstanceId).Pin() : nullptr;
-				if (!SessionInfoPtr)
-				{
-					return false;
-				}
-
-				// In Direct trace mode we if we get connection details from the remote instance, we can assume the connection succeeded.
-				return SessionInfoPtr->GetTraceConnectionDetails().IsValid();
-
-			}, RecordingAttemptFailedCallback);
-		}
-		else if (RecordingParams.TransportMode == ETraceTransportMode::Relay)
-		{
-			// Relay Tracing mode requires opening a trace session in the editor side first before starting the recording
-			// so targets have something to connect to (Usually targets connect to a running trace server)
-			if (!MainTabSharedPtr->ConnectToLiveSession_Relay(SessionInfoRef->InstanceId, CurrentLoadingMode))
-			{
-				UE_LOG(LogChaosVDEditor, Error, TEXT("[%hs] Failed to connect to live session using relay trace mode"), __func__);
-				RecordingAttemptFailedCallback();
-				return;
-			}
-
-			RemoteSessionManager->SendCommand(SessionInfoRef->Address, RecordingParams);
-
-#if WITH_CHAOS_VISUAL_DEBUGGER
-			// Once the start recording command is issued, we can try to connect to the created session (if everything went well).
-			// If it didn't go well, this will take care of update the UI to notify the user
-			ExecuteAsyncConnectionAttemptTaskWithRetry(SessionInfoRef->InstanceId, RemainingRetries, [WeakThis = AsWeak(), SessionInstanceId = SessionInfoRef->InstanceId]()
-			{
-				const TSharedPtr<FRemoteSessionsManager> RemoteSessionManagerPtr = FChaosVDEngineEditorBridge::Get().GetSessionsManager();
-				const TSharedPtr<FSessionInfo> SessionInfoPtr = RemoteSessionManagerPtr ? RemoteSessionManagerPtr->GetSessionInfo(SessionInstanceId).Pin() : nullptr;
-				if (!SessionInfoPtr)
-				{
-					return false;
-				}
-
-				const FTraceConnectionDetails& RecordingSessionDetails = SessionInfoPtr->GetTraceConnectionDetails();
-
-				// We didn't receive the connection details yet. Trying again later...
-				if (!RecordingSessionDetails.IsValid())
-				{
-					return false;
-				}
-
-				// Check if there is a connection attempt in progress and its results
-				const EConnectionAttemptResult ConnectionAttemptStatus = FChaosVDEngineEditorBridge::Get().GetTraceRelayTransportInstance()->GetConnectionAttemptResult(SessionInstanceId);
-				switch (ConnectionAttemptStatus)
-				{
-				case EConnectionAttemptResult::InProgress:
-					{
-						// A connection attempt started, but we don't have the results yet. Trying again later...
-						return false;
-					}
-				case EConnectionAttemptResult::NotStarted:
-					{
-						FRelayConnectionInfo ConnectionInfo;
-						ConnectionInfo.Port = RecordingSessionDetails.Port;
-						ConnectionInfo.Address = RecordingSessionDetails.TraceTarget;
-						ConnectionInfo.CertificateAuthority = RecordingSessionDetails.CertAuth;
-						EConnectionAttemptResult Result = FChaosVDEngineEditorBridge::Get().GetTraceRelayTransportInstance()->ConnectToRelay(SessionInstanceId, ConnectionInfo);
-
-						return Result == EConnectionAttemptResult::Success;
-					}
-				case EConnectionAttemptResult::Success:
-					// The current connection attempt succeeded, we can stop the async retry loop
-					return true;
-				case EConnectionAttemptResult::Failed:
-				default:
-					// There was a connection attempt, but it failed. Abandoning the retry attempt as they will not succeed
-					return false;
-				}
-			},
-			RecordingAttemptFailedCallback);
-#endif
-		}
-	}
-	else
-	{
-		CancelInFlightConnectionAttempt(SessionInfoRef.ToSharedPtr());
-		RemoteSessionManager->SendCommand(SessionInfoRef->Address, FChaosVDStopRecordingCommandMessage{});
-	}
-}
-
-FReply FChaosVDRecordingControls::ToggleRecordingState()
-{
-	using namespace UE::TraceBasedDebuggers;
-	TSharedPtr<FSessionInfo> SessionInfoPtr = GetCurrentSessionInfo();
-	if (!SessionInfoPtr)
-	{
-		return FReply::Handled();
-	}
-
-	if (EnumHasAnyFlags(SessionInfoPtr->GetSessionTypeAttributes(), ERemoteSessionAttributes::IsMultiSessionWrapper))
-	{
-		ToggleMultiSessionSessionRecordingState(StaticCastSharedPtr<FMultiSessionInfo>(SessionInfoPtr).ToSharedRef());
-	}
-	else
-	{
-		ToggleSingleSessionRecordingState(SessionInfoPtr.ToSharedRef());
-	}
-
-	return FReply::Handled();
-}
-
-bool FChaosVDRecordingControls::IsRecordingToggleButtonEnabled() const
-{
-	using namespace UE::TraceBasedDebuggers;
-	if (CurrentSelectedSessionId == FRemoteSessionsManager::InvalidSessionGUID)
-	{
-		return false;
-	}
-
-	TSharedPtr<FSessionInfo> SessionInfo = GetCurrentSessionInfo();
-	if (!SessionInfo)
-	{
-		return false;
-	}
-
-	// Do not allow recording if another trace-based debuggers is already recording in the current session
-	if (SessionInfo->LastKnownRecordingState.DebuggerId.IsValid()
-		&& SessionInfo->LastKnownRecordingState.DebuggerId != Chaos::VD::DebuggerGuid
-		&& SessionInfo->LastKnownRecordingState.bIsRecording)
-	{
-		return false;
-	}
-
-	if (!IsRecording())
-	{
-		return true;
-	}
-
-	// If we are recording, don't show the stop button for the mode that is disabled
-	if (EnumHasAnyFlags(SessionInfo->GetSessionTypeAttributes(), ERemoteSessionAttributes::IsMultiSessionWrapper))
-	{
-		TSharedPtr<FMultiSessionInfo> AsMultiSessionInfo = StaticCastSharedPtr<FMultiSessionInfo>(SessionInfo);
-
-		return CanSelectMultiSessionTarget(SessionInfo.ToSharedRef());
-	}
-
-	return true;
-}
-
-EVisibility FChaosVDRecordingControls::IsRecordingToggleButtonVisible() const
-{
-	return EVisibility::Visible;
-}
-
-void FChaosVDRecordingControls::AddToMenu(FName ExistingMenu)
-{
-	UToolMenu* ToolBar = UToolMenus::Get()->ExtendMenu(ExistingMenu);
-	if (!ensure(ToolBar))
-	{
-		return;
-	}
-
-	FToolMenuSection& Section = ToolBar->AddSection("RecordingControls");
-
-	Section.AddDynamicEntry("Controls", FNewToolMenuSectionDelegate::CreateLambda([](FToolMenuSection& InSection)
-	{
-		const UChaosVDMainToolbarMenuContext* Context = InSection.FindContext<UChaosVDMainToolbarMenuContext>();
-		TSharedPtr<SChaosVDMainTab> MainTabPtr = Context->MainTab.Pin();
-		TSharedPtr<FChaosVDRecordingControls> RecordingControls = MainTabPtr ? MainTabPtr->GetRecordingControls() : nullptr;
-		if (!ensure(RecordingControls))
-		{
-			return;
-		}
-
-		TSharedRef<SWidget> StartRecordingButton = SNew(SBox).Padding(4.0f,0.0f)[ RecordingControls->GenerateToggleRecordingStateButton(LOCTEXT("RecordButtonDesc", "Starts a recording for the current session, saving it directly to file")) ];
-		TSharedRef<SWidget> SessionsDropdown = SNew(SBox).Padding(4.0f,0.0f)[ RecordingControls->GenerateTargetSessionSelector() ];
-		TSharedRef<SWidget> RecordingTime = RecordingControls->GenerateRecordingTimeTextBlock();
-		TSharedRef<SWidget> DataChannelsButton = SNew(SBox).Padding(4.0f,0.0f)[ RecordingControls->GenerateDataChannelsButton() ];
-		TSharedRef<SWidget> AdvancedRecordingOptionsButton = SNew(SBox).Padding(1.0f,0.0f)[ RecordingControls->GenerateAdvancedRecordOptionsButton() ];
-
-		InSection.AddSeparator("RecordingControlsDivider");
-
-		InSection.AddEntry(
-			FToolMenuEntry::InitWidget(
-				"AvailableSessions",
-				SessionsDropdown,
-				FText::GetEmpty(),
-				false,
-				false
-			));
-
-		InSection.AddEntry(
-			FToolMenuEntry::InitWidget(
-				"StartRecording",
-				StartRecordingButton,
-				FText::GetEmpty(),
-				true,
-				false
-			));
-
-		InSection.AddEntry(
-		FToolMenuEntry::InitWidget(
-			"AdvancedRecordingOptions",
-			AdvancedRecordingOptionsButton,
-			FText::GetEmpty(),
-			true,
-			false
-		));
-
-		InSection.AddEntry(
-			FToolMenuEntry::InitWidget(
-				"RecordingTime",
-				RecordingTime,
-				FText::GetEmpty(),
-				false,
-				false
-			));
-
-		InSection.AddEntry(
-			FToolMenuEntry::InitWidget(
-				"DataChannelsButton",
-				DataChannelsButton,
-				FText::GetEmpty(),
-				false,
-				false
-			));
-	}));
-}
-
-bool FChaosVDRecordingControls::IsRecording() const
-{
-	using namespace UE::TraceBasedDebuggers;
-	const TSharedPtr<FSessionInfo> SessionInfo = GetCurrentSessionInfo();
-	return SessionInfo ? SessionInfo->IsRecording(Chaos::VD::DebuggerGuid) : false;
-}
-
-FText FChaosVDRecordingControls::GetRecordingTimeText() const
-{
-	using namespace UE::TraceBasedDebuggers;
-	if (const TSharedPtr<FSessionInfo> SessionInfo = GetCurrentSessionInfo())
-	{
-		FNumberFormattingOptions FormatOptions;
-		FormatOptions.MinimumFractionalDigits = 2;
-		FormatOptions.MaximumFractionalDigits = 2;
-
-		FText SecondsText = FText::AsNumber(SessionInfo->LastKnownRecordingState.ElapsedTime, &FormatOptions);
-		return FText::Format(LOCTEXT("RecordingTimer","{0} s"), SecondsText);
-	}
-
-	return LOCTEXT("RecordingTimerError","Failed to get time information");
-}
-
-TSharedPtr<SNotificationItem> FChaosVDRecordingControls::PushConnectionAttemptNotification(const FNotificationInfo& InNotificationInfo)
-{
-	TSharedPtr<SNotificationItem> ConnectionAttemptNotification = FSlateNotificationManager::Get().AddNotification(InNotificationInfo);
-	
-	if (ConnectionAttemptNotification.IsValid())
-	{
-		ConnectionAttemptNotification->SetCompletionState(SNotificationItem::CS_Pending);
-		return ConnectionAttemptNotification;
-	}
-
-	return nullptr;
-}
-
-void FChaosVDRecordingControls::UpdateConnectionAttemptNotification(const TSharedPtr<SNotificationItem>& InNotification, int32 AttemptsRemaining)
-{
-	if (InNotification)
-	{
-		InNotification->SetSubText(FText::FormatOrdered(LOCTEXT("SessionConnectionAttemptSubText", "Attempts Remaining {0}"), AttemptsRemaining));
-	}
-}
-
-void FChaosVDRecordingControls::HandleConnectionAttemptResult(FGuid SessionGUID, EChaosVDLiveConnectionAttemptResult Result, const TSharedPtr<SNotificationItem>& InNotification)
-{
-	InFlightAsyncConnectionRequests.Remove(SessionGUID);
-
-	if (InNotification)
-	{
-		switch (Result)
-		{
-		case EChaosVDLiveConnectionAttemptResult::Success:
-			{
-				InNotification->SetText(LOCTEXT("SessionConnectionSuccess", "Connected!"));
-				InNotification->SetSubText(FText::GetEmpty());
-				InNotification->SetCompletionState(SNotificationItem::ECompletionState::CS_Success);
-				break;
-			}
-		case EChaosVDLiveConnectionAttemptResult::Canceled:
-			{
-				InNotification->SetText(LOCTEXT("SessionConnectionCanceledText", "Connection Attempt Canceled"));
-				InNotification->SetCompletionState(SNotificationItem::ECompletionState::CS_None);
-				break;
-			}
-		case EChaosVDLiveConnectionAttemptResult::Failed:
-			{
-				InNotification->SetText(LOCTEXT("SessionConnectionFailedText", "Failed to connect"));
-				InNotification->SetSubText(LOCTEXT("SessionConnectionFailedSubText", "See the logs for more details..."));
-				InNotification->SetCompletionState(SNotificationItem::ECompletionState::CS_Fail);
-				break;
-			}
-		default:
-			{
-				InNotification->SetText(LOCTEXT("SessionConnectionUnknownText", "Something went wrong..."));
-				InNotification->SetSubText(LOCTEXT("SessionConnectionUnknownTextSubText", "See the logs for more details..."));
-				InNotification->SetCompletionState(SNotificationItem::ECompletionState::CS_Fail);
-				break;
-			}
-		}
-
-		InNotification->ExpireAndFadeout();
-	}
-}
-
-#undef LOCTEXT_NAMESPACE 
+#undef LOCTEXT_NAMESPACE
diff --git a/Engine/Plugins/ChaosVD/Source/ChaosVD/Private/Widgets/ChaosVDRecordingControls.h b/Engine/Plugins/ChaosVD/Source/ChaosVD/Private/Widgets/ChaosVDRecordingControls.h
--- a/Engine/Plugins/ChaosVD/Source/ChaosVD/Private/Widgets/ChaosVDRecordingControls.h
+++ b/Engine/Plugins/ChaosVD/Source/ChaosVD/Private/Widgets/ChaosVDRecordingControls.h
@@ -1,72 +1,65 @@
-// Copyright Epic Games, Inc. All Rights Reserved.
+﻿// Copyright Epic Games, Inc. All Rights Reserved.
 
 #pragma once
 
-#include "StatusBarSubsystem.h"
-#include "ChaosVDEngine.h"
-#include "RemoteSessionsManager.h"
-#include "SessionInfo.h"
-#include "Templates/SharedPointer.h"
-#include "Widgets/SCompoundWidget.h"
+#include "Widgets/RecordingControls.h"
 
-#include "ChaosVDRecordingControls.generated.h"
-
-struct FChaosVDSessionData;
-class FMenuBuilder;
+namespace UE::TraceBasedDebuggers
+{
+struct FMultiSessionInfo;
+}
 
 namespace Chaos::VisualDebugger
 {
-	struct FChaosVDOptionalDataChannel;
+struct FChaosVDOptionalDataChannel;
 }
 
-class FAsyncCompilationNotification;
-class FReply;
-class SButton;
+enum class EChaosVDLoadRecordedDataMode : uint8;
+struct FChaosVDSessionData;
+class FMenuBuilder;
 class SChaosVDMainTab;
-struct FSlateBrush;
-struct FSlateIcon;
-
-UENUM()
-enum class EChaosVDLiveConnectionAttemptResult
-{
-	Success,
-	Failed,
-	Canceled
-};
 
 typedef Chaos::VisualDebugger::FChaosVDOptionalDataChannel FCVDDataChannel;
 
-class FChaosVDRecordingControls : public TSharedFromThis<FChaosVDRecordingControls>
+class FChaosVDRecordingControls : public UE::TraceBasedDebuggers::FRecordingControls
 {
 public:
 
-	FChaosVDRecordingControls(const TSharedRef<SChaosVDMainTab>& InMainTabSharedRef);
-	~FChaosVDRecordingControls();
-
-	void Initialize();
-	void Deinitialize();
-	void AddToMenu(FName ExistingMenu);
+	FChaosVDRecordingControls(const TSharedRef<SChaosVDMainTab>& InMainTab, const TSharedRef<UE::TraceBasedDebuggers::FRemoteSessionsManager>& InRemoteSessionManager);
 
 protected:
+	//~ Begin FRecordingControls interface
+	virtual void AddToMenuInternal(FToolMenuSection& InSection) override;
+	virtual void OnGenerateTargetSessionSelectorInternal(FMenuBuilder& InMenuBuilder) override;
+
+	virtual int32 GetMaxConnectionRetriesInternal() override;
+	virtual UE::TraceBasedDebuggers::ETraceTransportMode GetTransportModeOverrideForTargetTypeInternal(EBuildTargetType InBuildTarget) override;
+
+	virtual bool ConnectToLiveSession_RelayInternal(const FGuid& InSessionId) override;
+	virtual bool ConnectToLiveSession_DirectInternal(const FGuid& InSessionId) override;
+
+	[[nodiscard]] virtual TNotNull<TUniquePtr<UE::TraceBasedDebuggers::FStartRecordingCommandMessage>> BuildStartRecordingParamsInternal() override;
+	virtual void SendStartRecordingCommandInternal(const FMessageAddress& InAddress, TUniquePtr<UE::TraceBasedDebuggers::FStartRecordingCommandMessage> InStartRecordingCommand) override;
+	virtual void SendStopRecordingCommandInternal(const FMessageAddress& InAddress) override;
+
+	virtual bool IsRecordingToggleButtonEnabledInternal(const TSharedRef<UE::TraceBasedDebuggers::FSessionInfo>& InSessionInfo) const override;
+	virtual void OnToggleRecordingStateInternal(const TSharedRef<UE::TraceBasedDebuggers::FSessionInfo>& InSessionInfo) override;
+
+	virtual void CloseSessionByRemoteSessionIDInternal(const FGuid& InSessionId) override;
+
+	virtual UE::TraceBasedDebuggers::FTraceSessionDescriptor& GetCurrentSessionDescriptorInternal() override;
+	//~ End FRecordingControls interface
 
-	TSharedRef<SWidget> GenerateToggleRecordingStateButton(const FText& StartRecordingTooltip);
-	TSharedRef<SWidget> GenerateTargetSessionSelector();
-	TSharedRef<SWidget> GenerateTargetSessionDropdown();
 	TSharedRef<SWidget> GenerateDataChannelsMenu();
 	TSharedRef<SWidget> GenerateDataChannelsButton();
 	TSharedRef<SWidget> GenerateLoadingModeSelector();
-	TSharedRef<SWidget> GenerateRecordingTimeTextBlock();
 	TSharedRef<SWidget> GenerateAdvancedRecordOptionsButton();
 	TSharedRef<SWidget> GenerateAdvancedRecordOptionsMenu();
-	
-	EVisibility GetRecordingTimeTextBlockVisibility() const;
 
 	void ToggleChannelEnabledState(FString ChannelName);
 	bool IsChannelEnabled(FString ChannelName);
 	bool CanChangeChannelEnabledState(FString ChannelName);
 
-	void SelectTargetSession(FGuid SessionId);
-
 	void GenerateCustomTargetsMenu(FMenuBuilder& MenuBuilder);
 
 	bool IsSessionPartOfCustomTargetSelection(FGuid SessionGuid);
@@ -81,130 +74,15 @@ protected:
 	void ToggleSaveTracesToDisk();
 	bool IsSaveTracesToDiskEnabled();
 
-	FSlateIcon GetIconForSession(FGuid SessionId);
-
-	TSharedPtr<UE::TraceBasedDebuggers::FSessionInfo> GetCurrentSessionInfo() const;
-	TSharedPtr<UE::TraceBasedDebuggers::FSessionInfo> GetSessionInfo(FGuid Id) const;
-
 	FChaosVDSessionData* GetChaosDataForCurrentSession() const;
 
-	void HandleRemoteSessionsUpdated();
-
 	bool HasDataChannelsSupport() const;
 
 	bool CanChangeLoadingMode() const;
-	
-	const FSlateBrush* GetRecordOrStopButton() const;
-	
-	void HandleRecordingStop(const TSharedPtr<UE::TraceBasedDebuggers::FSessionInfo>& SessionInfo);
-	void HandleRecordingStart(const TSharedPtr<UE::TraceBasedDebuggers::FSessionInfo>& SessionInfo);
 
-	void ExecuteAsyncConnectionAttemptTaskWithRetry(FGuid SessionID, int32 RemainingRetries, const TFunction<bool()>& InRecordingStartAttemptCallback, const TFunction<void()>& InRecordingFailedCallback);
-
-	FReply ToggleRecordingState();
 	void ToggleMultiSessionSessionRecordingState(const TSharedRef<UE::TraceBasedDebuggers::FMultiSessionInfo>& InSessionInfoRef);
-	void ToggleSingleSessionRecordingState(const TSharedRef<UE::TraceBasedDebuggers::FSessionInfo>& SessionInfoRef);
-	void SetSessionRecordingState(bool bIsRecording, const TSharedRef<UE::TraceBasedDebuggers::FSessionInfo>& SessionInfoRef);
-
-	bool IsRecordingToggleButtonEnabled() const;
-	EVisibility IsRecordingToggleButtonVisible() const;
-
-	bool IsRecording() const;
-
-	FText GetRecordingTimeText() const;
-
-	TSharedPtr<SNotificationItem> PushConnectionAttemptNotification(const FNotificationInfo& InNotificationInfo);
-	void UpdateConnectionAttemptNotification(const TSharedPtr<SNotificationItem>& InNotification, int32 AttemptsRemaining);
-	void HandleConnectionAttemptResult(FGuid SessionGUID, EChaosVDLiveConnectionAttemptResult Result, const TSharedPtr<SNotificationItem>& InNotification);
-
-	struct FTraceStoreConnectionSettings
-	{
-		int32 TraceID = INDEX_NONE;
-		FString SessionAddress;
-	};
-
-	FTraceStoreConnectionSettings GetTraceStoreConnectionSettings(const TSharedRef<UE::TraceBasedDebuggers::FSessionInfo>& InSessionInfoRef);
-
-	FText GetCurrentSelectedSessionName() const;
-
-	FGuid CurrentSelectedSessionId = FGuid();
-
-	FName StatusBarID;
-	
-	FStatusBarMessageHandle RecordingMessageHandle;
-	FStatusBarMessageHandle RecordingPathMessageHandle;
-	FStatusBarMessageHandle LiveSessionEndedMessageHandle;
-	FDelegateHandle RecordingStartedHandle;
-	FDelegateHandle RecordingStoppedHandle;
 
 	TWeakPtr<SChaosVDMainTab> MainTabWeakPtr;
 
-	float IntervalBetweenAutoplayConnectionAttemptsSeconds = 1.0f;
-
-	bool bRecordingButtonHovered = false;
-
-	FCurveSequence RecordingAnimation;
-
-	EChaosVDLoadRecordedDataMode CurrentLoadingMode = EChaosVDLoadRecordedDataMode::SingleSource;
-
-	void CancelInFlightConnectionAttempt(const TSharedPtr<UE::TraceBasedDebuggers::FSessionInfo>& SessionInfo);
-
-	enum class ERecordingStartRequestState
-	{
-		Invalid,
-		InProgress,
-		Success,
-		Failed,
-		Canceled,
-	};
-
-	struct FAsyncConnectionAttemptTask : public TSharedFromThis<FAsyncConnectionAttemptTask>
-	{
-		explicit FAsyncConnectionAttemptTask(FGuid InSessionID, const TFunction<bool()>& InRecordingStartAttemptCallback, const TFunction<void()>& InRecordingFailedCallback, const TWeakPtr<FChaosVDRecordingControls>& RequestOwner)
-			: SessionID(InSessionID),
-			RecordingStartAttemptCallback(InRecordingStartAttemptCallback),
-			RecordingFailedCallback(InRecordingFailedCallback),
-			Owner(RequestOwner)
-		{
-		}
-
-		~FAsyncConnectionAttemptTask();
-
-		enum class EState : uint8
-		{
-			NotStarted,
-			InProgress,
-			Success,
-			Canceled,
-			Failed
-		};
-
-		void Cancel();
-
-		void Start(int32 InMaxRetriesAttempts, float InIntervalBetweenAttemptsSeconds);
-		void ScheduleRetry();
-		void HandleSuccess();
-		void HandleFailure();
-
-	private:
-
-		bool CanExecute() const;
-
-		void Execute();
-
-		FGuid SessionID;
-		int32 RemainingRetries = 10;
-
-		float IntervalBetweenAttemptsSeconds = 1.0f;
-	
-		TFunction<bool()> RecordingStartAttemptCallback;
-		TSharedPtr<SNotificationItem> ProgressNotification;
-		const TFunction<void()> RecordingFailedCallback;
-		EState State = EState::NotStarted;
-		FTSTicker::FDelegateHandle RetryDelegateHandle;
-		TWeakPtr<FChaosVDRecordingControls> Owner;
-	};
-
-	TMap<FGuid, TSharedPtr<FAsyncConnectionAttemptTask>> InFlightAsyncConnectionRequests;
-	bool bInitialized = false;
+	EChaosVDLoadRecordedDataMode CurrentLoadingMode;
 };
diff --git a/Engine/Plugins/ChaosVD/Source/ChaosVD/Private/Widgets/SChaosVDMainTab.cpp b/Engine/Plugins/ChaosVD/Source/ChaosVD/Private/Widgets/SChaosVDMainTab.cpp
--- a/Engine/Plugins/ChaosVD/Source/ChaosVD/Private/Widgets/SChaosVDMainTab.cpp
+++ b/Engine/Plugins/ChaosVD/Source/ChaosVD/Private/Widgets/SChaosVDMainTab.cpp
@@ -1,4 +1,4 @@
-// Copyright Epic Games, Inc. All Rights Reserved.
+﻿// Copyright Epic Games, Inc. All Rights Reserved.
 
 #include "Widgets/SChaosVDMainTab.h"
 
@@ -33,7 +33,7 @@
 #include "IDesktopPlatform.h"
 #include "PropertyEditorModule.h"
 #include "SChaosVDPlaybackViewport.h"
-#include "SConnectionStatusBar.h"
+#include "Widgets/SConnectionStatusBar.h"
 #include "Misc/MessageDialog.h"
 #include "StatusBarSubsystem.h"
 #include "ToolMenu.h"
@@ -265,25 +265,39 @@ void SChaosVDMainTab::ExtendMainToolbar()
 {
 	if (!RecordingControls)
 	{
-		RecordingControls = MakeShared<FChaosVDRecordingControls>(StaticCastSharedRef<SChaosVDMainTab>(AsShared()));
-		RecordingControls->Initialize();
+		if (TSharedPtr<UE::TraceBasedDebuggers::FRemoteSessionsManager> RemoteSessionManager = FChaosVDEngineEditorBridge::Get().GetSessionsManager())
+		{
+			RecordingControls = MakeShared<FChaosVDRecordingControls>(StaticCastSharedRef<SChaosVDMainTab>(AsShared()), RemoteSessionManager.ToSharedRef());
+			RecordingControls->Initialize();
+		}
 	}
-	
-	RecordingControls->AddToMenu(MainToolBarName);
+
+	if (RecordingControls)
+	{
+		RecordingControls->AddToMenu(MainToolBarName);
+	}
+
 	AddSettingsMenu();
 }
 
 void SChaosVDMainTab::ExtendStatusToolbar()
 {
-	FName StatusBarToolbarID = FName(StatusBarID.ToString() + TEXT(".Toolbar"));
-
-	if (UToolMenu* StatusToolBarMenu = UToolMenus::Get()->ExtendMenu(StatusBarToolbarID))
+	if (TSharedPtr<UE::TraceBasedDebuggers::FRemoteSessionsManager> RemoteSessionManager = FChaosVDEngineEditorBridge::Get().GetSessionsManager())
 	{
-		FToolMenuSection& ConnectionStatusSection = StatusToolBarMenu->AddSection("ConnectionStatus");
-
-		constexpr bool bNoIndent = true;
-		constexpr bool bSearchable = false;
-		ConnectionStatusSection.AddEntry(FToolMenuEntry::InitWidget("ConnectionStatusBar", SNew(Chaos::VD::SConnectionStatusBar), FText::GetEmpty(), bNoIndent, bSearchable));
+		FName StatusBarToolbarID = FName(StatusBarID.ToString() + TEXT(".Toolbar"));
+	
+		if (UToolMenu* StatusToolBarMenu = UToolMenus::Get()->ExtendMenu(StatusBarToolbarID))
+		{
+			FToolMenuSection& ConnectionStatusSection = StatusToolBarMenu->AddSection("ConnectionStatus");
+	
+			constexpr bool bNoIndent = true;
+			constexpr bool bSearchable = false;
+			ConnectionStatusSection.AddEntry(FToolMenuEntry::InitWidget("ConnectionStatusBar"
+				, SNew(UE::TraceBasedDebuggers::SConnectionStatusBar, Chaos::VD::DebuggerGuid, RemoteSessionManager.ToSharedRef())
+				, FText::GetEmpty()
+				, bNoIndent
+				, bSearchable));
+		}
 	}
 }
 
diff --git a/Engine/Plugins/ChaosVD/Source/ChaosVD/Private/Widgets/SConnectionStatusBar.cpp b/Engine/Plugins/ChaosVD/Source/ChaosVD/Private/Widgets/SConnectionStatusBar.cpp
deleted file mode 100644
diff --git a/Engine/Plugins/ChaosVD/Source/ChaosVD/Private/Widgets/SConnectionStatusBar.h b/Engine/Plugins/ChaosVD/Source/ChaosVD/Private/Widgets/SConnectionStatusBar.h
deleted file mode 100644
diff --git a/Engine/Source/Runtime/Experimental/TraceBasedDebuggersAnalysis/Private/TraceSessionsManager.cpp b/Engine/Source/Runtime/Experimental/TraceBasedDebuggersAnalysis/Private/TraceSessionsManager.cpp
--- a/Engine/Source/Runtime/Experimental/TraceBasedDebuggersAnalysis/Private/TraceSessionsManager.cpp
+++ b/Engine/Source/Runtime/Experimental/TraceBasedDebuggersAnalysis/Private/TraceSessionsManager.cpp
@@ -7,9 +7,12 @@
 #include "Containers/Ticker.h"
 #include "EngineEditorBridge.h"
 #include "Features/IModularFeatures.h"
+#if WITH_EDITOR
 #include "Framework/Notifications/NotificationManager.h"
+#endif
 #include "Internationalization/Internationalization.h"
 #include "Internationalization/Text.h"
+#include "Misc/Paths.h"
 #include "Modules/ModuleManager.h"
 #include "Trace/DataStream.h"
 #include "Trace/StoreClient.h"
@@ -93,7 +96,10 @@ private:
 
 		float ElapsedTimeSinceLastUpdate = 0.0f;
 
+#if WITH_EDITOR
 		FProgressNotificationHandle NotificationHandle;
+#endif // WITH_EDITOR
+
 		std::atomic<int32> CurrentProgress = 0;
 		FText NotificationTitle;
 	};
@@ -107,6 +113,7 @@ private:
 
 FFileDataStream::FAsyncProgressNotification::~FAsyncProgressNotification()
 {
+#if WITH_EDITOR
 	if (NotificationHandle.IsValid())
 	{
 		// FSlateNotificationManager is not thread safe, so we need to make sure this runs in the game thread
@@ -124,6 +131,7 @@ FFileDataStream::FAsyncProgressNotification::~FAsyncProgressNotification()
 				});
 		}
 	}
+#endif // WITH_EDITOR
 }
 
 void FFileDataStream::FAsyncProgressNotification::EnterProgress(const int32 InCurrentProgress)
@@ -133,6 +141,7 @@ void FFileDataStream::FAsyncProgressNotification::EnterProgress(const int32 InCu
 
 bool FFileDataStream::FAsyncProgressNotification::Tick(const float DeltaTime)
 {
+#if WITH_EDITOR
 	constexpr int32 MaxProgress = 100;
 	constexpr float TimeBetweenUpdatesSeconds = 0.5f;
 
@@ -166,6 +175,9 @@ bool FFileDataStream::FAsyncProgressNotification::Tick(const float DeltaTime)
 	}
 
 	return true;
+#else
+	return false;
+#endif // WITH_EDITOR
 }
 
 FTraceSessionsManager::FTraceSessionsManager(
diff --git a/Engine/Source/Runtime/Experimental/TraceBasedDebuggersAnalysis/Private/Widgets/RecordingControls.cpp b/Engine/Source/Runtime/Experimental/TraceBasedDebuggersAnalysis/Private/Widgets/RecordingControls.cpp
new file mode 100644
--- /dev/null
+++ b/Engine/Source/Runtime/Experimental/TraceBasedDebuggersAnalysis/Private/Widgets/RecordingControls.cpp
@@ -0,0 +1,910 @@
+// Copyright Epic Games, Inc. All Rights Reserved.
+
+#if WITH_EDITOR
+
+#include "Widgets/RecordingControls.h"
+#include "Editor.h"
+#include "Framework/MultiBox/MultiBoxBuilder.h"
+#include "Framework/SlateDelegates.h"
+#include "RemoteSessionsManager.h"
+#include "SEnumCombo.h"
+#include "SessionInfo.h"
+#include "SocketSubsystem.h"
+#include "ToolMenuEntry.h"
+#include "ToolMenus.h"
+#include "ToolMenuSection.h"
+#include "TraceSessionsManager.h"
+#include "Widgets/Images/SImage.h"
+#include "Widgets/Input/SButton.h"
+#include "Widgets/Input/SComboButton.h"
+#include "Widgets/Notifications/SNotificationList.h"
+#include "Widgets/Text/STextBlock.h"
+
+#include UE_INLINE_GENERATED_CPP_BY_NAME(RecordingControls)
+
+#define LOCTEXT_NAMESPACE "RecordingControls"
+namespace UE::TraceBasedDebuggers
+{
+
+FRecordingControls::FRecordingControls(const FName InStatusBarId
+	, const TSharedRef<FRemoteSessionsManager>& InRemoteSessionManager
+	, const FGuid& InDebuggerId
+	, const FLogCategoryBase& InLogCategory)
+	: RemoteSessionManagerWeakPtr(InRemoteSessionManager)
+	, LogCategory(InLogCategory)
+	, DebuggerGuid(InDebuggerId)
+	, StatusBarID(InStatusBarId)
+{
+	RecordingAnimation = FCurveSequence();
+	RecordingAnimation.AddCurve(0.f, 1.5f, ECurveEaseFunction::Linear);
+
+	CurrentSelectedSessionId = FRemoteSessionsManager::LocalEditorSessionID;
+}
+
+FRecordingControls::~FRecordingControls()
+{
+	Deinitialize();
+}
+
+void FRecordingControls::Initialize()
+{
+	if (bInitialized)
+	{
+		return;
+	}
+
+	if (const TSharedPtr<FRemoteSessionsManager> RemoteSessionManager = RemoteSessionManagerWeakPtr.Pin())
+	{
+		RemoteSessionManager->OnSessionRecordingStarted().AddSP(this, &FRecordingControls::HandleRecordingStart);
+		RemoteSessionManager->OnSessionRecordingStopped().AddSP(this, &FRecordingControls::HandleRecordingStop);
+		RemoteSessionManager->OnSessionExpired().AddSP(this, &FRecordingControls::CancelInFlightConnectionAttempt);
+		RemoteSessionManager->OnSessionsUpdated().AddSP(this, &FRecordingControls::HandleRemoteSessionsUpdated);
+	}
+
+	bInitialized = true;
+}
+
+void FRecordingControls::Deinitialize()
+{
+	if (!bInitialized)
+	{
+		return;
+	}
+
+	if (const TSharedPtr<FRemoteSessionsManager> RemoteSessionManager = RemoteSessionManagerWeakPtr.Pin())
+	{
+		RemoteSessionManager->OnSessionRecordingStarted().RemoveAll(this);
+		RemoteSessionManager->OnSessionRecordingStopped().RemoveAll(this);
+		RemoteSessionManager->OnSessionExpired().RemoveAll(this);
+		RemoteSessionManager->OnSessionsUpdated().RemoveAll(this);
+	}
+
+	bInitialized = false;
+}
+
+TSharedRef<SWidget> FRecordingControls::GenerateToggleRecordingStateButton(const FText& StartRecordingTooltip)
+{
+	TSharedRef<SImage> Animation = SNew(SImage)
+		.Image_Raw(this, &FRecordingControls::GetRecordOrStopButton);
+
+	Animation->SetColorAndOpacity(MakeAttributeSPLambda(Animation, [this, Animation]
+		{
+			if (IsRecording())
+			{
+				if (!RecordingAnimation.IsPlaying())
+				{
+					RecordingAnimation.Play(Animation, true);
+				}
+
+				const FLinearColor Color = bRecordingButtonHovered ? FLinearColor::Red : FLinearColor::White;
+				return FSlateColor(bRecordingButtonHovered ? Color : Color.CopyWithNewOpacity(0.2f + 0.8f * RecordingAnimation.GetLerp()));
+			}
+
+			RecordingAnimation.Pause();
+			return FSlateColor::UseSubduedForeground();
+		}));
+
+	return SNew(SButton)
+		.OnClicked(FOnClicked::CreateRaw(this, &FRecordingControls::ToggleRecordingState))
+		.ForegroundColor(FSlateColor::UseForeground())
+		.IsFocusable(false)
+		.IsEnabled_Raw(this, &FRecordingControls::IsRecordingToggleButtonEnabled)
+		.VAlign(VAlign_Center)
+		.HAlign(HAlign_Center)
+		.OnHovered_Lambda([this]() { bRecordingButtonHovered = true; })
+		.OnUnhovered_Lambda([this]() { bRecordingButtonHovered = false; })
+		.ToolTipText_Lambda([this, StartRecordingTooltip]()
+			{
+				return IsRecording() ? LOCTEXT("StopRecordButtonDesc", "Stop the current recording ") : StartRecordingTooltip;
+			})
+		[
+			SNew(SHorizontalBox)
+			+ SHorizontalBox::Slot()
+			.Padding(FMargin(0, 0, 0, 0))
+			.VAlign(VAlign_Center)
+			.AutoWidth()
+			[
+				Animation
+			]
+			+ SHorizontalBox::Slot()
+			.Padding(FMargin(4, 0, 0, 0))
+			.VAlign(VAlign_Center)
+			.AutoWidth()
+			[
+				SNew(STextBlock)
+				.Visibility_Lambda([this]() {return IsRecording() ? EVisibility::Collapsed : EVisibility::Visible; })
+				.TextStyle(FAppStyle::Get(), "SmallButtonText")
+				.Text(LOCTEXT("StartRecordingButtonLabel", "Start Recording"))
+			]
+		];
+}
+
+FText FRecordingControls::GetCurrentSelectedSessionName() const
+{
+	if (const TSharedPtr<FSessionInfo> CurrentSessionPtr = GetCurrentSessionInfo())
+	{
+		return FText::AsCultureInvariant(CurrentSessionPtr->SessionName);
+	}
+
+	static FText InvalidSessionName = LOCTEXT("InvalidSessionLabel", "Select a Session");
+
+	return InvalidSessionName;
+}
+
+TSharedRef<SWidget> FRecordingControls::GenerateTargetSessionSelector()
+{
+	return SNew(SHorizontalBox)
+		+SHorizontalBox::Slot()
+		.AutoWidth()
+		[
+			SNew(SComboButton)
+			.MenuPlacement(MenuPlacement_AboveAnchor).ComboButtonStyle(&FAppStyle::Get().GetWidgetStyle<FComboButtonStyle>("ComboButton"))
+			.OnGetMenuContent(this, &FRecordingControls::GenerateTargetSessionDropdown)
+			.HasDownArrow(true)
+			.ButtonContent()
+			[
+				SNew(STextBlock)
+				.Text_Raw(this, &FRecordingControls::GetCurrentSelectedSessionName)
+				.TextStyle(&FAppStyle::Get().GetWidgetStyle<FTextBlockStyle>("NormalText"))
+			]];
+}
+
+TSharedRef<SWidget> FRecordingControls::GenerateTargetSessionDropdown()
+{
+	using namespace UE::TraceBasedDebuggers;
+
+	FMenuBuilder MenuBuilder(true, nullptr);
+
+	MenuBuilder.BeginSection("RecordingWidgetTargets", LOCTEXT("RecordingTargetsMenu", "Available Targets"));
+	{
+		if (const TSharedPtr<FRemoteSessionsManager> RemoteSessionManager = RemoteSessionManagerWeakPtr.Pin())
+		{
+			RemoteSessionManager->EnumerateActiveSessions([this, &MenuBuilder](const TSharedRef<FSessionInfo>& InSessionInfoRef)
+				{
+					if (EnumHasAnyFlags(InSessionInfoRef->GetSessionTypeAttributes(), ERemoteSessionAttributes::IsMultiSessionWrapper))
+					{
+						return true;
+					}
+
+					const FText SessionNameAsText = FText::AsCultureInvariant(InSessionInfoRef->SessionName);
+					MenuBuilder.AddMenuEntry(
+						SessionNameAsText,
+						FText::Format(LOCTEXT("SingleTargetItemTooltip", "Select {0} session as current target"), SessionNameAsText),
+						GetIconForSession(InSessionInfoRef->InstanceId),
+						FUIAction(FExecuteAction::CreateSP(this, &FRecordingControls::SelectTargetSession, InSessionInfoRef->InstanceId), EUIActionRepeatMode::RepeatDisabled));
+					return true;
+				});
+		}
+	}
+
+	MenuBuilder.EndSection();
+
+	OnGenerateTargetSessionSelectorInternal(MenuBuilder);
+
+	return MenuBuilder.MakeWidget();
+}
+
+TSharedRef<SWidget> FRecordingControls::GenerateRecordingTimeTextBlock()
+{
+	return SNew(SBox)
+		.VAlign(VAlign_Center)
+		.Visibility_Raw(this, &FRecordingControls::GetRecordingTimeTextBlockVisibility)
+		.Padding(12.0f, 0.0f, 0.0f, 0.0f)
+		[
+			SNew(STextBlock)
+			.TextStyle(FAppStyle::Get(), "SmallButtonText")
+			.Text_Raw(this, &FRecordingControls::GetRecordingTimeText)
+			.ColorAndOpacity(FColor::White)
+		];
+}
+
+EVisibility FRecordingControls::GetRecordingTimeTextBlockVisibility() const
+{
+	using namespace UE::TraceBasedDebuggers;
+	const TSharedPtr<FSessionInfo> SessionInfo = GetCurrentSessionInfo();
+	const bool bIsVisible = SessionInfo && !EnumHasAnyFlags(SessionInfo->GetSessionTypeAttributes(), ERemoteSessionAttributes::IsMultiSessionWrapper) && IsRecording();
+
+	return bIsVisible ? EVisibility::Visible : EVisibility::Collapsed;
+}
+
+void FRecordingControls::SelectTargetSession(const FGuid SessionId)
+{
+	CurrentSelectedSessionId = SessionId;
+}
+
+FSlateIcon FRecordingControls::GetIconForSession(const FGuid SessionId) const
+{
+	if (const TSharedPtr<FSessionInfo> SessionInfo = GetSessionInfo(SessionId))
+	{
+		return SessionInfo->IsRecording(DebuggerGuid)
+			? RecordingSessionIcon
+			: FSlateIcon();
+	}
+
+	return FSlateIcon();
+}
+
+TSharedPtr<FSessionInfo> FRecordingControls::GetCurrentSessionInfo() const
+{
+	return GetSessionInfo(CurrentSelectedSessionId);
+}
+
+TSharedPtr<FSessionInfo> FRecordingControls::GetSessionInfo(const FGuid Id) const
+{
+	using namespace UE::TraceBasedDebuggers;
+	const TSharedPtr<FRemoteSessionsManager> RemoteSessionManager = RemoteSessionManagerWeakPtr.Pin();
+	TSharedPtr<FSessionInfo> SessionInfo = RemoteSessionManager ? RemoteSessionManager->GetSessionInfo(Id).Pin() : nullptr;
+
+	return SessionInfo;
+}
+
+const FSlateBrush* FRecordingControls::GetRecordOrStopButton() const
+{
+	if (bRecordingButtonHovered && IsRecording())
+	{
+		return StopRecordingBrush;
+	}
+
+	return StartRecordingBrush;
+}
+
+void FRecordingControls::HandleRecordingStop(const TSharedPtr<FSessionInfo>& InSessionInfo)
+{
+	if (DebuggerGuid != InSessionInfo->LastKnownRecordingState.DebuggerId)
+	{
+		// Recording was stopped by another debugger, ignore it
+		return;
+	}
+
+	if (UStatusBarSubsystem* StatusBarSubsystem = GEditor ? GEditor->GetEditorSubsystem<UStatusBarSubsystem>() : nullptr)
+	{
+		StatusBarSubsystem->PopStatusBarMessage(StatusBarID, RecordingMessageHandle);
+
+		const FText LiveSessionEnded = LOCTEXT("RecordingSessionEndedMessage", "Recording session has ended");
+		LiveSessionEndedMessageHandle = StatusBarSubsystem->PushStatusBarMessage(StatusBarID, LiveSessionEnded);
+	}
+}
+
+void FRecordingControls::HandleRecordingStart(const TSharedPtr<FSessionInfo>& InSessionInfo)
+{
+	if (DebuggerGuid != InSessionInfo->LastKnownRecordingState.DebuggerId)
+	{
+		// Recording was started by another debugger, ignore it
+		return;
+	}
+
+	UStatusBarSubsystem* StatusBarSubsystem = GEditor ? GEditor->GetEditorSubsystem<UStatusBarSubsystem>() : nullptr;
+	if (!StatusBarSubsystem)
+	{
+		return;
+	}
+
+	if (RecordingPathMessageHandle.IsValid())
+	{
+		StatusBarSubsystem->PopStatusBarMessage(StatusBarID, RecordingPathMessageHandle);
+		RecordingPathMessageHandle = FStatusBarMessageHandle();
+	}
+
+	if (LiveSessionEndedMessageHandle.IsValid())
+	{
+		StatusBarSubsystem->PopStatusBarMessage(StatusBarID, LiveSessionEndedMessageHandle);
+		LiveSessionEndedMessageHandle = FStatusBarMessageHandle();
+	}
+
+	RecordingMessageHandle = StatusBarSubsystem->PushStatusBarMessage(StatusBarID, LOCTEXT("RecordingMessage", "Recording..."));
+}
+
+void FRecordingControls::ExecuteAsyncConnectionAttemptTaskWithRetry(
+	FGuid SessionID
+	, const int32 RemainingRetries
+	, const TFunction<bool()>& InRecordingStartAttemptCallback
+	, const TFunction<void()>& InRecordingFailedCallback)
+{
+	const TSharedPtr<FAsyncConnectionAttemptTask> ConnectionAttemptRequest = MakeShared<FAsyncConnectionAttemptTask>(
+		SessionID
+		, LogCategory
+		, InRecordingStartAttemptCallback
+		, InRecordingFailedCallback
+		, StaticCastWeakPtr<FRecordingControls>(AsWeak()));
+
+	ConnectionAttemptRequest->Start(RemainingRetries, IntervalBetweenAutoplayConnectionAttemptsSeconds);
+}
+
+void FRecordingControls::FAsyncConnectionAttemptTask::HandleSuccess()
+{
+	State = EState::Success;
+
+	if (const TSharedPtr<FRecordingControls> OwnerPtr = Owner.Pin())
+	{
+		OwnerPtr->HandleConnectionAttemptResult(SessionID, ELiveConnectionAttemptResult::Success, ProgressNotification);
+	}
+}
+
+void FRecordingControls::FAsyncConnectionAttemptTask::HandleFailure()
+{
+	State = EState::Failed;
+
+	if (RecordingFailedCallback)
+	{
+		RecordingFailedCallback();
+	}
+
+	if (const TSharedPtr<FRecordingControls> OwnerPtr = Owner.Pin())
+	{
+		OwnerPtr->HandleConnectionAttemptResult(SessionID, ELiveConnectionAttemptResult::Failed, ProgressNotification);
+	}
+}
+
+bool FRecordingControls::FAsyncConnectionAttemptTask::CanExecute() const
+{
+	// The only two states where Execute() can be run are Not Started and InProgress
+	// In Progress is allowed because we might need to execute again as part of the retry process
+	return State == EState::NotStarted || State == EState::InProgress;
+}
+
+void FRecordingControls::CancelInFlightConnectionAttempt(const TSharedPtr<FSessionInfo>& SessionInfo)
+{
+	const FGuid SessionID = SessionInfo ? SessionInfo->InstanceId : FGuid{};
+	if (const TSharedPtr<FAsyncConnectionAttemptTask>* ExistingRequestInFlightPtrPtr = InFlightAsyncConnectionRequests.Find(SessionID))
+	{
+		if (const TSharedPtr<FAsyncConnectionAttemptTask> ExistingRequestInFlightPtr = *ExistingRequestInFlightPtrPtr)
+		{
+			ExistingRequestInFlightPtr->Cancel();
+		}
+		else
+		{
+			UE_LOG_REF(LogCategory, Verbose, TEXT("[%hs] Attempted to cancel a connection of a Session that is no linger valid | Session ID [%s] ."), __func__, *SessionID.ToString());
+		}
+
+		return;
+	}
+
+	UE_LOG_REF(LogCategory, Verbose, TEXT("[%hs] Session ID [%s] not found."), __func__, *SessionID.ToString());
+}
+
+void FRecordingControls::HandleRemoteSessionsUpdated()
+{
+	using namespace UE::TraceBasedDebuggers;
+	const TSharedPtr<FSessionInfo> CurrentSessionPtr = GetCurrentSessionInfo();
+	if (!CurrentSessionPtr)
+	{
+		// if the session that is currently selected is not valid, try to find a valid session and select that as default target session
+		if (const TSharedPtr<FRemoteSessionsManager> RemoteSessionManager = RemoteSessionManagerWeakPtr.Pin())
+		{
+			if (const TSharedPtr<FSessionInfo> FirstValidSession = RemoteSessionManager->GetFirstAvailableSingleSessionInfo().Pin())
+			{
+				SelectTargetSession(FirstValidSession->InstanceId);
+				return;
+			}
+		}
+
+		SelectTargetSession(FRemoteSessionsManager::InvalidSessionGUID);
+	}
+}
+
+FRecordingControls::FAsyncConnectionAttemptTask::~FAsyncConnectionAttemptTask()
+{
+	if (RetryDelegateHandle.IsValid())
+	{
+		FTSTicker::RemoveTicker(RetryDelegateHandle);
+		RetryDelegateHandle = FTSTicker::FDelegateHandle();
+	}
+}
+
+void FRecordingControls::FAsyncConnectionAttemptTask::Cancel()
+{
+	State = EState::Canceled;
+
+	if (RecordingFailedCallback)
+	{
+		RecordingFailedCallback();
+	}
+
+	if (const TSharedPtr<FRecordingControls> OwnerPtr = Owner.Pin())
+	{
+		OwnerPtr->HandleConnectionAttemptResult(SessionID, ELiveConnectionAttemptResult::Canceled, ProgressNotification);
+	}
+}
+
+void FRecordingControls::FAsyncConnectionAttemptTask::ScheduleRetry()
+{
+	if (!CanExecute())
+	{
+		UE_LOG_REF(LogCategory, Error, TEXT("[%hs] This connection attempt is in an invalid state | Current State Value [%u]"), __func__, State);
+		return;
+	}
+
+	if (RetryDelegateHandle.IsValid())
+	{
+		UE_LOG_REF(LogCategory, Warning, TEXT("[%hs] Attempted to schedule a retry when there is one already scheduled"), __func__);
+		return;
+	}
+
+	if (RemainingRetries <= 0)
+	{
+		UE_LOG_REF(LogCategory, Warning, TEXT("[%hs] Failed to connect to live session | attempts exhausted..."), __func__);
+		HandleFailure();
+		return;
+	}
+
+	RemainingRetries--;
+
+	RetryDelegateHandle = FTSTicker::GetCoreTicker().AddTicker(FTickerDelegate::CreateSPLambda(AsShared(), [this](float DeltaTime)
+		{
+			Execute();
+			return false;
+		}), IntervalBetweenAttemptsSeconds);
+}
+
+void FRecordingControls::FAsyncConnectionAttemptTask::Start(const int32 InMaxRetriesAttempts, const float InIntervalBetweenAttemptsSeconds)
+{
+	if (State != EState::NotStarted)
+	{
+		UE_LOG_REF(LogCategory, Error, TEXT("[%hs] This connection attempt is already started or in an invalid state | Current State Value [%u]"), __func__, State);
+		return;
+	}
+
+	const TSharedPtr<FRecordingControls> OwnerPtr = Owner.Pin();
+	if (!OwnerPtr)
+	{
+		UE_LOG_REF(LogCategory, Error, TEXT("[%hs] This connection attempt is does not have a valid owner | Cannot start attempt... "), __func__);
+		return;
+	}
+
+	if (const TSharedPtr<FAsyncConnectionAttemptTask>* ExistingRequestInFlight = OwnerPtr->InFlightAsyncConnectionRequests.Find(SessionID))
+	{
+		UE_LOG_REF(LogCategory, Warning, TEXT("[%hs] There is a connection attempt in fight for this session. attempting to cancel it | Session ID [%s]"), __func__, *SessionID.ToString());
+		(*ExistingRequestInFlight)->Cancel();
+	}
+
+	OwnerPtr->InFlightAsyncConnectionRequests.Add(SessionID, AsShared());
+
+	State = EState::InProgress;
+
+	RemainingRetries = InMaxRetriesAttempts;
+	IntervalBetweenAttemptsSeconds = InIntervalBetweenAttemptsSeconds;
+
+	if (!ProgressNotification)
+	{
+		FNotificationInfo Info(LOCTEXT("ConnectingToSessionMessage", "Connecting Session ..."));
+		Info.bFireAndForget = false;
+		Info.FadeOutDuration = 3.0f;
+		Info.ExpireDuration = 0.0f;
+		Info.ButtonDetails.Add(FNotificationButtonInfo(LOCTEXT("CancelConnectionAttemptButton", "Cancel"), LOCTEXT("CancelConnectionAttemptButtonTip", "Cancels the active connection attempt."),
+			FSimpleDelegate::CreateSP(AsShared(), &FAsyncConnectionAttemptTask::Cancel)));
+
+		ProgressNotification = OwnerPtr->PushConnectionAttemptNotification(Info);
+	}
+
+	// We need to wait at least one tick before attempting to connect to give it time to the trace to be initialized, write to disk, and for the
+	// session manager to hear back from a remote instance
+	ScheduleRetry();
+}
+
+void FRecordingControls::FAsyncConnectionAttemptTask::Execute()
+{
+	if (RetryDelegateHandle.IsValid())
+	{
+		FTSTicker::RemoveTicker(RetryDelegateHandle);
+		RetryDelegateHandle.Reset();
+	}
+
+	if (!CanExecute())
+	{
+		UE_LOG_REF(LogCategory, Error, TEXT("[%hs] This connection attempt is in an invalid state | Current State Value [%u]"), __func__, State);
+		return;
+	}
+
+	if (!RecordingStartAttemptCallback)
+	{
+		UE_LOG_REF(LogCategory, Error, TEXT("[%hs] Failed to connect to session | Invalid callback provided..."), __func__);
+		HandleFailure();
+		return;
+	}
+
+	if (const TSharedPtr<FRecordingControls> OwnerPtr = Owner.Pin())
+	{
+		OwnerPtr->UpdateConnectionAttemptNotification(ProgressNotification, RemainingRetries);
+	}
+
+	// Debugger needs the trace session name to be able to load a live session. Although the session exist, the session name might not be written right away
+	// Trace files don't really have metadata, it is all part of the same stream, so we need to wait until it is written which might take a few ticks.
+	// Therefore, if it is not ready, try again a few times.
+	if (!RecordingStartAttemptCallback())
+	{
+		UE_LOG_REF(LogCategory, Verbose, TEXT("[%hs] Failed to connect to live session | Attempting again in [%f]..."), __func__, IntervalBetweenAttemptsSeconds);
+		ScheduleRetry();
+	}
+	else
+	{
+		HandleSuccess();
+	}
+}
+
+void FRecordingControls::ToggleSingleSessionRecordingState(const TSharedRef<FSessionInfo>& InSessionInfo)
+{
+	SetSessionRecordingState(!InSessionInfo->IsRecording(DebuggerGuid), InSessionInfo);
+}
+
+void FRecordingControls::SetSessionRecordingState(const bool bIsRecording, const TSharedRef<FSessionInfo>& SessionInfoRef)
+{
+	using namespace UE::TraceBasedDebuggers;
+	const TSharedPtr<FRemoteSessionsManager> RemoteSessionManager = RemoteSessionManagerWeakPtr.Pin();
+	if (!RemoteSessionManager)
+	{
+		UE_LOG_REF(LogCategory, Error, TEXT("[%hs] Session Manager is not available"), __func__);
+		return;
+	}
+
+	if (bIsRecording)
+	{
+		const int32 RemainingRetries = GetMaxConnectionRetriesInternal();
+		const ETraceTransportMode DataTransportOverrideMode = GetTransportModeOverrideForTargetTypeInternal(SessionInfoRef->BuildTargetType);
+
+		// Set the session in a busy state so the auto-expire policy for it is relaxed. Starting a recording might cause the remote instance to stall for a few seconds
+		// if the map loaded is too complex
+		SessionInfoRef->ReadyState = ERemoteSessionReadyState::Busy;
+
+		const TSharedPtr<FRemoteSessionsManager> RemoteSessionManagerPtr = RemoteSessionManagerWeakPtr.Pin();
+		const TSharedPtr<FSessionInfo> SessionInfoPtr = RemoteSessionManagerPtr ? RemoteSessionManagerPtr->GetSessionInfo(SessionInfoRef->InstanceId).Pin() : nullptr;
+
+		auto RecordingAttemptFailedCallback =
+			[SessionGUID = SessionInfoRef->InstanceId, SessionInfoPtr, WeakThis = AsWeak()]()
+			{
+				if (!SessionInfoPtr)
+				{
+					return;
+				}
+
+				SessionInfoPtr->ReadyState = ERemoteSessionReadyState::Ready;
+
+				if (const TSharedPtr<FRecordingControls> Controls = WeakThis.Pin())
+				{
+					Controls->CloseSessionByRemoteSessionIDInternal(SessionGUID);
+				}
+			};
+
+		TNotNull<TUniquePtr<FStartRecordingCommandMessage>> RecordingParams = BuildStartRecordingParamsInternal();
+		RecordingParams->RecordingMode = ERecordingMode::Live;
+		RecordingParams->TransportMode = DataTransportOverrideMode != ETraceTransportMode::Invalid ? DataTransportOverrideMode : ETraceTransportMode::Direct;
+
+		check(GLog);
+		bool bOutCanBindAll = false;
+		//TODO: Add a way to specify a local address in case we have multiple adapters?
+		const TSharedRef<FInternetAddr> LocalIP = ISocketSubsystem::Get(PLATFORM_SOCKETSUBSYSTEM)->GetLocalHostAddr(*GLog, bOutCanBindAll);
+
+		constexpr bool bAppendPort = false;
+		RecordingParams->Target = LocalIP->ToString(bAppendPort);
+
+		if (RecordingParams->TransportMode == ETraceTransportMode::Direct)
+		{
+			// Direct Tracing mode requires opening a trace session in the editor side first before starting the recording
+			// so targets have something to connect to (Usually targets connect to a running trace server)
+			if (!ConnectToLiveSession_DirectInternal(SessionInfoRef->InstanceId))
+			{
+				UE_LOG_REF(LogCategory, Error, TEXT("[%hs] Failed to connect to live session using direct trace mode"), __func__);
+				RecordingAttemptFailedCallback();
+				return;
+			}
+
+			// If we get here it means we have at least one valid session
+			const FTraceSessionDescriptor& SessionDesc = GetCurrentSessionDescriptorInternal();
+			RecordingParams->Target.Appendf(TEXT(":%u"), SessionDesc.SessionPort);
+
+			SendStartRecordingCommandInternal(SessionInfoRef->Address, MoveTemp(RecordingParams));
+
+			ExecuteAsyncConnectionAttemptTaskWithRetry(SessionInfoRef->InstanceId, RemainingRetries, [SessionInfoPtr]()
+				{
+					if (!SessionInfoPtr)
+					{
+						return false;
+					}
+
+					// In Direct trace mode we if we get connection details from the remote instance, we can assume the connection succeeded.
+					return SessionInfoPtr->GetTraceConnectionDetails().IsValid();
+
+				}, RecordingAttemptFailedCallback);
+		}
+		else if (RecordingParams->TransportMode == ETraceTransportMode::Relay)
+		{
+			// Relay Tracing mode requires opening a trace session in the editor side first before starting the recording
+			// so targets have something to connect to (Usually targets connect to a running trace server)
+			if (!ConnectToLiveSession_RelayInternal(SessionInfoRef->InstanceId))
+			{
+				UE_LOG_REF(LogCategory, Error, TEXT("[%hs] Failed to connect to live session using relay trace mode"), __func__);
+				RecordingAttemptFailedCallback();
+				return;
+			}
+
+			SendStartRecordingCommandInternal(SessionInfoRef->Address, MoveTemp(RecordingParams));
+
+			// Once the start recording command is issued, we can try to connect to the created session (if everything went well).
+			// If it didn't go well, this will take care of update the UI to notify the user
+			ExecuteAsyncConnectionAttemptTaskWithRetry(SessionInfoRef->InstanceId, RemainingRetries, 
+			[SessionInfoPtr, SessionInstanceId = SessionInfoRef->InstanceId, RelayTraceDataTransportInstance = TraceRelayTransportInstanceWeakPtr.Pin()]
+				{
+					if (!SessionInfoPtr)
+					{
+						return false;
+					}
+
+					const FTraceConnectionDetails& RecordingSessionDetails = SessionInfoPtr->GetTraceConnectionDetails();
+
+					// We didn't receive the connection details yet. Trying again later...
+					if (!RecordingSessionDetails.IsValid())
+					{
+						return false;
+					}
+
+					// Check if there is a connection attempt in progress and its results
+					const EConnectionAttemptResult ConnectionAttemptStatus = RelayTraceDataTransportInstance.IsValid()
+						? RelayTraceDataTransportInstance->GetConnectionAttemptResult(SessionInstanceId)
+						: EConnectionAttemptResult::Failed;
+
+					switch (ConnectionAttemptStatus)
+					{
+					case EConnectionAttemptResult::InProgress:
+					{
+						// A connection attempt started, but we don't have the results yet. Trying again later...
+						return false;
+					}
+					case EConnectionAttemptResult::NotStarted:
+					{
+						FRelayConnectionInfo ConnectionInfo;
+						ConnectionInfo.Port = RecordingSessionDetails.Port;
+						ConnectionInfo.Address = RecordingSessionDetails.TraceTarget;
+						ConnectionInfo.CertificateAuthority = RecordingSessionDetails.CertAuth;
+						EConnectionAttemptResult Result = RelayTraceDataTransportInstance.IsValid()
+							? RelayTraceDataTransportInstance->ConnectToRelay(SessionInstanceId, ConnectionInfo)
+							: EConnectionAttemptResult::Failed;
+
+						return Result == EConnectionAttemptResult::Success;
+					}
+					case EConnectionAttemptResult::Success:
+						// The current connection attempt succeeded, we can stop the async retry loop
+						return true;
+					case EConnectionAttemptResult::Failed:
+					default:
+						// There was a connection attempt, but it failed. Abandoning the retry attempt as they will not succeed
+						return false;
+					}
+				}, RecordingAttemptFailedCallback);
+		}
+	}
+	else
+	{
+		CancelInFlightConnectionAttempt(SessionInfoRef.ToSharedPtr());
+		SendStopRecordingCommandInternal(SessionInfoRef->Address);
+	}
+}
+
+FReply FRecordingControls::ToggleRecordingState()
+{
+	using namespace UE::TraceBasedDebuggers;
+	const TSharedPtr<FSessionInfo> SessionInfoPtr = GetCurrentSessionInfo();
+	if (!SessionInfoPtr)
+	{
+		return FReply::Handled();
+	}
+
+	OnToggleRecordingStateInternal(SessionInfoPtr.ToSharedRef());
+
+	return FReply::Handled();
+}
+
+bool FRecordingControls::IsRecordingToggleButtonEnabled() const
+{
+	using namespace UE::TraceBasedDebuggers;
+	if (CurrentSelectedSessionId == FRemoteSessionsManager::InvalidSessionGUID)
+	{
+		return false;
+	}
+
+	const TSharedPtr<FSessionInfo> SessionInfo = GetCurrentSessionInfo();
+	if (!SessionInfo)
+	{
+		return false;
+	}
+
+	// Do not allow recording if another trace-based debuggers is already recording in the current session
+	if (SessionInfo->LastKnownRecordingState.DebuggerId.IsValid()
+		&& SessionInfo->LastKnownRecordingState.DebuggerId != DebuggerGuid
+		&& SessionInfo->LastKnownRecordingState.bIsRecording)
+	{
+		return false;
+	}
+
+	// Allow derived classes to deny
+	if (!IsRecordingToggleButtonEnabledInternal(SessionInfo.ToSharedRef()))
+	{
+		return false;
+	}
+
+	return true;
+}
+
+EVisibility FRecordingControls::IsRecordingToggleButtonVisible()
+{
+	return EVisibility::Visible;
+}
+
+void FRecordingControls::AddToMenu(const FName ExistingMenu)
+{
+	UToolMenu* ToolBar = UToolMenus::Get()->ExtendMenu(ExistingMenu);
+	if (!ensure(ToolBar))
+	{
+		return;
+	}
+
+	FToolMenuSection& Section = ToolBar->AddSection("RecordingControls");
+
+	Section.AddDynamicEntry("Controls", FNewToolMenuSectionDelegate::CreateSPLambda(this, [RecordingControls = this](FToolMenuSection& InSection)
+		{
+			const TSharedRef<SWidget> StartRecordingButton = SNew(SBox).Padding(4.0f, 0.0f)
+				[
+					RecordingControls->GenerateToggleRecordingStateButton(LOCTEXT("RecordButtonDesc", "Starts a recording for the current session, saving it directly to file"))
+				];
+			const TSharedRef<SWidget> SessionsDropdown = SNew(SBox).Padding(4.0f, 0.0f)
+				[
+					RecordingControls->GenerateTargetSessionSelector()
+				];
+
+			const TSharedRef<SWidget> RecordingTime = RecordingControls->GenerateRecordingTimeTextBlock();
+
+			InSection.AddSeparator("RecordingControlsDivider");
+
+			InSection.AddEntry(
+				FToolMenuEntry::InitWidget(
+					"AvailableSessions",
+					SessionsDropdown,
+					FText::GetEmpty(),
+					false,
+					false
+				));
+
+			InSection.AddEntry(
+				FToolMenuEntry::InitWidget(
+					"StartRecording",
+					StartRecordingButton,
+					FText::GetEmpty(),
+					true,
+					false
+				));
+
+			InSection.AddEntry(
+				FToolMenuEntry::InitWidget(
+					"RecordingTime",
+					RecordingTime,
+					FText::GetEmpty(),
+					false,
+					false
+				));
+
+			RecordingControls->AddToMenuInternal(InSection);
+		}));
+}
+
+int32 FRecordingControls::GetMaxConnectionRetriesInternal()
+{
+	return DefaultNumberOfConnectionRetries;
+}
+
+ETraceTransportMode FRecordingControls::GetTransportModeOverrideForTargetTypeInternal(EBuildTargetType InBuildTarget)
+{
+	return ETraceTransportMode::Invalid;
+}
+
+bool FRecordingControls::IsRecordingToggleButtonEnabledInternal(const TSharedRef<FSessionInfo>& InSessionInfo) const
+{
+	return true;
+}
+
+bool FRecordingControls::IsRecording() const
+{
+	using namespace UE::TraceBasedDebuggers;
+	const TSharedPtr<FSessionInfo> SessionInfo = GetCurrentSessionInfo();
+	return SessionInfo ? SessionInfo->IsRecording(DebuggerGuid) : false;
+}
+
+FText FRecordingControls::GetRecordingTimeText() const
+{
+	using namespace UE::TraceBasedDebuggers;
+	if (const TSharedPtr<FSessionInfo> SessionInfo = GetCurrentSessionInfo())
+	{
+		FNumberFormattingOptions FormatOptions;
+		FormatOptions.MinimumFractionalDigits = 2;
+		FormatOptions.MaximumFractionalDigits = 2;
+
+		const FText SecondsText = FText::AsNumber(SessionInfo->LastKnownRecordingState.ElapsedTime, &FormatOptions);
+		return FText::Format(LOCTEXT("RecordingTimer", "{0} s"), SecondsText);
+	}
+
+	return LOCTEXT("RecordingTimerError", "Failed to get time information");
+}
+
+TSharedPtr<SNotificationItem> FRecordingControls::PushConnectionAttemptNotification(const FNotificationInfo& InNotificationInfo)
+{
+	TSharedPtr<SNotificationItem> ConnectionAttemptNotification = FSlateNotificationManager::Get().AddNotification(InNotificationInfo);
+
+	if (ConnectionAttemptNotification.IsValid())
+	{
+		ConnectionAttemptNotification->SetCompletionState(SNotificationItem::CS_Pending);
+		return ConnectionAttemptNotification;
+	}
+
+	return nullptr;
+}
+
+void FRecordingControls::UpdateConnectionAttemptNotification(const TSharedPtr<SNotificationItem>& InNotification, int32 AttemptsRemaining)
+{
+	if (InNotification)
+	{
+		InNotification->SetSubText(FText::FormatOrdered(LOCTEXT("SessionConnectionAttemptSubText", "Attempts Remaining {0}"), AttemptsRemaining));
+	}
+}
+
+void FRecordingControls::HandleConnectionAttemptResult(const FGuid SessionGUID, const ELiveConnectionAttemptResult Result, const TSharedPtr<SNotificationItem>& InNotification)
+{
+	InFlightAsyncConnectionRequests.Remove(SessionGUID);
+
+	if (InNotification)
+	{
+		switch (Result)
+		{
+		case ELiveConnectionAttemptResult::Success:
+		{
+			InNotification->SetText(LOCTEXT("SessionConnectionSuccess", "Connected!"));
+			InNotification->SetSubText(FText::GetEmpty());
+			InNotification->SetCompletionState(SNotificationItem::ECompletionState::CS_Success);
+			break;
+		}
+		case ELiveConnectionAttemptResult::Canceled:
+		{
+			InNotification->SetText(LOCTEXT("SessionConnectionCanceledText", "Connection Attempt Canceled"));
+			InNotification->SetCompletionState(SNotificationItem::ECompletionState::CS_None);
+			break;
+		}
+		case ELiveConnectionAttemptResult::Failed:
+		{
+			InNotification->SetText(LOCTEXT("SessionConnectionFailedText", "Failed to connect"));
+			InNotification->SetSubText(LOCTEXT("SessionConnectionFailedSubText", "See the logs for more details..."));
+			InNotification->SetCompletionState(SNotificationItem::ECompletionState::CS_Fail);
+			break;
+		}
+		default:
+		{
+			InNotification->SetText(LOCTEXT("SessionConnectionUnknownText", "Something went wrong..."));
+			InNotification->SetSubText(LOCTEXT("SessionConnectionUnknownTextSubText", "See the logs for more details..."));
+			InNotification->SetCompletionState(SNotificationItem::ECompletionState::CS_Fail);
+			break;
+		}
+		}
+
+		InNotification->ExpireAndFadeout();
+	}
+}
+
+} // UE::TraceBasedDebuggers
+#undef LOCTEXT_NAMESPACE
+
+#endif // WITH_EDITOR
\ No newline at end of file
diff --git a/Engine/Source/Runtime/Experimental/TraceBasedDebuggersAnalysis/Private/Widgets/SConnectionStatusBar.cpp b/Engine/Source/Runtime/Experimental/TraceBasedDebuggersAnalysis/Private/Widgets/SConnectionStatusBar.cpp
new file mode 100644
--- /dev/null
+++ b/Engine/Source/Runtime/Experimental/TraceBasedDebuggersAnalysis/Private/Widgets/SConnectionStatusBar.cpp
@@ -0,0 +1,224 @@
+﻿// Copyright Epic Games, Inc. All Rights Reserved.
+
+#include "Widgets/SConnectionStatusBar.h"
+
+#if WITH_EDITOR
+
+#include "Components/HorizontalBox.h"
+#include "RemoteSessionsManager.h"
+#include "SessionInfo.h"
+#include "Widgets/Images/SImage.h"
+#include "Widgets/SOverlay.h"
+#include "Widgets/Text/STextBlock.h"
+
+#define LOCTEXT_NAMESPACE "ConnectionStatusBar"
+
+namespace UE::TraceBasedDebuggers
+{
+
+SLATE_IMPLEMENT_WIDGET(SConnectionStatusBar)
+void SConnectionStatusBar::PrivateRegisterAttributes(FSlateAttributeInitializer& AttributeInitializer)
+{
+}
+
+void SConnectionStatusBar::Construct(const FArguments& InArgs, const FGuid& InDebuggerGuid, const TSharedRef<FRemoteSessionsManager>& InRemoteSessionManager)
+{
+	DebuggerGuid = InDebuggerGuid;
+	RemoteSessionManagerWeakPtr = InRemoteSessionManager;
+
+	ChildSlot
+	[
+		SNew(SHorizontalBox)
+	
+			+ SHorizontalBox::Slot()
+			.AutoWidth()
+			.VAlign(VAlign_Center)
+			.Padding(0, 0, 3, 0)
+			[
+				SNew(SOverlay)
+				+ SOverlay::Slot()
+				.HAlign(HAlign_Center)
+				.VAlign(VAlign_Top)
+				[
+					SNew(SImage)
+					.ColorAndOpacity(FSlateColor::UseForeground())
+					.Image_Raw(this, &SConnectionStatusBar::GetStateBackgroundIcon)
+				]
+				+ SOverlay::Slot()
+				.HAlign(HAlign_Center)
+				.VAlign(VAlign_Top)
+				[
+					SNew(SImage)
+					.Image(FAppStyle::Get().GetBrush("DerivedData.RemoteCache.Busy"))
+					.ColorAndOpacity_Raw(this, &SConnectionStatusBar::GetWaitingIconColor)
+				]
+				+ SOverlay::Slot()
+				.HAlign(HAlign_Center)
+				.VAlign(VAlign_Top)
+				[
+					SNew(SImage)
+					.Image(FAppStyle::Get().GetBrush("DerivedData.RemoteCache.Uploading"))
+					.ColorAndOpacity_Raw(this, &SConnectionStatusBar::GetUploadingIconColor)
+				]
+				+ SOverlay::Slot()
+				.HAlign(HAlign_Center)
+				.VAlign(VAlign_Top)
+				[
+					SNew(SImage)
+					.Image(FAppStyle::Get().GetBrush("DerivedData.RemoteCache.Downloading"))
+					.ColorAndOpacity_Raw(this, &SConnectionStatusBar::GetDownloadIconColor)
+				]
+				+ SOverlay::Slot()
+				.HAlign(HAlign_Center)
+				.VAlign(VAlign_Top)
+				[
+					SNew(SImage)
+					.Image(FAppStyle::Get().GetBrush("Zen.Status.Idle"))
+					.ColorAndOpacity_Raw(this, &SConnectionStatusBar::GetReadyIconColor)
+				]
+			]
+			+ SHorizontalBox::Slot()
+			.AutoWidth()
+			.VAlign(VAlign_Center)
+			.Padding(0, 0, 10, 0)
+			[
+				SNew(STextBlock)
+				.Text_Raw(this, &SConnectionStatusBar::GetTitleText)
+				.ToolTipText_Raw(this, &SConnectionStatusBar::GetTitleToolTipText)
+			]
+	];
+}
+
+const FSlateBrush* SConnectionStatusBar::GetStateBackgroundIcon() const
+{
+	return State == EState::WaitingForData ? FAppStyle::Get().GetBrush("DerivedData.RemoteCache.BusyBG") : FAppStyle::Get().GetBrush("DerivedData.RemoteCache.IdleBG");
+}
+
+FText SConnectionStatusBar::GetTitleText() const
+{
+	static const FNumberFormattingOptions SizeFormattingOptions = FNumberFormattingOptions().SetMinimumFractionalDigits(2).SetMaximumFractionalDigits(2);
+
+	switch (State)
+	{
+	case EState::Ready:
+		return LOCTEXT("ConnectionReadyStatusText", "Ready");
+	case EState::ReceivingData:
+		{
+			FText BytesPerSecondText = FText::AsMemory(CurrentReceivedDataBytesPerSecond, &SizeFormattingOptions, nullptr, IEC);
+			FText PendingBytesText = FText::AsMemory(CurrentPendingBytes, &SizeFormattingOptions, nullptr, IEC);
+			FText TotalBytesText = FText::AsMemory(CurrentTotalReceivedData, &SizeFormattingOptions, nullptr, IEC);
+			return FText::FormatOrdered(LOCTEXT("ConnectionDownloadedStatusText", "Receiving {0}/s | {1} Pending | Total {2}"), BytesPerSecondText, PendingBytesText, TotalBytesText);
+		}
+	case EState::WaitingForData:
+	default:
+		{
+			FText PendingBytesText = FText::AsMemory(CurrentPendingBytes, &SizeFormattingOptions, nullptr, IEC);
+			FText TotalBytesText = FText::AsMemory(CurrentTotalReceivedData, &SizeFormattingOptions, nullptr, IEC);
+			return FText::FormatOrdered(LOCTEXT("ConnectionWaitingStatusText", "Waiting for Data | {0} Pending | Total {1}"), PendingBytesText, TotalBytesText);
+		}
+	}
+}
+
+FText SConnectionStatusBar::GetTitleToolTipText() const
+{
+	return LOCTEXT("DownloadingStatusToolTipText", "Current state of the connection to the recording target.");
+}
+
+FSlateColor SConnectionStatusBar::GetDownloadIconColor() const
+{
+	return State == EState::ReceivingData ? FLinearColor::White.CopyWithNewOpacity(FMath::MakePulsatingValue(ElapsedTime, 2)) : FLinearColor(0.0,0.0,0.0,0.0);
+}
+
+FSlateColor SConnectionStatusBar::GetUploadingIconColor() const
+{
+	// The debugger does not upload data, but if we are downloading data we want to slightly see the upload arrow
+	return State == EState::ReceivingData ? FLinearColor::White.CopyWithNewOpacity(0.2f) : FLinearColor(0.0,0.0,0.0,0.0);
+}
+
+FSlateColor SConnectionStatusBar::GetWaitingIconColor() const
+{
+	return State == EState::WaitingForData ? FSlateColor::UseForeground() : FLinearColor(0.0,0.0,0.0,0.0);
+}
+
+FSlateColor SConnectionStatusBar::GetReadyIconColor() const
+{
+	return State == EState::Ready ? FSlateColor::UseForeground() : FLinearColor(0.0,0.0,0.0,0.0);
+}
+
+void SConnectionStatusBar::Tick(const FGeometry& AllottedGeometry, const double InCurrentTime, const float InDeltaTime)
+{
+	SCompoundWidget::Tick(AllottedGeometry, InCurrentTime, InDeltaTime);
+
+	ElapsedTime += InDeltaTime;
+
+	UpdateConnectionStatus(InDeltaTime);
+}
+
+void SConnectionStatusBar::UpdateConnectionStatus(double DeltaTime)
+{
+	using namespace UE::TraceBasedDebuggers;
+
+	ElapsedTimeSinceLastUpdate += DeltaTime;
+
+	constexpr double UpdateIntervalSeconds = 1.0;
+	if (ElapsedTimeSinceLastUpdate < UpdateIntervalSeconds)
+	{
+		return;
+	}
+
+	ElapsedTimeSinceLastUpdate = 0;
+
+	const TSharedPtr<FRemoteSessionsManager> RemoteSessionManager = RemoteSessionManagerWeakPtr.Pin();
+	if (!RemoteSessionManager)
+	{
+		State = EState::Ready;
+		return;
+	}
+
+	bool bIsAnyRecording = false;
+	int64 PendingBytesNum = 0;
+	uint64 BytesReadSinceLastMeasurement = 0;
+	RemoteSessionManager->EnumerateActiveSessions([this, &BytesReadSinceLastMeasurement, &bIsAnyRecording, &PendingBytesNum](const TSharedRef<FSessionInfo>& InSessionInfoRef)
+	{
+		if (!EnumHasAnyFlags(InSessionInfoRef->GetSessionTypeAttributes(), ERemoteSessionAttributes::IsMultiSessionWrapper))
+		{
+			if (InSessionInfoRef->IsRecording(DebuggerGuid))
+			{
+				bIsAnyRecording = true;
+				PendingBytesNum += InSessionInfoRef->LastKnownRecordingState.BufferedDataBytesSize;
+				BytesReadSinceLastMeasurement += InSessionInfoRef->GetReceivedBytesPerSecond();
+			}
+		}
+
+		return true;
+	});
+
+	CurrentPendingBytes = PendingBytesNum;
+
+	if (bIsAnyRecording)
+	{
+		if (BytesReadSinceLastMeasurement > 0)
+		{
+			CurrentReceivedDataBytesPerSecond = BytesReadSinceLastMeasurement;
+			CurrentTotalReceivedData += BytesReadSinceLastMeasurement;
+			State = EState::ReceivingData;
+		}
+		else
+		{
+			State = EState::WaitingForData;
+		}
+	}
+	else
+	{
+		State = EState::Ready;
+		CurrentPendingBytes = 0;
+		CurrentReceivedDataBytesPerSecond = 0;
+		CurrentTotalReceivedData = 0;
+	}
+}
+
+}
+
+#undef LOCTEXT_NAMESPACE
+
+#endif // WITH_EDITOR
\ No newline at end of file
diff --git a/Engine/Source/Runtime/Experimental/TraceBasedDebuggersAnalysis/Public/Widgets/RecordingControls.h b/Engine/Source/Runtime/Experimental/TraceBasedDebuggersAnalysis/Public/Widgets/RecordingControls.h
new file mode 100644
--- /dev/null
+++ b/Engine/Source/Runtime/Experimental/TraceBasedDebuggersAnalysis/Public/Widgets/RecordingControls.h
@@ -0,0 +1,228 @@
+// Copyright Epic Games, Inc. All Rights Reserved.
+
+#pragma once
+
+#if WITH_EDITOR
+
+#include "Containers/Ticker.h"
+#include "StatusBarSubsystem.h"
+#include "Textures/SlateIcon.h"
+#include "ToolMenuSection.h"
+#include "TraceDataRelayTransport.h"
+#include "TraceSessionsManager.h"
+#include "RecordingControls.generated.h"
+
+#define UE_API TRACEBASEDDEBUGGERSANALYSIS_API
+
+struct FMessageAddress;
+
+namespace UE::TraceBasedDebuggers
+{
+struct FStartRecordingCommandMessage;
+struct FRemoteSessionsManager;
+struct FSessionInfo;
+
+UENUM()
+enum class ELiveConnectionAttemptResult
+{
+	Success,
+	Failed,
+	Canceled
+};
+
+struct FRecordingControls : TSharedFromThis<FRecordingControls>
+{
+	UE_API FRecordingControls(const FName InStatusBarId
+		, const TSharedRef<FRemoteSessionsManager>& InRemoteSessionManager
+		, const FGuid& InDebuggerId
+		, const FLogCategoryBase& InLogCategory);
+
+	UE_API virtual ~FRecordingControls();
+
+	UE_API void Initialize();
+	UE_API void Deinitialize();
+	UE_API void AddToMenu(FName ExistingMenu);
+
+protected:
+
+	virtual void AddToMenuInternal(FToolMenuSection& InSection)
+	{
+	}
+	virtual void OnGenerateTargetSessionSelectorInternal(FMenuBuilder& InMenuBuilder)
+	{
+	}
+
+	UE_API virtual int32 GetMaxConnectionRetriesInternal();
+	UE_API virtual ETraceTransportMode GetTransportModeOverrideForTargetTypeInternal(EBuildTargetType InBuildTarget);
+
+	virtual bool ConnectToLiveSession_RelayInternal(const FGuid& InSessionId) = 0;
+	virtual bool ConnectToLiveSession_DirectInternal(const FGuid& InSessionId) = 0;
+
+	/**
+	 * Derived classes must override this method to provide the specific message type required to start recording
+	 * That type must inherit from FStartRecordingCommandMessage.
+	 * Once message is ready to be sent the RecordingControls base class will call SendStartRecordingCommandInternal
+	 * so the type can be de-referenced to call the templated function SendCommand of the sessions manager.
+	 * e.g.,
+	 *    TNotNull<TUniquePtr<FStartRecordingCommandMessage>> FMyRecordingControls::BuildStartRecordingParamsInternal()
+	 *    {
+	 *        return MakeUnique<FMyStartRecordingCommandMessage>();
+	 *    }
+	 *    
+	 *    void FMyRecordingControls::SendStartRecordingCommandInternal(const FMessageAddress& InAddress, TUniquePtr<FStartRecordingCommandMessage> InStartRecordingCommand)
+	 *    {
+	 *        if (const TSharedPtr<FRemoteSessionsManager> SessionsManager = RemoteSessionManagerWeakPtr.Pin())
+	 *        {
+	 *            const FMyStartRecordingCommandMessage* StartRecordingCommand = static_cast<FMyStartRecordingCommandMessage*>(InStartRecordingCommand.Get());
+	 *            SessionsManager->SendCommand(InAddress, *StartRecordingCommand);
+	 *        }
+	 *    }
+	 * @see SendStartRecordingCommandInternal
+	 */
+	[[nodiscard]] virtual TNotNull<TUniquePtr<FStartRecordingCommandMessage>> BuildStartRecordingParamsInternal() = 0;
+
+	virtual void SendStartRecordingCommandInternal(const FMessageAddress& InAddress, TUniquePtr<FStartRecordingCommandMessage> InStartRecordingCommand) = 0;
+
+	virtual void SendStopRecordingCommandInternal(const FMessageAddress& InAddress) = 0;
+
+	UE_API virtual bool IsRecordingToggleButtonEnabledInternal(const TSharedRef<FSessionInfo>& InSessionInfo) const;
+
+	virtual void OnToggleRecordingStateInternal(const TSharedRef<FSessionInfo>& InSessionInfo) = 0;
+	virtual void CloseSessionByRemoteSessionIDInternal(const FGuid& InSessionId) = 0;
+
+	virtual FTraceSessionDescriptor& GetCurrentSessionDescriptorInternal() = 0;
+
+	UE_API void SelectTargetSession(FGuid SessionId);
+	UE_API TSharedPtr<FSessionInfo> GetCurrentSessionInfo() const;
+
+	UE_API FReply ToggleRecordingState();
+	UE_API void ToggleSingleSessionRecordingState(const TSharedRef<FSessionInfo>& InSessionInfo);
+
+	UE_API void SetSessionRecordingState(bool bIsRecording, const TSharedRef<FSessionInfo>& SessionInfoRef);
+	UE_API bool IsRecording() const;
+	UE_API bool IsRecordingToggleButtonEnabled() const;
+
+	UE_API FSlateIcon GetIconForSession(FGuid SessionId) const;
+
+	const TWeakPtr<FRemoteSessionsManager> RemoteSessionManagerWeakPtr;
+	TWeakPtr<IDataRelayTransport> TraceRelayTransportInstanceWeakPtr;
+	const FSlateBrush* StartRecordingBrush = nullptr;
+	const FSlateBrush* StopRecordingBrush = nullptr;
+	FSlateIcon RecordingSessionIcon;
+
+private:
+
+	struct FAsyncConnectionAttemptTask : public TSharedFromThis<FAsyncConnectionAttemptTask>
+	{
+		explicit FAsyncConnectionAttemptTask(const FGuid InSessionID
+			, const FLogCategoryBase& LogCategory
+			, const TFunction<bool()>& InRecordingStartAttemptCallback
+			, const TFunction<void()>& InRecordingFailedCallback
+			, const TWeakPtr<FRecordingControls>& RequestOwner
+		)
+			: RecordingStartAttemptCallback(InRecordingStartAttemptCallback),
+			RecordingFailedCallback(InRecordingFailedCallback),
+			Owner(RequestOwner),
+			SessionID(InSessionID),
+			LogCategory(LogCategory)
+		{
+		}
+
+		~FAsyncConnectionAttemptTask();
+
+		enum class EState : uint8
+		{
+			NotStarted,
+			InProgress,
+			Success,
+			Canceled,
+			Failed
+		};
+
+		void Cancel();
+
+		void Start(int32 InMaxRetriesAttempts, float InIntervalBetweenAttemptsSeconds);
+		void ScheduleRetry();
+		void HandleSuccess();
+		void HandleFailure();
+
+	private:
+
+		bool CanExecute() const;
+
+		void Execute();
+
+		TFunction<bool()> RecordingStartAttemptCallback;
+		TSharedPtr<SNotificationItem> ProgressNotification;
+		const TFunction<void()> RecordingFailedCallback;
+		FTSTicker::FDelegateHandle RetryDelegateHandle;
+		TWeakPtr<FRecordingControls> Owner;
+		FGuid SessionID;
+		const FLogCategoryBase& LogCategory;
+		int32 RemainingRetries = 10;
+		float IntervalBetweenAttemptsSeconds = 1.0f;
+		EState State = EState::NotStarted;
+	};
+
+
+	TSharedRef<SWidget> GenerateToggleRecordingStateButton(const FText& StartRecordingTooltip);
+	TSharedRef<SWidget> GenerateTargetSessionSelector();
+	TSharedRef<SWidget> GenerateTargetSessionDropdown();
+	TSharedRef<SWidget> GenerateRecordingTimeTextBlock();
+
+	EVisibility GetRecordingTimeTextBlockVisibility() const;
+
+	TSharedPtr<FSessionInfo> GetSessionInfo(FGuid Id) const;
+
+	void HandleRemoteSessionsUpdated();
+
+	const FSlateBrush* GetRecordOrStopButton() const;
+
+	void HandleRecordingStop(const TSharedPtr<FSessionInfo>& InSessionInfo);
+	void HandleRecordingStart(const TSharedPtr<FSessionInfo>& InSessionInfo);
+
+	void ExecuteAsyncConnectionAttemptTaskWithRetry(FGuid SessionID, int32 RemainingRetries, const TFunction<bool()>& InRecordingStartAttemptCallback, const TFunction<void()>& InRecordingFailedCallback);
+
+	static EVisibility IsRecordingToggleButtonVisible();
+
+	FText GetRecordingTimeText() const;
+
+	static TSharedPtr<SNotificationItem> PushConnectionAttemptNotification(const FNotificationInfo& InNotificationInfo);
+	static void UpdateConnectionAttemptNotification(const TSharedPtr<SNotificationItem>& InNotification, int32 AttemptsRemaining);
+	void HandleConnectionAttemptResult(FGuid SessionGUID, ELiveConnectionAttemptResult Result, const TSharedPtr<SNotificationItem>& InNotification);
+	void CancelInFlightConnectionAttempt(const TSharedPtr<FSessionInfo>& SessionInfo);
+
+	FText GetCurrentSelectedSessionName() const;
+
+	FGuid CurrentSelectedSessionId = FGuid();
+
+	FDelegateHandle RecordingStartedHandle;
+	FDelegateHandle RecordingStoppedHandle;
+
+	const FLogCategoryBase& LogCategory;
+	FGuid DebuggerGuid;
+
+	FCurveSequence RecordingAnimation;
+
+	TMap<FGuid, TSharedPtr<FAsyncConnectionAttemptTask>> InFlightAsyncConnectionRequests;
+
+	const FName StatusBarID;
+
+	FStatusBarMessageHandle RecordingMessageHandle;
+	FStatusBarMessageHandle RecordingPathMessageHandle;
+	FStatusBarMessageHandle LiveSessionEndedMessageHandle;
+
+	float IntervalBetweenAutoplayConnectionAttemptsSeconds = 1.0f;
+
+	bool bInitialized = false;
+
+	bool bRecordingButtonHovered = false;
+
+	static constexpr int32 DefaultNumberOfConnectionRetries = 3;
+};
+
+} // UE::TraceBasedDebuggers
+
+#undef UE_API
+
+#endif // WITH_EDITOR
\ No newline at end of file
diff --git a/Engine/Source/Runtime/Experimental/TraceBasedDebuggersAnalysis/Public/Widgets/SConnectionStatusBar.h b/Engine/Source/Runtime/Experimental/TraceBasedDebuggersAnalysis/Public/Widgets/SConnectionStatusBar.h
new file mode 100644
--- /dev/null
+++ b/Engine/Source/Runtime/Experimental/TraceBasedDebuggersAnalysis/Public/Widgets/SConnectionStatusBar.h
@@ -0,0 +1,67 @@
+﻿// Copyright Epic Games, Inc. All Rights Reserved.
+
+#pragma once
+
+#if WITH_EDITOR
+
+#include "Widgets/SCompoundWidget.h"
+
+#define UE_API TRACEBASEDDEBUGGERSANALYSIS_API
+
+namespace UE::TraceBasedDebuggers
+{
+struct FRemoteSessionsManager;
+
+/**
+ * Status Bar widget that provides information about active trace-based debugger connections
+ */
+class SConnectionStatusBar : public SCompoundWidget
+{
+public:
+	SLATE_DECLARE_WIDGET_API(SConnectionStatusBar, SCompoundWidget, UE_API)
+
+	SLATE_BEGIN_ARGS(SConnectionStatusBar)
+	{
+	}
+	SLATE_END_ARGS()
+
+	UE_API void Construct(const FArguments& InArgs, const FGuid& InDebuggerGuid, const TSharedRef<FRemoteSessionsManager>& InRemoteSessionManager);
+
+private:
+
+	const FSlateBrush* GetStateBackgroundIcon() const;
+	FText GetTitleText() const;
+	FText GetTitleToolTipText() const;
+	FSlateColor GetDownloadIconColor() const;
+	FSlateColor GetUploadingIconColor() const;
+	FSlateColor GetWaitingIconColor() const;
+	FSlateColor GetReadyIconColor() const;
+
+	UE_API virtual void Tick(const FGeometry& AllottedGeometry, const double InCurrentTime, const float InDeltaTime) override;
+
+	void UpdateConnectionStatus(double DeltaTime);
+
+	double ElapsedTime = 0.0;
+	double ElapsedTimeSinceLastUpdate = 0.0;
+
+	enum class EState
+	{
+		Ready,
+		ReceivingData,
+		WaitingForData
+	};
+
+	TWeakPtr<FRemoteSessionsManager> RemoteSessionManagerWeakPtr;
+	FGuid DebuggerGuid;
+	EState State = EState::Ready;
+
+	int64 CurrentPendingBytes = 0;
+	uint64 CurrentReceivedDataBytesPerSecond = 0;
+	uint64 CurrentTotalReceivedData = 0;
+};
+
+}
+
+#undef UE_API
+
+#endif // WITH_EDITOR
\ No newline at end of file
diff --git a/Engine/Source/Runtime/Experimental/TraceBasedDebuggersAnalysis/TraceBasedDebuggersAnalysis.Build.cs b/Engine/Source/Runtime/Experimental/TraceBasedDebuggersAnalysis/TraceBasedDebuggersAnalysis.Build.cs
--- a/Engine/Source/Runtime/Experimental/TraceBasedDebuggersAnalysis/TraceBasedDebuggersAnalysis.Build.cs
+++ b/Engine/Source/Runtime/Experimental/TraceBasedDebuggersAnalysis/TraceBasedDebuggersAnalysis.Build.cs
@@ -10,24 +10,35 @@ namespace UnrealBuildTool.Rules
 	{
 		public TraceBasedDebuggersAnalysis(ReadOnlyTargetRules Target) : base(Target)
 		{
+			bAllowUETypesInNamespaces = true;
+			CppCompileWarningSettings.UnsafeTypeCastWarningLevel = WarningLevel.Error;
+
 			PublicDependencyModuleNames.AddRange(
 				new string[]
 				{
 					"Core",
 					"CoreUObject",
+					"Engine",
 					"TraceBasedDebuggers"
 				}
 				);
 
-			PrivateDependencyModuleNames.AddRange(
-				new string[]
-				{
-					"Slate",
-				}
-				);
-
-			bAllowUETypesInNamespaces = true;
-			CppCompileWarningSettings.UnsafeTypeCastWarningLevel = WarningLevel.Error;
+			if (Target.bBuildEditor || Target.bCompileAgainstEditor)
+			{
+				PrivateDependencyModuleNames.AddRange(
+					new string[]
+					{
+						"EditorSubsystem",
+						"EditorWidgets",
+						"Slate",
+						"SlateCore",
+						"StatusBar",
+						"Sockets",
+						"ToolMenus",
+						"UnrealEd"
+					}
+					);
+			}
 
 			// Allow debugger trace analysis on desktop platforms
 			if (Target.Platform.IsInGroup(UnrealPlatformGroup.Desktop))
