Fix CVD recording channels being overwritten by defaults when clicking Start Recording.

#Synced-CL 49287886

diff --git a/Engine/Plugins/ChaosVD/Source/ChaosVD/Private/Widgets/ChaosVDRecordingControls.cpp b/Engine/Plugins/ChaosVD/Source/ChaosVD/Private/Widgets/ChaosVDRecordingControls.cpp
--- a/Engine/Plugins/ChaosVD/Source/ChaosVD/Private/Widgets/ChaosVDRecordingControls.cpp
+++ b/Engine/Plugins/ChaosVD/Source/ChaosVD/Private/Widgets/ChaosVDRecordingControls.cpp
@@ -4,6 +4,7 @@
 
 #include "Chaos/ChaosVDEngineEditorBridge.h"
 #include "Chaos/ChaosVDRemoteSessionsManager.h"
+#include "ChaosVisualDebugger/ChaosVDOptionalDataChannel.h"
 #include "ChaosVDEngine.h"
 #include "ChaosVDModule.h"
 #include "ChaosVDSettingsManager.h"
@@ -161,7 +162,21 @@ bool FChaosVDRecordingControls::ConnectToLiveSession_DirectInternal(const FGuid&
 
 TNotNull<TUniquePtr<UE::TraceBasedDebuggers::FStartRecordingCommandMessage>> FChaosVDRecordingControls::BuildStartRecordingParamsInternal()
 {
-	return MakeUnique<FChaosVDStartRecordingCommandMessage>();
+	TNotNull<TUniquePtr<FChaosVDStartRecordingCommandMessage>> RecordingParams = MakeUnique<FChaosVDStartRecordingCommandMessage>();
+
+	Chaos::VisualDebugger::FChaosVDDataChannelsManager::Get().EnumerateDataChannels([&RecordingParams](const TWeakPtr< Chaos::VisualDebugger::FChaosVDOptionalDataChannel>& Channel)
+	{
+		if (TSharedPtr< Chaos::VisualDebugger::FChaosVDOptionalDataChannel> LockedChannel = Channel.Pin())
+		{
+			if (LockedChannel->IsChannelEnabled())
+			{
+				RecordingParams->DataChannelsEnabledOverrideList.Add(LockedChannel->GetId().ToString());
+			}
+		}
+		return true;
+	});
+
+	return RecordingParams;
 }
 
 void FChaosVDRecordingControls::SendStartRecordingCommandInternal(const FMessageAddress& InAddress
