[Backout] - CL49344222
#rnx
Original CL Desc
-----------------------------------------------------------------
Fix CVD recording channels being overwritten by defaults when clicking Start Recording.

#Synced-CL 49287886

#Synced-CL 49287886

diff --git a/Engine/Plugins/ChaosVD/Source/ChaosVD/Private/Widgets/ChaosVDRecordingControls.cpp b/Engine/Plugins/ChaosVD/Source/ChaosVD/Private/Widgets/ChaosVDRecordingControls.cpp
--- a/Engine/Plugins/ChaosVD/Source/ChaosVD/Private/Widgets/ChaosVDRecordingControls.cpp
+++ b/Engine/Plugins/ChaosVD/Source/ChaosVD/Private/Widgets/ChaosVDRecordingControls.cpp
@@ -4,7 +4,6 @@
 
 #include "Chaos/ChaosVDEngineEditorBridge.h"
 #include "Chaos/ChaosVDRemoteSessionsManager.h"
-#include "ChaosVisualDebugger/ChaosVDOptionalDataChannel.h"
 #include "ChaosVDEngine.h"
 #include "ChaosVDModule.h"
 #include "ChaosVDSettingsManager.h"
@@ -162,21 +161,7 @@ bool FChaosVDRecordingControls::ConnectToLiveSession_DirectInternal(const FGuid&
 
 TNotNull<TUniquePtr<UE::TraceBasedDebuggers::FStartRecordingCommandMessage>> FChaosVDRecordingControls::BuildStartRecordingParamsInternal()
 {
-	TNotNull<TUniquePtr<FChaosVDStartRecordingCommandMessage>> RecordingParams = MakeUnique<FChaosVDStartRecordingCommandMessage>();
-
-	Chaos::VisualDebugger::FChaosVDDataChannelsManager::Get().EnumerateDataChannels([&RecordingParams](const TWeakPtr< Chaos::VisualDebugger::FChaosVDOptionalDataChannel>& Channel)
-	{
-		if (TSharedPtr< Chaos::VisualDebugger::FChaosVDOptionalDataChannel> LockedChannel = Channel.Pin())
-		{
-			if (LockedChannel->IsChannelEnabled())
-			{
-				RecordingParams->DataChannelsEnabledOverrideList.Add(LockedChannel->GetId().ToString());
-			}
-		}
-		return true;
-	});
-
-	return RecordingParams;
+	return MakeUnique<FChaosVDStartRecordingCommandMessage>();
 }
 
 void FChaosVDRecordingControls::SendStartRecordingCommandInternal(const FMessageAddress& InAddress
