Fix CVD recording channels being overwritten by defaults when clicking Start Recording.

#Synced-CL 49287886

#Synced-CL 49287886

#was-backout 49345502
#Synced-CL 49345530

diff --git a/Engine/Plugins/ChaosVD/Source/ChaosVD/Private/Widgets/ChaosVDRecordingControls.cpp b/Engine/Plugins/ChaosVD/Source/ChaosVD/Private/Widgets/ChaosVDRecordingControls.cpp
--- a/Engine/Plugins/ChaosVD/Source/ChaosVD/Private/Widgets/ChaosVDRecordingControls.cpp
+++ b/Engine/Plugins/ChaosVD/Source/ChaosVD/Private/Widgets/ChaosVDRecordingControls.cpp
@@ -4,6 +4,7 @@
 
 #include "Chaos/ChaosVDEngineEditorBridge.h"
 #include "Chaos/ChaosVDRemoteSessionsManager.h"
+#include "ChaosVisualDebugger/ChaosVDOptionalDataChannel.h"
 #include "ChaosVDEngine.h"
 #include "ChaosVDModule.h"
 #include "ChaosVDSettingsManager.h"
@@ -161,7 +162,23 @@ bool FChaosVDRecordingControls::ConnectToLiveSession_DirectInternal(const FGuid&
 
 TNotNull<TUniquePtr<UE::TraceBasedDebuggers::FStartRecordingCommandMessage>> FChaosVDRecordingControls::BuildStartRecordingParamsInternal()
 {
-	return MakeUnique<FChaosVDStartRecordingCommandMessage>();
+	TNotNull<TUniquePtr<FChaosVDStartRecordingCommandMessage>> RecordingParams = MakeUnique<FChaosVDStartRecordingCommandMessage>();
+
+#if WITH_CHAOS_VISUAL_DEBUGGER
+	Chaos::VisualDebugger::FChaosVDDataChannelsManager::Get().EnumerateDataChannels([&RecordingParams](const TWeakPtr< Chaos::VisualDebugger::FChaosVDOptionalDataChannel>& Channel)
+	{
+		if (TSharedPtr< Chaos::VisualDebugger::FChaosVDOptionalDataChannel> LockedChannel = Channel.Pin())
+		{
+			if (LockedChannel->IsChannelEnabled())
+			{
+				RecordingParams->DataChannelsEnabledOverrideList.Add(LockedChannel->GetId().ToString());
+			}
+		}
+		return true;
+	});
+#endif
+
+	return RecordingParams;
 }
 
 void FChaosVDRecordingControls::SendStartRecordingCommandInternal(const FMessageAddress& InAddress
