Chaos Visual Debugger: Adding the ability to group data structures under named headers for the multiple view panel

#Synced-CL 51470979

diff --git a/Engine/Plugins/ChaosVD/Source/ChaosVD/Private/DetailsCustomizations/ChaosVDSelectionMultipleViewCustomization.cpp b/Engine/Plugins/ChaosVD/Source/ChaosVD/Private/DetailsCustomizations/ChaosVDSelectionMultipleViewCustomization.cpp
--- a/Engine/Plugins/ChaosVD/Source/ChaosVD/Private/DetailsCustomizations/ChaosVDSelectionMultipleViewCustomization.cpp
+++ b/Engine/Plugins/ChaosVD/Source/ChaosVD/Private/DetailsCustomizations/ChaosVDSelectionMultipleViewCustomization.cpp
@@ -36,19 +36,22 @@ void FChaosVDSelectionMultipleViewCustomization::CustomizeDetails(IDetailLayoutB
 	}
 
 	FChaosVDSelectionMultipleView* StructView = reinterpret_cast<FChaosVDSelectionMultipleView*>(SelectionViewStructOnScope->GetStructMemory());
-	
-	IDetailCategoryBuilder& CategoryBuilder = DetailBuilder.EditCategory("Recorded Data");
-	
-	for (const TSharedPtr<FStructOnScope>& DataInstance : StructView->DataInstances)
+
+	static const FName DefaultCategoryName("Recorded Data");
+
+	for (const FChaosVDSelectionMultipleView::FDataEntry& Entry : StructView->DataInstances)
 	{
-		if (!DataInstance || !DataInstance->IsValid())
+		if (!Entry.StructOnScope || !Entry.StructOnScope->IsValid())
 		{
 			continue;
 		}
 
-		if (IDetailPropertyRow* CreatedRow = CategoryBuilder.AddExternalStructure(DataInstance))
+		const FName CategoryName = Entry.GroupName.IsNone() ? DefaultCategoryName : Entry.GroupName;
+		IDetailCategoryBuilder& CategoryBuilder = DetailBuilder.EditCategory(CategoryName);
+
+		if (IDetailPropertyRow* CreatedRow = CategoryBuilder.AddExternalStructure(Entry.StructOnScope))
 		{
-			CreatedRow->DisplayName(DataInstance->GetStructPtr()->GetDisplayNameText());
+			CreatedRow->DisplayName(Entry.StructOnScope->GetStructPtr()->GetDisplayNameText());
 			CreatedRow->ShouldAutoExpand(true);
 		}
 	}
diff --git a/Engine/Plugins/ChaosVD/Source/ChaosVD/Public/ChaosVDSolverDataSelection.h b/Engine/Plugins/ChaosVD/Source/ChaosVD/Public/ChaosVDSolverDataSelection.h
--- a/Engine/Plugins/ChaosVD/Source/ChaosVD/Public/ChaosVDSolverDataSelection.h
+++ b/Engine/Plugins/ChaosVD/Source/ChaosVD/Public/ChaosVDSolverDataSelection.h
@@ -15,19 +15,27 @@ struct FChaosVDSelectionContext
 };
 
 /** Struct used to create a combined view of multiple structs to be used in a vanilla details panel.
- * This struct type has a customization that will show each data entry as an individual property
+ * This struct type has a customization that will show each data entry as an individual property.
+ * Entries may optionally carry a group name; entries sharing the same name are displayed under
+ * a single collapsible category with that name. Entries with no group go into "Recorded Data".
  */
 USTRUCT()
 struct FChaosVDSelectionMultipleView
 {
 	GENERATED_BODY()
 
+	struct FDataEntry
+	{
+		TSharedPtr<FStructOnScope> StructOnScope;
+		FName GroupName = NAME_None;
+	};
+
 	template<typename StructType>
-	void AddData(StructType* Struct);
-	
-	void AddData(TSharedPtr<FStructOnScope> StructOnScope)
+	void AddData(StructType* Struct, FName GroupName = NAME_None);
+
+	void AddData(TSharedPtr<FStructOnScope> StructOnScope, FName GroupName = NAME_None)
 	{
-		DataInstances.Add(StructOnScope);
+		DataInstances.Add({ StructOnScope, GroupName });
 	}
 
 	void Clear()
@@ -36,20 +44,20 @@ struct FChaosVDSelectionMultipleView
 	}
 
 protected:
-	TArray<TSharedPtr<FStructOnScope>> DataInstances;
+	TArray<FDataEntry> DataInstances;
 
 	friend class FChaosVDSelectionMultipleViewCustomization;
 };
 
 template <typename StructType>
-void FChaosVDSelectionMultipleView::AddData(StructType* Struct)
+void FChaosVDSelectionMultipleView::AddData(StructType* Struct, FName GroupName)
 {
 	if (!Struct)
 	{
 		return;
 	}
 
-	DataInstances.Emplace(MakeShared<FStructOnScope>(StructType::StaticStruct(), reinterpret_cast<uint8*>(Struct)));
+	DataInstances.Add({ MakeShared<FStructOnScope>(StructType::StaticStruct(), reinterpret_cast<uint8*>(Struct)), GroupName });
 }
 
 /**
