ChaosVD plugin: particle extra data load-time processing layer

Data layer that deserializes raw trace bytes into in-memory structures at
recording load time and exposes them for display components:

- ChaosVDParticleExtraDataContainer.h: USTRUCT container stored in
FChaosVDSolverFrameData::CustomData; maps SolverID -> ParticleID ->
FChaosVDParticleExtraData; carries the session name table and the set of
natively-serialized struct type paths populated at load time
- ChaosVDDataProcessorBase.h: new virtuals ShouldAlwaysSkip(),
IsBackwardsCompatible(), and GetPostLoadNativeTypesWithChannels(TMap<FName,FName>&)
-- maps each natively-serialized type path to the CVD optional data channel
it was traced through (NAME_None if not specified)
- ChaosVDParticleExtraDataProcessor.h/.cpp: FChaosVDParticleExtraDataProcessor --
deserializes FChaosVDParticleExtraData trace entries into the container;
implements ShouldAlwaysSkip/IsBackwardsCompatible/GetPostLoadNativeTypesWithChannels
based on EChaosVDParticleExtraDataLoadingMode; caches the mode from settings and
updates it via OnSettingsChanged (avoids per-call settings lookup in ProcessRawData);
filters NativeSerialization entries in-processor when mode is SerializeBinOnly
- ChaosVDParticleExtraDataExtension.h/.cpp + ChaosVDBuiltInExtensions.cpp:
registers the processor via the extensions system

Depends on: [CL A - Chaos runtime]

#Synced-CL 52152484

diff --git a/Engine/Plugins/ChaosVD/Source/ChaosVD/Private/Trace/DataProcessors/ChaosVDParticleExtraDataProcessor.cpp b/Engine/Plugins/ChaosVD/Source/ChaosVD/Private/Trace/DataProcessors/ChaosVDParticleExtraDataProcessor.cpp
new file mode 100644
--- /dev/null
+++ b/Engine/Plugins/ChaosVD/Source/ChaosVD/Private/Trace/DataProcessors/ChaosVDParticleExtraDataProcessor.cpp
@@ -0,0 +1,119 @@
+// Copyright Epic Games, Inc. All Rights Reserved.
+
+#include "Trace/DataProcessors/ChaosVDParticleExtraDataProcessor.h"
+
+#include "ChaosVDModule.h"
+#include "ChaosVisualDebugger/ChaosVDParticleExtraData.h"
+#include "DataWrappers/ChaosVDParticleExtraDataContainer.h"
+#include "Trace/DataProcessors/ChaosVDDataProcessorBase.h"
+#include "Trace/ChaosVDTraceProvider.h"
+#include "ChaosVDSettingsManager.h"
+#include "Settings/ChaosVDGeneralSettings.h"
+
+using namespace Chaos::VisualDebugger;
+
+FChaosVDParticleExtraDataProcessor::FChaosVDParticleExtraDataProcessor()
+	: FChaosVDDataProcessorBase(FChaosVDParticleExtraData::WrapperTypeName)
+	, CachedLoadingMode(EChaosVDParticleExtraDataLoadingMode::SerializeBinOnly)
+{
+	// Raw struct serialisation with no explicit versioning  - not safe to read across binary-incompatible builds.
+	bBackwardsCompatible = false;
+
+	if (UChaosVDGeneralSettings* Settings = FChaosVDSettingsManager::Get().GetSettingsObject<UChaosVDGeneralSettings>())
+	{
+		CachedLoadingMode = Settings->ParticleExtraDataLoadingMode;
+		Settings->OnSettingsChanged().AddRaw(this, &FChaosVDParticleExtraDataProcessor::HandleSettingsChanged);
+	}
+}
+
+FChaosVDParticleExtraDataProcessor::~FChaosVDParticleExtraDataProcessor()
+{
+	if (UChaosVDGeneralSettings* Settings = FChaosVDSettingsManager::Get().GetSettingsObject<UChaosVDGeneralSettings>())
+	{
+		Settings->OnSettingsChanged().RemoveAll(this);
+	}
+}
+
+void FChaosVDParticleExtraDataProcessor::HandleSettingsChanged(UObject* InSettings)
+{
+	if (const UChaosVDGeneralSettings* Settings = Cast<UChaosVDGeneralSettings>(InSettings))
+	{
+		CachedLoadingMode = Settings->ParticleExtraDataLoadingMode;
+	}
+}
+
+bool FChaosVDParticleExtraDataProcessor::ShouldAlwaysSkip() const
+{
+	return CachedLoadingMode == EChaosVDParticleExtraDataLoadingMode::SkipAll;
+}
+
+bool FChaosVDParticleExtraDataProcessor::IsBackwardsCompatible() const
+{
+	// SerializeBinOnly strips NativeSerialization entries internally -- no crash risk.
+	// Tell the provider we're "safe" so it won't add us to UnsafeTypesProcessed.
+	// LoadAll returns false so the unsafe-data warning is shown.
+	return CachedLoadingMode == EChaosVDParticleExtraDataLoadingMode::SerializeBinOnly;
+}
+
+void FChaosVDParticleExtraDataProcessor::GetPostLoadNativeTypesWithChannels(TMap<FName, FName>& OutTypeToChannel) const
+{
+	OutTypeToChannel.Append(NativeTypesToChannel);
+}
+
+bool FChaosVDParticleExtraDataProcessor::ProcessRawData(const TArray<uint8>& InData)
+{
+	FChaosVDDataProcessorBase::ProcessRawData(InData);
+
+	const TSharedPtr<FChaosVDTraceProvider> ProviderSharedPtr = TraceProvider.Pin();
+	if (!ensure(ProviderSharedPtr.IsValid()))
+	{
+		return false;
+	}
+
+	const TSharedPtr<FChaosVDParticleExtraData> ExtraData = MakeShared<FChaosVDParticleExtraData>();
+	const bool bSuccess = ReadDataFromBuffer(InData, *ExtraData, ProviderSharedPtr.ToSharedRef());
+
+	if (bSuccess && ExtraData->HasValidData())
+	{
+		if (FChaosVDSolverFrameData* SolverFrame = ProviderSharedPtr->GetCurrentSolverFrame(ExtraData->SolverID))
+		{
+			if (TSharedPtr<FChaosVDParticleExtraDataContainer> Container =
+				SolverFrame->GetCustomData().GetOrAddDefaultData<FChaosVDParticleExtraDataContainer>())
+			{
+				// Cache the recording's name table in the container so the display component
+				// can deserialize inner struct FNames without a separate route to the provider.
+				if (!Container->NameTable)
+				{
+					Container->NameTable = ProviderSharedPtr->GetNameTableInstance();
+				}
+
+				Container->DataBySolverAndParticleID.FindOrAdd(ExtraData->SolverID).Add(ExtraData->ParticleID, ExtraData);
+
+				const bool bFilterNative = CachedLoadingMode == EChaosVDParticleExtraDataLoadingMode::SerializeBinOnly;
+
+				for (FChaosVDExtraDataCategory& Category : ExtraData->Categories)
+				{
+					// Register native entries first (for the post-load dialog), then strip them if filtering.
+					for (const FChaosVDExtraDataStructEntry& Entry : Category.Entries)
+					{
+						if (Entry.SerializationMode == EChaosVDExtraDataSerializationMode::NativeSerialization)
+						{
+							NativeTypesToChannel.FindOrAdd(Entry.StructTypePath) = Category.SourceChannelId;
+							Container->NativeSerializedStructTypes.Add(Entry.StructTypePath);
+						}
+					}
+
+					if (bFilterNative)
+					{
+						Category.Entries.RemoveAll([](const FChaosVDExtraDataStructEntry& Entry)
+						{
+							return Entry.SerializationMode == EChaosVDExtraDataSerializationMode::NativeSerialization;
+						});
+					}
+				}
+			}
+		}
+	}
+
+	return bSuccess;
+}
diff --git a/Engine/Plugins/ChaosVD/Source/ChaosVD/Public/DataWrappers/ChaosVDParticleExtraDataContainer.h b/Engine/Plugins/ChaosVD/Source/ChaosVD/Public/DataWrappers/ChaosVDParticleExtraDataContainer.h
new file mode 100644
--- /dev/null
+++ b/Engine/Plugins/ChaosVD/Source/ChaosVD/Public/DataWrappers/ChaosVDParticleExtraDataContainer.h
@@ -0,0 +1,36 @@
+// Copyright Epic Games, Inc. All Rights Reserved.
+
+#pragma once
+
+#include "ChaosVisualDebugger/ChaosVDParticleExtraData.h"
+#include "ChaosVisualDebugger/ChaosVDMemWriterReader.h"
+#include "Containers/Map.h"
+#include "Templates/SharedPointer.h"
+
+#include "ChaosVDParticleExtraDataContainer.generated.h"
+
+/**
+ * Custom data container stored in FChaosVDSolverFrameData::CustomData.
+ *
+ * Maps SolverID -> ParticleID -> extra data payload for O(1) lookup at display time.
+ * Must be a USTRUCT so FChaosVDCustomFrameData::GetOrAddDefaultData<T>() can key it
+ * via StaticStruct()->GetFName()
+ *
+ * The NameTable is set by FChaosVDParticleExtraDataProcessor at load time so that
+ * UChaosVDParticleExtraDataComponent can deserialize inner struct bytes without
+ * needing a separate route to the trace provider.
+ */
+USTRUCT()
+struct FChaosVDParticleExtraDataContainer
+{
+	GENERATED_BODY()
+
+	TMap<int32, TMap<int32, TSharedPtr<Chaos::VisualDebugger::FChaosVDParticleExtraData>>> DataBySolverAndParticleID;
+
+	/** Recording-session name table, cached here by the processor so the component can resolve FName IDs in inner struct bytes. */
+	TSharedPtr<Chaos::VisualDebugger::FChaosVDSerializableNameTable> NameTable;
+
+	/** Struct type paths for entries that used NativeSerialization mode.
+	 *  Populated by FChaosVDParticleExtraDataProcessor at load time. */
+	TSet<FName> NativeSerializedStructTypes;
+};
diff --git a/Engine/Plugins/ChaosVD/Source/ChaosVD/Public/Settings/ChaosVDGeneralSettings.h b/Engine/Plugins/ChaosVD/Source/ChaosVD/Public/Settings/ChaosVDGeneralSettings.h
--- a/Engine/Plugins/ChaosVD/Source/ChaosVD/Public/Settings/ChaosVDGeneralSettings.h
+++ b/Engine/Plugins/ChaosVD/Source/ChaosVD/Public/Settings/ChaosVDGeneralSettings.h
@@ -6,6 +6,36 @@
 #include "ChaosVDRecordingDetails.h"
 #include "ChaosVDGeneralSettings.generated.h"
 
+/** Controls how FChaosVDParticleExtraDataProcessor handles particle extra data at load time.
+ *
+ *  SerializeBin entries use property-tagged serialization: properties are matched by name at load
+ *  time against the *current* struct layout, so adds/removes/renames are handled gracefully
+ *  (missing props show defaults; added props are lost). No crash risk.
+ *
+ *  NativeSerialization entries invoke the struct's native Serialize function, which reads bytes
+ *  at fixed offsets. When the struct layout matches, values are exact. When it does not match,
+ *  the result is a crash or silently corrupt values — there is no graceful degradation.
+ */
+UENUM()
+enum class EChaosVDParticleExtraDataLoadingMode : uint8
+{
+	/** Skip the processor entirely — no extra data is loaded or displayed. The only mode that
+	 *  guarantees neither incorrect values nor crashes regardless of how much struct definitions
+	 *  have changed since recording. */
+	SkipAll          UMETA(DisplayName = "Skip All (Safe)"),
+
+	/** Load only SerializeBin entries; NativeSerialization entries are stripped in the processor
+	 *  before reaching memory. No crash risk, but displayed values may be inaccurate if struct
+	 *  definitions changed since recording. */
+	SerializeBinOnly UMETA(DisplayName = "SerializeBin Only"),
+
+	/** Load all entries including NativeSerialization. Displays exact recorded values when the
+	 *  struct layout matches the recording. If the layout has changed, NativeSerialization entries
+	 *  may crash or display corrupt values. Only use when the recording and viewer binaries are
+	 *  known to match. */
+	LoadAll          UMETA(DisplayName = "Load All (Unsafe)"),
+};
+
 /**
  * General settings that controls how CVD behaves
  */
@@ -76,6 +106,56 @@ PRAGMA_ENABLE_DEPRECATION_WARNINGS
 	UPROPERTY(Config, EditAnywhere, Category = Connection)
 	bool bAutoStartLiveSessionsPlayback = true;
 
+	/**
+	 * If true, CVD will skip data whose binary format is not guaranteed to be backwards compatible
+	 * (e.g. raw-serialized structs such as generic particle extra data).
+	 * This prevents crashes when opening a recording made with a different executable than the viewer.
+	 *
+	 * When enabled, any skipped data types are reported in a dialog after loading.
+	 * Disable this only when you are certain the recording was made with the same binary as the viewer.
+	 */
+	UPROPERTY(Config, EditAnywhere, Category = "Backwards Compatibility")
+	bool bSafeLoadingMode = true;
+
+	/**
+	 * Controls how particle extra data is loaded. Particle extra data is not backwards compatible
+	 * by design -- it accepts arbitrary UStructs without requiring dedicated data wrappers.
+	 *
+	 * SkipAll: processor is skipped entirely. No extra data is loaded or displayed. The only mode
+	 *   that guarantees neither crashes nor silently incorrect values, regardless of how much struct
+	 *   definitions have changed since recording.
+	 *
+	 * SerializeBinOnly (default): only property-tagged (SerializeBin) entries are loaded;
+	 *   NativeSerialization entries are filtered out in the processor and never stored in memory.
+	 *   No crash risk, but values may be inaccurate if struct properties were added, removed, or
+	 *   renamed since recording.
+	 *
+	 * LoadAll: all entries including NativeSerialization are loaded. Displays exact recorded values
+	 *   when the viewer binary matches the recording binary. If struct layouts have changed,
+	 *   NativeSerialization entries may crash or show corrupt values. Only use when recording and
+	 *   viewer are known to match.
+	 *
+	 * Note: bSafeLoadingMode applies to any other non-backwards-compatible processors; this enum
+	 * takes precedence for particle extra data specifically.
+	 */
+	UPROPERTY(Config, EditAnywhere, Category = "Backwards Compatibility")
+	EChaosVDParticleExtraDataLoadingMode ParticleExtraDataLoadingMode = EChaosVDParticleExtraDataLoadingMode::SerializeBinOnly;
+
+	/** If true, a warning dialog is shown after loading when non-backwards-compatible data was processed
+	 *  with Safe Loading Mode disabled. Uncheck to suppress the warning. */
+	UPROPERTY(Config, EditAnywhere, Category = "Backwards Compatibility")
+	bool bShowUnsafeDataWarning = true;
+
+	/** If true, an informational dialog is shown after loading when data types were skipped because
+	 *  Safe Loading Mode is enabled. Uncheck to suppress the warning. */
+	UPROPERTY(Config, EditAnywhere, Category = "Backwards Compatibility")
+	bool bShowSafeModeSkippedWarning = true;
+
+	/** If true, shows a warning dialog after loading when native-serialized extra data is
+	 *  present and ParticleExtraDataLoadingMode is LoadAll. */
+	UPROPERTY(Config, EditAnywhere, Category = "Backwards Compatibility")
+	bool bShowNativeSerializationWarning = true;
+
 private:
 	/** [DEBUG Setting] Default Data Transport for Game Servers - Used to change the transport mode for the trace data. This is not saved by design */
 	UPROPERTY(EditAnywhere, Category = Connection)
diff --git a/Engine/Plugins/ChaosVD/Source/ChaosVD/Public/Trace/DataProcessors/ChaosVDDataProcessorBase.h b/Engine/Plugins/ChaosVD/Source/ChaosVD/Public/Trace/DataProcessors/ChaosVDDataProcessorBase.h
--- a/Engine/Plugins/ChaosVD/Source/ChaosVD/Public/Trace/DataProcessors/ChaosVDDataProcessorBase.h
+++ b/Engine/Plugins/ChaosVD/Source/ChaosVD/Public/Trace/DataProcessors/ChaosVDDataProcessorBase.h
@@ -6,7 +6,9 @@
 #include "ChaosVisualDebugger/ChaosVDMemWriterReader.h"
 #include "ChaosVisualDebugger/ChaosVisualDebuggerTrace.h"
 #include "Containers/Array.h"
+#include "Containers/Map.h"
 #include "Containers/StringFwd.h"
+#include "UObject/NameTypes.h"
 #include "DataWrappers/ChaosVDImplicitObjectDataWrapper.h"
 #include "HAL/Platform.h"
 #include "Templates/SharedPointer.h"
@@ -91,10 +93,38 @@ public:
 	/** Sets the Trace Provider that is storing the data being analyzed */
 	CHAOSVD_API void SetTraceProvider(const TSharedPtr<FChaosVDTraceProvider>& InProvider);
 
+	/** Returns natively-serialized struct types encountered during ProcessRawData, mapped to the
+	 *  CVD optional data channel they were traced through (NAME_None if unknown).
+	 *  Used by the post-load NativeSerialization warning dialog. */
+	virtual void GetPostLoadNativeTypesWithChannels(TMap<FName, FName>& OutTypeToChannel) const {}
+
+
+	/**
+	 * Returns true if this processor's data format is backwards compatible, meaning it
+	 * uses UE's property system or explicit version fields that allow graceful forward/backward reading.
+	 *
+	 * Processors that serialise raw binary struct layouts (no versioning) must set
+	 * bBackwardsCompatible = false in their constructor. When Safe Loading Mode is enabled
+	 * in CVD General Settings, the trace provider will skip those processors entirely rather than
+	 * risk a crash on a format mismatch.
+	 */
+	virtual bool IsBackwardsCompatible() const { return bBackwardsCompatible; }
+
+	/** When true, the provider skips this processor silently regardless of bSafeLoadingMode.
+	 *  Use when the processor has its own loading-mode control (e.g. EChaosVDParticleExtraDataLoadingMode). */
+	virtual bool ShouldAlwaysSkip() const { return false; }
+
 protected:
 	TWeakPtr<FChaosVDTraceProvider> TraceProvider;
 	FStringView CompatibleType;
 	uint64 ProcessedBytes;
+
+	/**
+	 * Set to false in the constructor of any derived processor whose serialised format is not
+	 * guaranteed to be stable across engine versions or local source changes.
+	 * Defaults to true so existing processors are unaffected until explicitly opted out.
+	 */
+	bool bBackwardsCompatible = true;
 };
 
 class FChaosVDGenericDataProcessor : FChaosVDDataProcessorBase
diff --git a/Engine/Plugins/ChaosVD/Source/ChaosVD/Public/Trace/DataProcessors/ChaosVDParticleExtraDataProcessor.h b/Engine/Plugins/ChaosVD/Source/ChaosVD/Public/Trace/DataProcessors/ChaosVDParticleExtraDataProcessor.h
new file mode 100644
--- /dev/null
+++ b/Engine/Plugins/ChaosVD/Source/ChaosVD/Public/Trace/DataProcessors/ChaosVDParticleExtraDataProcessor.h
@@ -0,0 +1,31 @@
+// Copyright Epic Games, Inc. All Rights Reserved.
+
+#pragma once
+
+#include "HAL/Platform.h"
+#include "Trace/DataProcessors/ChaosVDDataProcessorBase.h"
+
+enum class EChaosVDParticleExtraDataLoadingMode : uint8;
+
+/**
+ * Data processor that deserializes traced FChaosVDParticleExtraData payloads and stores
+ * them in FChaosVDParticleExtraDataContainer inside the solver frame's custom data.
+ */
+class FChaosVDParticleExtraDataProcessor final : public FChaosVDDataProcessorBase
+{
+public:
+	CHAOSVD_API explicit FChaosVDParticleExtraDataProcessor();
+	CHAOSVD_API ~FChaosVDParticleExtraDataProcessor();
+
+	CHAOSVD_API virtual bool ProcessRawData(const TArray<uint8>& InData) override;
+
+	virtual bool ShouldAlwaysSkip() const override;
+	virtual bool IsBackwardsCompatible() const override;
+	virtual void GetPostLoadNativeTypesWithChannels(TMap<FName, FName>& OutTypeToChannel) const override;
+
+private:
+	void HandleSettingsChanged(UObject* InSettings);
+
+	TMap<FName, FName> NativeTypesToChannel;
+	EChaosVDParticleExtraDataLoadingMode CachedLoadingMode;
+};
diff --git a/Engine/Plugins/ChaosVD/Source/ChaosVDBuiltInExtensions/Private/ChaosVDBuiltInExtensions.cpp b/Engine/Plugins/ChaosVD/Source/ChaosVDBuiltInExtensions/Private/ChaosVDBuiltInExtensions.cpp
--- a/Engine/Plugins/ChaosVD/Source/ChaosVDBuiltInExtensions/Private/ChaosVDBuiltInExtensions.cpp
+++ b/Engine/Plugins/ChaosVD/Source/ChaosVDBuiltInExtensions/Private/ChaosVDBuiltInExtensions.cpp
@@ -5,12 +5,14 @@
 #include "AccelerationStructures/ChaosVDAccelerationStructuresExtension.h"
 #include "GenericDebugDraw/ChaosVDGenericDebugDrawExtension.h"
 #include "CameraTraces/ChaosVDCameraTracesExtension.h"
+#include "ParticleExtraData/ChaosVDParticleExtraDataExtension.h"
 
 void FChaosVDBuiltInExtensionsModule::StartupModule()
 {
 	CreateAndRegisterExtensionInstance<FChaosVDGenericDebugDrawExtension>();
 	CreateAndRegisterExtensionInstance<FChaosVDAccelerationStructuresExtension>();
 	CreateAndRegisterExtensionInstance<FChaosVDCameraTracesExtension>();
+	CreateAndRegisterExtensionInstance<FChaosVDParticleExtraDataExtension>();
 }
 
 void FChaosVDBuiltInExtensionsModule::ShutdownModule()
diff --git a/Engine/Plugins/ChaosVD/Source/ChaosVDBuiltInExtensions/Private/ParticleExtraData/ChaosVDParticleExtraDataExtension.cpp b/Engine/Plugins/ChaosVD/Source/ChaosVDBuiltInExtensions/Private/ParticleExtraData/ChaosVDParticleExtraDataExtension.cpp
new file mode 100644
--- /dev/null
+++ b/Engine/Plugins/ChaosVD/Source/ChaosVDBuiltInExtensions/Private/ParticleExtraData/ChaosVDParticleExtraDataExtension.cpp
@@ -0,0 +1,20 @@
+// Copyright Epic Games, Inc. All Rights Reserved.
+
+#include "ParticleExtraData/ChaosVDParticleExtraDataExtension.h"
+
+#include "Trace/DataProcessors/ChaosVDParticleExtraDataProcessor.h"
+#include "Trace/ChaosVDTraceProvider.h"
+
+FChaosVDParticleExtraDataExtension::FChaosVDParticleExtraDataExtension()
+{
+	ExtensionName = FName(TEXT("FChaosVDParticleExtraDataExtension"));
+}
+
+void FChaosVDParticleExtraDataExtension::RegisterDataProcessorsInstancesForProvider(const TSharedRef<FChaosVDTraceProvider>& InTraceProvider)
+{
+	FChaosVDExtension::RegisterDataProcessorsInstancesForProvider(InTraceProvider);
+
+	TSharedPtr<FChaosVDParticleExtraDataProcessor> Processor = MakeShared<FChaosVDParticleExtraDataProcessor>();
+	Processor->SetTraceProvider(InTraceProvider);
+	InTraceProvider->RegisterDataProcessor(Processor);
+}
diff --git a/Engine/Plugins/ChaosVD/Source/ChaosVDBuiltInExtensions/Private/ParticleExtraData/ChaosVDParticleExtraDataExtension.h b/Engine/Plugins/ChaosVD/Source/ChaosVDBuiltInExtensions/Private/ParticleExtraData/ChaosVDParticleExtraDataExtension.h
new file mode 100644
--- /dev/null
+++ b/Engine/Plugins/ChaosVD/Source/ChaosVDBuiltInExtensions/Private/ParticleExtraData/ChaosVDParticleExtraDataExtension.h
@@ -0,0 +1,20 @@
+// Copyright Epic Games, Inc. All Rights Reserved.
+
+#pragma once
+
+#include "ExtensionsSystem/ChaosVDExtension.h"
+
+/**
+ * Built-in CVD extension that registers the FChaosVDParticleExtraDataProcessor.
+ *
+ * UChaosVDParticleExtraDataComponent is registered directly as a default subobject
+ * on AChaosVDSolverInfoActor, so this extension only handles processor registration.
+ */
+class FChaosVDParticleExtraDataExtension final : public FChaosVDExtension
+{
+public:
+	FChaosVDParticleExtraDataExtension();
+	virtual ~FChaosVDParticleExtraDataExtension() override = default;
+
+	virtual void RegisterDataProcessorsInstancesForProvider(const TSharedRef<FChaosVDTraceProvider>& InTraceProvider) override;
+};
