ChaosVD plugin: particle extra data display, enrichment hooks, and backwards-compatibility UX

UX and display layer that surfaces particle extra data in the CVD details panel:

- ChaosVDScene.h: OnParticleDetailsEnrichmentRequested /
OnParticleDetailsEnrichmentRefreshRequested delegates on FChaosVDScene,
fired during particle details panel construction and per-frame refresh
- ChaosVDSolverInfoActor.h/.cpp: creates and registers
UChaosVDParticleExtraDataComponent on the solver info actor
- ChaosVDSceneParticleCustomization.h/.cpp: fires the enrichment delegates at
the right points in the IDetailCustomization lifecycle
- ChaosVDParticleExtraDataComponent.h/.cpp: UChaosVDParticleExtraDataComponent --
binds to scene enrichment delegates; caches FStructOnScope per particle for
efficient frame-to-frame in-place updates without layout rebuilds; NativeSerialization
filtering is handled by the processor so the component displays whatever entries
are present in the container
- ChaosVDGeneralSettings.h: EChaosVDParticleExtraDataLoadingMode enum
(SkipAll / SerializeBinOnly / LoadAll) replaces bSkipNativeSerializedEntries;
bSafeLoadingMode retained for other processors; three bShow*Warning booleans
all grouped under "Backwards Compatibility"
- ChaosVDTraceProvider.h/.cpp: ShouldAlwaysSkip() checked before the
IsBackwardsCompatible() gate in ProcessBinaryData; ShowNativeSerializationDialog
only fires in LoadAll mode and annotates each type with its source channel
("via X channel"); new ShowSerializeBinOnlyStrippedDialog fires in SerializeBinOnly
mode when native entries were present (shares bShowSafeModeSkippedWarning suppress);
uses GetPostLoadNativeTypesWithChannels across all processors

#Synced-CL 52152484

diff --git a/Engine/Plugins/ChaosVD/Source/ChaosVD/Private/Actors/ChaosVDSolverInfoActor.cpp b/Engine/Plugins/ChaosVD/Source/ChaosVD/Private/Actors/ChaosVDSolverInfoActor.cpp
--- a/Engine/Plugins/ChaosVD/Source/ChaosVD/Private/Actors/ChaosVDSolverInfoActor.cpp
+++ b/Engine/Plugins/ChaosVD/Source/ChaosVD/Private/Actors/ChaosVDSolverInfoActor.cpp
@@ -9,6 +9,7 @@
 #include "EditorActorFolders.h"
 #include "Components/ChaosVDAdditionalGTDataRouterComponent.h"
 #include "Components/ChaosVDParticleDataComponent.h"
+#include "Components/ChaosVDParticleExtraDataComponent.h"
 #include "Components/ChaosVDSceneQueryDataComponent.h"
 #include "Components/ChaosVDSolverCharacterGroundConstraintDataComponent.h"
 #include "Components/ChaosVDSolverCollisionDataComponent.h"
@@ -27,6 +28,7 @@ AChaosVDSolverInfoActor::AChaosVDSolverInfoActor()
 	CharacterGroundConstraintDataComponent = CreateDefaultSubobject<UChaosVDSolverCharacterGroundConstraintDataComponent>(TEXT("CharacterGroundConstraintDataComponent"));
 	SceneQueryDataComponent = CreateDefaultSubobject<UChaosVDSceneQueryDataComponent>(TEXT("ChaosVDSceneQueryDataComponent"));
 	GTDataReRouteComponent = CreateDefaultSubobject<UChaosVDAdditionalGTDataRouterComponent>(TEXT("ChaosVDAdditionalGTDataRouterComponent"));
+	ParticleExtraDataComponent = CreateDefaultSubobject<UChaosVDParticleExtraDataComponent>(TEXT("ParticleExtraDataComponent"));
 
 	bIsServer = false;
 }
diff --git a/Engine/Plugins/ChaosVD/Source/ChaosVD/Private/Actors/ChaosVDSolverInfoActor.h b/Engine/Plugins/ChaosVD/Source/ChaosVD/Private/Actors/ChaosVDSolverInfoActor.h
--- a/Engine/Plugins/ChaosVD/Source/ChaosVD/Private/Actors/ChaosVDSolverInfoActor.h
+++ b/Engine/Plugins/ChaosVD/Source/ChaosVD/Private/Actors/ChaosVDSolverInfoActor.h
@@ -14,6 +14,7 @@ class UChaosVDSolverCharacterGroundConstraintDataComponent;
 class UChaosVDSolverJointConstraintDataComponent;
 struct FChaosVDParticleDataWrapper;
 class UChaosVDParticleDataComponent;
+class UChaosVDParticleExtraDataComponent;
 class UChaosVDSceneQueryDataComponent;
 class UChaosVDSolverCollisionDataComponent;
 class UChaosVDCameraDataComponent;
@@ -126,4 +127,7 @@ protected:
 
 	UPROPERTY()
 	TObjectPtr<UChaosVDAdditionalGTDataRouterComponent> GTDataReRouteComponent;
+
+	UPROPERTY()
+	TObjectPtr<UChaosVDParticleExtraDataComponent> ParticleExtraDataComponent;
 };
diff --git a/Engine/Plugins/ChaosVD/Source/ChaosVD/Private/ChaosVDSceneParticleCustomization.cpp b/Engine/Plugins/ChaosVD/Source/ChaosVD/Private/ChaosVDSceneParticleCustomization.cpp
--- a/Engine/Plugins/ChaosVD/Source/ChaosVD/Private/ChaosVDSceneParticleCustomization.cpp
+++ b/Engine/Plugins/ChaosVD/Source/ChaosVD/Private/ChaosVDSceneParticleCustomization.cpp
@@ -70,6 +70,9 @@ void FChaosVDSceneParticleCustomization::AddParticleDataButtons(IDetailLayoutBui
 
 void FChaosVDSceneParticleCustomization::CustomizeDetails(IDetailLayoutBuilder& DetailBuilder)
 {
+	CachedDetailBuilder = &DetailBuilder;
+	TGuardValue<bool> InsideCustomizeDetailsGuard(bIsInsideCustomizeDetails, true);
+
 	FChaosVDDetailsCustomizationUtils::HideAllCategories(DetailBuilder, AllowedCategories);
 
 	TSharedPtr<SChaosVDMainTab> MainTabPtr =  MainTabWeakPtr.Pin(); 
@@ -154,12 +157,25 @@ void FChaosVDSceneParticleCustomization::CustomizeDetails(IDetailLayoutBuilder&
 				TSharedRef<IPropertyHandle> InspectedDataPropertyHandleRef = InspectedDataPropertyHandlePtr.ToSharedRef();
 				FChaosVDDetailsCustomizationUtils::HideInvalidCVDDataWrapperProperties({&InspectedDataPropertyHandleRef, 1}, DetailBuilder);
 			}
+
+			// Fire the extra data delegate so any registered component (e.g. UChaosVDParticleExtraDataComponent)
+			// can append additional categories for this particle.
+			if (TSharedPtr<const FChaosVDParticleDataWrapper> ParticleData = SelectedParticleInstance->GetParticleData())
+			{
+				const int32 SolverID = ParticleData->SolverID;
+				const int32 ParticleID = ParticleData->ParticleIndex;
+				if (SolverID != INDEX_NONE && ParticleID != INDEX_NONE)
+				{
+					Scene->OnParticleDetailsExtraDataRequested().Broadcast(SolverID, ParticleID, DetailBuilder);
+				}
+			}
 		}
 	}
 	else
 	{
 		ResetCachedView();
 	}
+
 }
 
 void FChaosVDSceneParticleCustomization::HandleSceneUpdated()
@@ -204,6 +220,34 @@ void FChaosVDSceneParticleCustomization::HandleSceneUpdated()
 			CachedParticleMetadata = FChaosVDParticleMetadata();
 		}
 	}
+
+	// Notify extra data subscribers (e.g. UChaosVDParticleExtraDataComponent) that new frame data
+	// is available. They update their cached FStructOnScope memory in-place so Slate reflects the
+	// new values on the next redraw. If the data layout has changed a subscriber sets
+	// bNeedsFullRebuild=true and we fall back to a ForceRefreshDetails.
+	if (const TSharedPtr<FChaosVDScene> Scene = SceneWeakPtr.Pin())
+	{
+		if (Scene->OnParticleDetailsExtraDataRefreshRequested().IsBound())
+		{
+			if (TSharedPtr<const FChaosVDParticleDataWrapper> ParticleData = CurrentObservedParticle->GetParticleData())
+			{
+				const int32 SolverID = ParticleData->SolverID;
+				const int32 ParticleID = ParticleData->ParticleIndex;
+				if (SolverID != INDEX_NONE && ParticleID != INDEX_NONE)
+				{
+					bool bNeedsFullRebuild = false;
+					Scene->OnParticleDetailsExtraDataRefreshRequested().Broadcast(SolverID, ParticleID, bNeedsFullRebuild);
+
+					if (bNeedsFullRebuild && CachedDetailBuilder && !bIsInsideCustomizeDetails)
+					{
+						IDetailLayoutBuilder* BuilderToRefresh = CachedDetailBuilder;
+						CachedDetailBuilder = nullptr;
+						BuilderToRefresh->ForceRefreshDetails();
+					}
+				}
+			}
+		}
+	}
 }
 
 bool FChaosVDSceneParticleCustomization::GetCollisionDataButtonEnabled() const
@@ -301,9 +345,9 @@ void FChaosVDSceneParticleCustomization::RegisterCVDScene(const TSharedPtr<FChao
 	}
 }
 
-TSharedPtr<SWidget> FChaosVDSceneParticleCustomization::GenerateShowCollisionDataButton()
+static TSharedRef<SWidget> MakeCVDDetailsButton(const FText& Tooltip, const FText& Label, FOnClicked OnClicked, TAttribute<bool> IsEnabled = true)
 {
-	TSharedPtr<SWidget> ShowCollisionButton = SNew(SHorizontalBox)
+	return SNew(SHorizontalBox)
 	+SHorizontalBox::Slot()
 	.VAlign(VAlign_Center)
 	.HAlign(HAlign_Center)
@@ -311,10 +355,10 @@ TSharedPtr<SWidget> FChaosVDSceneParticleCustomization::GenerateShowCollisionDat
 	.FillWidth(1.0f)
 	[
 		SNew(SButton)
-		.ToolTip(SNew(SToolTip).Text(LOCTEXT("OpenCollisionDataDesc", "Click here to open the collision data for this particle on the collision data inspector.")))
-		.IsEnabled_Raw(this, &FChaosVDSceneParticleCustomization::GetCollisionDataButtonEnabled)
+		.ToolTip(SNew(SToolTip).Text(Tooltip))
+		.IsEnabled(IsEnabled)
 		.ContentPadding(FMargin(0, 5.f, 0, 4.f))
-		.OnClicked_Raw(this, &FChaosVDSceneParticleCustomization::ShowCollisionDataForSelectedObject)
+		.OnClicked(OnClicked)
 		.Content()
 		[
 			SNew(SHorizontalBox)
@@ -325,43 +369,27 @@ TSharedPtr<SWidget> FChaosVDSceneParticleCustomization::GenerateShowCollisionDat
 			[
 				SNew(STextBlock)
 				.TextStyle(FAppStyle::Get(), "SmallButtonText")
-				.Text(LOCTEXT("ShowCollisionDataOnInspector","Show Collision Data in Inspector"))
+				.Text(Label)
 			]
 		]
 	];
+}
 
-	return ShowCollisionButton;
+TSharedPtr<SWidget> FChaosVDSceneParticleCustomization::GenerateShowCollisionDataButton()
+{
+	return MakeCVDDetailsButton(
+		LOCTEXT("OpenCollisionDataDesc", "Click here to open the collision data for this particle on the collision data inspector."),
+		LOCTEXT("ShowCollisionDataOnInspector", "Show Collision Data in Inspector"),
+		FOnClicked::CreateRaw(this, &FChaosVDSceneParticleCustomization::ShowCollisionDataForSelectedObject),
+		TAttribute<bool>::CreateRaw(this, &FChaosVDSceneParticleCustomization::GetCollisionDataButtonEnabled));
 }
 
 TSharedPtr<SWidget> FChaosVDSceneParticleCustomization::GenerateOpenInNewDetailsPanelButton()
 {
-	TSharedPtr<SWidget> OpenButton = SNew(SHorizontalBox)
-	+SHorizontalBox::Slot()
-	.VAlign(VAlign_Center)
-	.HAlign(HAlign_Center)
-	.Padding(12.0f, 7.0f, 12.0f, 7.0f)
-	.FillWidth(1.0f)
-	[
-		SNew(SButton)
-		.ToolTip(SNew(SToolTip).Text(LOCTEXT("OpenDetailsPanelDesc", "Click here to open a new (selection independent) details panel for this particle.")))
-		.ContentPadding(FMargin(0, 5.f, 0, 4.f))
-		.OnClicked_Raw(this, &FChaosVDSceneParticleCustomization::OpenNewDetailsPanel)
-		.Content()
-		[
-			SNew(SHorizontalBox)
-			+SHorizontalBox::Slot()
-			.Padding(FMargin(3, 0, 0, 0))
-			.VAlign(VAlign_Center)
-			.AutoWidth()
-			[
-				SNew(STextBlock)
-				.TextStyle(FAppStyle::Get(), "SmallButtonText")
-				.Text(LOCTEXT("OpenDetailsPanelText","Show Data in New Panel"))
-			]
-		]
-	];
-
-	return OpenButton;
+	return MakeCVDDetailsButton(
+		LOCTEXT("OpenDetailsPanelDesc", "Click here to open a new (selection independent) details panel for this particle."),
+		LOCTEXT("OpenDetailsPanelText", "Show Data in New Panel"),
+		FOnClicked::CreateRaw(this, &FChaosVDSceneParticleCustomization::OpenNewDetailsPanel));
 }
 
 void FChaosVDSceneParticleCustomization::UpdateObserverParticlePtr(FChaosVDSceneParticle* NewObservedParticle)
diff --git a/Engine/Plugins/ChaosVD/Source/ChaosVD/Private/ChaosVDSceneParticleCustomization.h b/Engine/Plugins/ChaosVD/Source/ChaosVD/Private/ChaosVDSceneParticleCustomization.h
--- a/Engine/Plugins/ChaosVD/Source/ChaosVD/Private/ChaosVDSceneParticleCustomization.h
+++ b/Engine/Plugins/ChaosVD/Source/ChaosVD/Private/ChaosVDSceneParticleCustomization.h
@@ -67,6 +67,15 @@ private:
 	TWeakPtr<FChaosVDScene> SceneWeakPtr = nullptr;
 
 	TWeakPtr<SChaosVDMainTab> MainTabWeakPtr = nullptr;
+
+	/** Cached builder pointer used to trigger a ForceRefreshDetails if an extra data subscriber
+	 *  reports a layout mismatch during an in-place refresh. Nulled before the call to avoid
+	 *  dangling use if the rebuild is synchronous. */
+	IDetailLayoutBuilder* CachedDetailBuilder = nullptr;
+
+	/** Guard that prevents HandleSceneUpdated from issuing a ForceRefreshDetails while
+	 *  CustomizeDetails is already on the call stack, avoiding an infinite rebuild loop. */
+	bool bIsInsideCustomizeDetails = false;
 };
 
 template <typename TStruct>
diff --git a/Engine/Plugins/ChaosVD/Source/ChaosVD/Private/Components/ChaosVDParticleExtraDataComponent.cpp b/Engine/Plugins/ChaosVD/Source/ChaosVD/Private/Components/ChaosVDParticleExtraDataComponent.cpp
new file mode 100644
--- /dev/null
+++ b/Engine/Plugins/ChaosVD/Source/ChaosVD/Private/Components/ChaosVDParticleExtraDataComponent.cpp
@@ -0,0 +1,240 @@
+// Copyright Epic Games, Inc. All Rights Reserved.
+
+#include "Components/ChaosVDParticleExtraDataComponent.h"
+
+#include "ChaosVDModule.h"
+#include "ChaosVDScene.h"
+#include "ChaosVisualDebugger/ChaosVDParticleExtraData.h"
+#include "ChaosVisualDebugger/ChaosVDStructCollectionMemWriterReader.h"
+#include "DataWrappers/ChaosVDParticleExtraDataContainer.h"
+#include "DetailCategoryBuilder.h"
+#include "DetailLayoutBuilder.h"
+#include "UObject/Class.h"
+#include "UObject/UObjectGlobals.h"
+
+#include UE_INLINE_GENERATED_CPP_BY_NAME(ChaosVDParticleExtraDataComponent)
+
+using namespace Chaos::VisualDebugger;
+
+void UChaosVDParticleExtraDataComponent::ClearData()
+{
+	CurrentFrameContainer.Reset();
+	CachedScopeEntries.Reset();
+	CachedSolverID = INDEX_NONE;
+	CachedParticleID = INDEX_NONE;
+}
+
+void UChaosVDParticleExtraDataComponent::UpdateFromSolverFrameData(const FChaosVDSolverFrameData& InSolverFrameData)
+{
+	CurrentFrameContainer = InSolverFrameData.GetCustomData().GetData<FChaosVDParticleExtraDataContainer>();
+}
+
+void UChaosVDParticleExtraDataComponent::SetScene(const TWeakPtr<FChaosVDScene>& InSceneWeakPtr)
+{
+	// Unbind from the previous scene
+	if (TSharedPtr<FChaosVDScene> OldScene = SceneWeakPtr.Pin())
+	{
+		UnbindFromScene(OldScene);
+	}
+
+	Super::SetScene(InSceneWeakPtr);
+
+	if (TSharedPtr<FChaosVDScene> NewScene = SceneWeakPtr.Pin())
+	{
+		ExtraDataDelegateHandle = NewScene->OnParticleDetailsExtraDataRequested().AddUObject(
+			this, &UChaosVDParticleExtraDataComponent::HandleParticleDetailsExtraData);
+
+		RefreshDelegateHandle = NewScene->OnParticleDetailsExtraDataRefreshRequested().AddUObject(
+			this, &UChaosVDParticleExtraDataComponent::HandleParticleDetailsRefresh);
+	}
+}
+
+void UChaosVDParticleExtraDataComponent::UnbindFromScene(const TSharedPtr<FChaosVDScene>& Scene)
+{
+	if (Scene)
+	{
+		if (ExtraDataDelegateHandle.IsValid())
+		{
+			Scene->OnParticleDetailsExtraDataRequested().Remove(ExtraDataDelegateHandle);
+			ExtraDataDelegateHandle.Reset();
+		}
+
+		if (RefreshDelegateHandle.IsValid())
+		{
+			Scene->OnParticleDetailsExtraDataRefreshRequested().Remove(RefreshDelegateHandle);
+			RefreshDelegateHandle.Reset();
+		}
+	}
+}
+
+void UChaosVDParticleExtraDataComponent::HandleParticleDetailsExtraData(int32 InSolverID, int32 InParticleID, IDetailLayoutBuilder& DetailBuilder)
+{
+	if (!CurrentFrameContainer)
+	{
+		return;
+	}
+
+	// Only handle data that belongs to our solver
+	if (InSolverID != SolverID)
+	{
+		return;
+	}
+
+	const TMap<int32, TSharedPtr<FChaosVDParticleExtraData>>* DataByParticle =
+		CurrentFrameContainer->DataBySolverAndParticleID.Find(SolverID);
+	if (!DataByParticle)
+	{
+		return;
+	}
+
+	const TSharedPtr<FChaosVDParticleExtraData>* ExtraDataPtr = DataByParticle->Find(InParticleID);
+	if (!ExtraDataPtr || !(*ExtraDataPtr))
+	{
+		return;
+	}
+
+	const FChaosVDParticleExtraData& ExtraData = **ExtraDataPtr;
+
+	TSharedPtr<FChaosVDSerializableNameTable> NameTable = CurrentFrameContainer->NameTable;
+	if (!NameTable)
+	{
+		return;
+	}
+
+	// Rebuild the scope cache when the observed particle changes.
+	// For the same particle we reuse the existing FStructOnScope allocations so Slate can
+	// pick up frame-to-frame value changes via HandleParticleDetailsRefresh without a rebuild.
+	const bool bNeedsCacheRebuild = (InSolverID != CachedSolverID || InParticleID != CachedParticleID || CachedScopeEntries.IsEmpty());
+
+	if (bNeedsCacheRebuild)
+	{
+		CachedScopeEntries.Reset();
+		CachedSolverID = InSolverID;
+		CachedParticleID = InParticleID;
+
+		for (const FChaosVDExtraDataCategory& Category : ExtraData.Categories)
+		{
+			for (const FChaosVDExtraDataStructEntry& Entry : Category.Entries)
+			{
+				UScriptStruct* Struct = FindObject<UScriptStruct>(nullptr, *Entry.StructTypePath.ToString());
+				if (!Struct)
+				{
+					UE_LOG(LogChaosVDEditor, Warning,
+						TEXT("UChaosVDParticleExtraDataComponent: Could not resolve UScriptStruct '%s'  - type may not be loaded."),
+						*Entry.StructTypePath.ToString());
+					continue;
+				}
+
+				TSharedPtr<FStructOnScope> StructScope = MakeShared<FStructOnScope>(Struct);
+				FChaosVDStructCollectionMemoryReader Reader(Entry.Bytes, NameTable.ToSharedRef());
+				if (Entry.SerializationMode == EChaosVDExtraDataSerializationMode::SerializeBin)
+				{
+					Struct->SerializeBin(Reader, StructScope->GetStructMemory());
+				}
+				else
+				{
+					Struct->SerializeItem(Reader, StructScope->GetStructMemory(), nullptr);
+				}
+
+				if (Reader.IsError())
+				{
+					UE_LOG(LogChaosVDEditor, Warning,
+						TEXT("UChaosVDParticleExtraDataComponent: Failed to deserialize '%s'  - skipping entry."),
+						*Entry.StructTypePath.ToString());
+					continue;
+				}
+
+				CachedScopeEntries.Add({ Category.CategoryName, Entry.StructTypePath, StructScope });
+			}
+		}
+	}
+
+	// Add all cached scopes to the layout builder
+	for (const FChaosVDExtraDataScopeEntry& ScopeEntry : CachedScopeEntries)
+	{
+		IDetailCategoryBuilder& CategoryBuilder = DetailBuilder.EditCategory(ScopeEntry.CategoryName);
+		if (IDetailPropertyRow* Row = CategoryBuilder.AddExternalStructure(ScopeEntry.Scope))
+		{
+			Row->DisplayName(ScopeEntry.Scope->GetStruct()->GetDisplayNameText());
+			Row->ShouldAutoExpand(true);
+		}
+	}
+}
+
+void UChaosVDParticleExtraDataComponent::HandleParticleDetailsRefresh(int32 InSolverID, int32 InParticleID, bool& bOutNeedsFullRebuild)
+{
+	// Only act if there are cached scopes for exactly this particle
+	if (CachedScopeEntries.IsEmpty() || InSolverID != CachedSolverID || InParticleID != CachedParticleID)
+	{
+		return;
+	}
+
+	if (!CurrentFrameContainer || InSolverID != SolverID)
+	{
+		return;
+	}
+
+	const TMap<int32, TSharedPtr<FChaosVDParticleExtraData>>* DataByParticle =
+		CurrentFrameContainer->DataBySolverAndParticleID.Find(InSolverID);
+	if (!DataByParticle)
+	{
+		return;
+	}
+
+	const TSharedPtr<FChaosVDParticleExtraData>* ExtraDataPtr = DataByParticle->Find(InParticleID);
+	if (!ExtraDataPtr || !(*ExtraDataPtr))
+	{
+		return;
+	}
+
+	const FChaosVDParticleExtraData& ExtraData = **ExtraDataPtr;
+
+	TSharedPtr<FChaosVDSerializableNameTable> NameTable = CurrentFrameContainer->NameTable;
+	if (!NameTable)
+	{
+		return;
+	}
+
+	// Walk the new frame data in lock-step with the cached scope entries.
+	// If the structure is stable (same categories and struct types in the same order),
+	// deserialize the new bytes directly into the existing FStructOnScope memory so
+	// Slate picks up the changes on the next redraw  - no layout rebuild required.
+	int32 ScopeIndex = 0;
+	for (const FChaosVDExtraDataCategory& Category : ExtraData.Categories)
+	{
+		for (const FChaosVDExtraDataStructEntry& Entry : Category.Entries)
+		{
+			if (ScopeIndex >= CachedScopeEntries.Num())
+			{
+				// More entries in the new frame than in the cache  - layout has grown.
+				bOutNeedsFullRebuild = true;
+				return;
+			}
+
+			FChaosVDExtraDataScopeEntry& ScopeEntry = CachedScopeEntries[ScopeIndex];
+
+			// If the layout has changed (different category or struct type), request a full rebuild
+			// so CustomizeDetails re-runs and the new structure is reflected correctly.
+			if (ScopeEntry.CategoryName != Category.CategoryName || ScopeEntry.StructTypePath != Entry.StructTypePath)
+			{
+				bOutNeedsFullRebuild = true;
+				return;
+			}
+
+			if (UScriptStruct* Struct = const_cast<UScriptStruct*>(Cast<const UScriptStruct>(ScopeEntry.Scope->GetStruct())))
+			{
+				FChaosVDStructCollectionMemoryReader Reader(Entry.Bytes, NameTable.ToSharedRef());
+				if (Entry.SerializationMode == EChaosVDExtraDataSerializationMode::SerializeBin)
+				{
+					Struct->SerializeBin(Reader, ScopeEntry.Scope->GetStructMemory());
+				}
+				else
+				{
+					Struct->SerializeItem(Reader, ScopeEntry.Scope->GetStructMemory(), nullptr);
+				}
+			}
+
+			++ScopeIndex;
+		}
+	}
+}
diff --git a/Engine/Plugins/ChaosVD/Source/ChaosVD/Private/Trace/ChaosVDTraceProvider.cpp b/Engine/Plugins/ChaosVD/Source/ChaosVD/Private/Trace/ChaosVDTraceProvider.cpp
--- a/Engine/Plugins/ChaosVD/Source/ChaosVD/Private/Trace/ChaosVDTraceProvider.cpp
+++ b/Engine/Plugins/ChaosVD/Source/ChaosVD/Private/Trace/ChaosVDTraceProvider.cpp
@@ -11,7 +11,16 @@
 #include "DataProcessors/ChaosVDCollisionChannelsInfoDataProcessor.h"
 #include "DataProcessors/ChaosVDParticleMetadataProcessor.h"
 #include "ExtensionsSystem/ChaosVDExtensionsManager.h"
+#include "Application/SlateApplicationBase.h"
+#include "HAL/PlatformApplicationMisc.h"
+#include "Framework/Application/SlateApplication.h"
 #include "Misc/MessageDialog.h"
+#include "Widgets/Images/SImage.h"
+#include "Widgets/Input/SButton.h"
+#include "Widgets/Input/SCheckBox.h"
+#include "Widgets/Layout/SBox.h"
+#include "Widgets/SWindow.h"
+#include "Widgets/Text/STextBlock.h"
 #include "Settings/ChaosVDGeneralSettings.h"
 #include "Trace/DataProcessors/ChaosVDCharacterGroundConstraintDataProcessor.h"
 #include "Trace/DataProcessors/ChaosVDConstraintDataProcessor.h"
@@ -337,6 +346,41 @@ bool FChaosVDTraceProvider::ProcessBinaryData(const int32 DataID)
 			{
 				if (TSharedPtr<FChaosVDDataProcessorBase> DataProcessorPtr = *DataProcessorPtrPtr)
 				{
+					if (DataProcessorPtr->ShouldAlwaysSkip())
+					{
+						UE_LOG(LogChaosVDEditor, Log,
+							TEXT("[%s] Silently skipping [%s] — processor opted out via ShouldAlwaysSkip()."),
+							ANSI_TO_TCHAR(__FUNCTION__), *UnprocessedData->TypeName);
+						UnprocessedDataByID.Remove(DataID);
+						DataProcessedSoFarCounter++;
+						return true;
+					}
+
+					// If this processor's format is not backwards compatible by design, gate on Safe Loading Mode.
+					if (!DataProcessorPtr->IsBackwardsCompatible())
+					{
+						const UChaosVDGeneralSettings* GeneralSettings = FChaosVDSettingsManager::Get().GetSettingsObject<UChaosVDGeneralSettings>();
+						if (GeneralSettings && GeneralSettings->bSafeLoadingMode)
+						{
+							// Treat deliberately-skipped data as a clean success so the rest of the load continues normally.
+							TypesSkippedInSafeMode.Add(UnprocessedData->TypeName);
+							UE_LOG(LogChaosVDEditor, Log,
+								TEXT("[%s] Skipping non-backwards-compatible data type [%s]  - Safe Loading Mode is enabled. See CVD General Settings to change this."),
+								ANSI_TO_TCHAR(__FUNCTION__), *UnprocessedData->TypeName);
+							UnprocessedDataByID.Remove(DataID);
+							DataProcessedSoFarCounter++;
+							return true;
+						}
+						else
+						{
+							// Record that we're about to process potentially unsafe data so we can warn after loading.
+							UnsafeTypesProcessed.Add(UnprocessedData->TypeName);
+							UE_LOG(LogChaosVDEditor, Log,
+								TEXT("[%s] Processing non-backwards-compatible data type [%s] with Safe Loading Mode disabled. Enable Safe Loading Mode in CVD General Settings if CVD crashes unexpectedly."),
+								ANSI_TO_TCHAR(__FUNCTION__), *UnprocessedData->TypeName);
+						}
+					}
+
 					if (ensure(DataProcessorPtr->ProcessRawData(*RawData)))
 					{
 						UnprocessedDataByID.Remove(DataID);
@@ -391,7 +435,7 @@ void FChaosVDTraceProvider::HandleAnalysisComplete()
 	if (!TypesFailedToSerialize.IsEmpty())
 	{
 		FString FailedTypeList = GenerateFormattedStringListFromSet(TypesFailedToSerialize);
-	
+
 		FText MissingDataProcessorsMessage = FText::FormatOrdered(LOCTEXT("FailedSerializationMessage", "The following data types were part of the recording, but they couldn't be read : \n\n {0} \n\n Visualization related to that data will not be shown."), FText::FromString(FailedTypeList));
 
 		FTSTicker::GetCoreTicker().AddTicker(FTickerDelegate::CreateLambda([MissingDataProcessorsMessage](float DeltaTime)
@@ -401,6 +445,11 @@ void FChaosVDTraceProvider::HandleAnalysisComplete()
 		}));
 	}
 
+	ShowSafeModeSkippedDialog();
+	ShowSerializeBinOnlyStrippedDialog();
+	ShowUnsafeDataDialog();
+	ShowNativeSerializationDialog();
+
 	UnprocessedDataByID.Reset();
 
 	UE_LOG(LogChaosVDEditor, Log, TEXT("Trace Analysis complete for session [%s] | Calculating data loaded stats..."), Session.GetName());
@@ -431,6 +480,512 @@ void FChaosVDTraceProvider::HandleAnalysisComplete()
 	UE_LOG(LogChaosVDEditor, Log, TEXT("Total size of loaded data => [%s]"), *FText::AsMemory(TotalBytes, &SizeFormattingOptions,nullptr, EMemoryUnitStandard::IEC).ToString());
 }
 
+void FChaosVDTraceProvider::ShowSafeModeSkippedDialog()
+{
+	if (TypesSkippedInSafeMode.IsEmpty())
+	{
+		return;
+	}
+
+	UChaosVDGeneralSettings* Settings = FChaosVDSettingsManager::Get().GetSettingsObject<UChaosVDGeneralSettings>();
+	if (Settings && !Settings->bShowSafeModeSkippedWarning)
+	{
+		return;
+	}
+
+	using namespace Chaos::VisualDebugger;
+
+	const FChaosVDArchiveHeader CurrentHeader = FChaosVDArchiveHeader::Current();
+	const bool bVersionsMatch = HeaderData.EngineVersion.ExactMatch(CurrentHeader.EngineVersion);
+	const FString RecordedVersionString = HeaderData.EngineVersion.ToString();
+	const FString CurrentVersionString = CurrentHeader.EngineVersion.ToString();
+
+	const FText VersionNote = bVersionsMatch
+		? LOCTEXT("SafeModeVersionMatch", "Engine versions match exactly, but local source changes may still affect the binary layout.")
+		: FText::FormatOrdered(
+			LOCTEXT("SafeModeVersionMismatch", "Recorded with engine version {0}; current version is {1}. Binary format may have changed."),
+			FText::FromString(RecordedVersionString),
+			FText::FromString(CurrentVersionString));
+
+	FText Message = FText::FormatOrdered(
+		LOCTEXT("SafeModeSkippedMessage",
+			"Some data was skipped because Safe Loading Mode is enabled and the following types are not guaranteed to be backwards compatible:\n\n"
+			"{0}\n\n"
+			"{1}\n\n"
+			"To load this data anyway, disable Safe Loading Mode in CVD General Settings (Project > Plugins > Chaos Visual Debugger > General).\n"
+			"Note: doing so may cause a crash if the binary format has changed since the recording was made."),
+		FText::FromString(GenerateFormattedStringListFromSet(TypesSkippedInSafeMode)),
+		VersionNote);
+
+	TWeakObjectPtr<UChaosVDGeneralSettings> WeakSettings(Settings);
+
+	FTSTicker::GetCoreTicker().AddTicker(FTickerDelegate::CreateLambda([Message, WeakSettings](float)
+	{
+		TSharedPtr<bool> bSuppressInFuture = MakeShared<bool>(false);
+
+		TSharedRef<SWindow> DialogWindow = SNew(SWindow)
+			.Title(LOCTEXT("SafeModeSkippedTitle", "Data Skipped - Safe Loading Mode Active"))
+			.SizingRule(ESizingRule::Autosized)
+			.SupportsMaximize(false)
+			.SupportsMinimize(false);
+
+		DialogWindow->SetContent(
+			SNew(SVerticalBox)
+			+ SVerticalBox::Slot().AutoHeight().Padding(16.f, 16.f, 16.f, 8.f)
+			[
+				SNew(STextBlock).Text(Message).WrapTextAt(500.f)
+			]
+			+ SVerticalBox::Slot().AutoHeight().Padding(16.f, 4.f)
+			[
+				SNew(SCheckBox)
+				.IsChecked(ECheckBoxState::Checked)
+				.OnCheckStateChanged_Lambda([bSuppressInFuture](ECheckBoxState NewState)
+				{
+					*bSuppressInFuture = (NewState == ECheckBoxState::Unchecked);
+				})
+				[
+					SNew(STextBlock).Text(LOCTEXT("SafeModeShowWarning", "Show this warning in future"))
+				]
+			]
+			+ SVerticalBox::Slot().AutoHeight().HAlign(HAlign_Right).Padding(16.f, 8.f, 16.f, 16.f)
+			[
+				SNew(SHorizontalBox)
+				+ SHorizontalBox::Slot().AutoWidth().Padding(0.f, 0.f, 8.f, 0.f)
+				[
+					SNew(SButton)
+					.ButtonStyle(FAppStyle::Get(), "SimpleButton")
+					.ToolTipText(LOCTEXT("SafeModeCopyTooltip", "Copy the text in this message to the clipboard (CTRL+C)"))
+					.ContentPadding(2.f)
+					.OnClicked_Lambda([Message]() -> FReply
+					{
+						FPlatformApplicationMisc::ClipboardCopy(*Message.ToString());
+						return FReply::Handled();
+					})
+					.Content()
+					[
+						SNew(SImage)
+						.Image(FAppStyle::Get().GetBrush("Icons.Clipboard"))
+						.ColorAndOpacity(FSlateColor::UseForeground())
+					]
+				]
+				+ SHorizontalBox::Slot().AutoWidth()
+				[
+					SNew(SButton)
+					.Text(LOCTEXT("SafeModeOK", "OK"))
+					.OnClicked_Lambda([DialogWindowWeak = TWeakPtr<SWindow>(DialogWindow), bSuppressInFuture, WeakSettings]() -> FReply
+					{
+						if (*bSuppressInFuture)
+						{
+							if (UChaosVDGeneralSettings* CVDSettings = WeakSettings.Get())
+							{
+								CVDSettings->bShowSafeModeSkippedWarning = false;
+								CVDSettings->SaveConfig();
+							}
+						}
+						if (TSharedPtr<SWindow> DialogWindowPinned = DialogWindowWeak.Pin())
+						{
+							DialogWindowPinned->RequestDestroyWindow();
+						}
+						return FReply::Handled();
+					})
+				]
+			]
+		);
+
+		FSlateApplication::Get().AddModalWindow(DialogWindow, FSlateApplication::Get().GetActiveTopLevelWindow());
+		return false;
+	}));
+}
+
+void FChaosVDTraceProvider::ShowUnsafeDataDialog()
+{
+	if (UnsafeTypesProcessed.IsEmpty())
+	{
+		return;
+	}
+
+	UChaosVDGeneralSettings* Settings = FChaosVDSettingsManager::Get().GetSettingsObject<UChaosVDGeneralSettings>();
+	if (Settings && !Settings->bShowUnsafeDataWarning)
+	{
+		return;
+	}
+
+	using namespace Chaos::VisualDebugger;
+
+	const FChaosVDArchiveHeader CurrentHeader = FChaosVDArchiveHeader::Current();
+	const bool bVersionsMatch = HeaderData.EngineVersion.ExactMatch(CurrentHeader.EngineVersion);
+	const FString RecordedVersionString = HeaderData.EngineVersion.ToString();
+	const FString CurrentVersionString = CurrentHeader.EngineVersion.ToString();
+
+	const FText VersionNote = bVersionsMatch
+		? LOCTEXT("UnsafeVersionMatch", "Engine versions match exactly. The risk of a crash is lower but not zero - local source changes may still affect the binary layout.")
+		: FText::FormatOrdered(
+			LOCTEXT("UnsafeVersionMismatch", "Recorded with engine version {0}; current version is {1}. The risk of a crash due to binary incompatibility is elevated."),
+			FText::FromString(RecordedVersionString),
+			FText::FromString(CurrentVersionString));
+
+	FText Message = FText::FormatOrdered(
+		LOCTEXT("UnsafeModeProcessedMessage",
+			"The following data types are not guaranteed to be backwards compatible and were loaded with Safe Loading Mode disabled:\n\n"
+			"{0}\n\n"
+			"{1}\n\n"
+			"If CVD crashes unexpectedly while viewing this recording, enable Safe Loading Mode in CVD General Settings."),
+		FText::FromString(GenerateFormattedStringListFromSet(UnsafeTypesProcessed)),
+		VersionNote);
+
+	TWeakObjectPtr<UChaosVDGeneralSettings> WeakSettings(Settings);
+
+	FTSTicker::GetCoreTicker().AddTicker(FTickerDelegate::CreateLambda([Message, WeakSettings](float)
+	{
+		TSharedPtr<bool> bSuppressInFuture = MakeShared<bool>(false);
+
+		TSharedRef<SWindow> DialogWindow = SNew(SWindow)
+			.Title(LOCTEXT("UnsafeModeProcessedTitle", "Non-Backwards-Compatible Data Loaded"))
+			.SizingRule(ESizingRule::Autosized)
+			.SupportsMaximize(false)
+			.SupportsMinimize(false);
+
+		DialogWindow->SetContent(
+			SNew(SVerticalBox)
+			+ SVerticalBox::Slot().AutoHeight().Padding(16.f, 16.f, 16.f, 8.f)
+			[
+				SNew(STextBlock).Text(Message).WrapTextAt(500.f)
+			]
+			+ SVerticalBox::Slot().AutoHeight().Padding(16.f, 4.f)
+			[
+				SNew(SCheckBox)
+				.IsChecked(ECheckBoxState::Checked)
+				.OnCheckStateChanged_Lambda([bSuppressInFuture](ECheckBoxState NewState)
+				{
+					*bSuppressInFuture = (NewState == ECheckBoxState::Unchecked);
+				})
+				[
+					SNew(STextBlock).Text(LOCTEXT("UnsafeModeShowWarning", "Show this warning in future"))
+				]
+			]
+			+ SVerticalBox::Slot().AutoHeight().HAlign(HAlign_Right).Padding(16.f, 8.f, 16.f, 16.f)
+			[
+				SNew(SHorizontalBox)
+				+ SHorizontalBox::Slot().AutoWidth().Padding(0.f, 0.f, 8.f, 0.f)
+				[
+					SNew(SButton)
+					.ButtonStyle(FAppStyle::Get(), "SimpleButton")
+					.ToolTipText(LOCTEXT("UnsafeModeCopyTooltip", "Copy the text in this message to the clipboard (CTRL+C)"))
+					.ContentPadding(2.f)
+					.OnClicked_Lambda([Message]() -> FReply
+					{
+						FPlatformApplicationMisc::ClipboardCopy(*Message.ToString());
+						return FReply::Handled();
+					})
+					.Content()
+					[
+						SNew(SImage)
+						.Image(FAppStyle::Get().GetBrush("Icons.Clipboard"))
+						.ColorAndOpacity(FSlateColor::UseForeground())
+					]
+				]
+				+ SHorizontalBox::Slot().AutoWidth()
+				[
+					SNew(SButton)
+					.Text(LOCTEXT("UnsafeModeOK", "OK"))
+					.OnClicked_Lambda([DialogWindowWeak = TWeakPtr<SWindow>(DialogWindow), bSuppressInFuture, WeakSettings]() -> FReply
+					{
+						if (*bSuppressInFuture)
+						{
+							if (UChaosVDGeneralSettings* CVDSettings = WeakSettings.Get())
+							{
+								CVDSettings->bShowUnsafeDataWarning = false;
+								CVDSettings->SaveConfig();
+							}
+						}
+						if (TSharedPtr<SWindow> DialogWindowPinned = DialogWindowWeak.Pin())
+						{
+							DialogWindowPinned->RequestDestroyWindow();
+						}
+						return FReply::Handled();
+					})
+				]
+			]
+		);
+
+		FSlateApplication::Get().AddModalWindow(DialogWindow, FSlateApplication::Get().GetActiveTopLevelWindow());
+		return false;
+	}));
+}
+
+void FChaosVDTraceProvider::ShowNativeSerializationDialog()
+{
+	// Only show in LoadAll mode — SerializeBinOnly strips native entries silently (notified via
+	// ShowSerializeBinOnlyStrippedDialog); SkipAll skips the processor entirely.
+	UChaosVDGeneralSettings* Settings = FChaosVDSettingsManager::Get().GetSettingsObject<UChaosVDGeneralSettings>();
+	if (!Settings || Settings->ParticleExtraDataLoadingMode != EChaosVDParticleExtraDataLoadingMode::LoadAll
+		|| !Settings->bShowNativeSerializationWarning)
+	{
+		return;
+	}
+
+	TMap<FName, FName> AllNativeTypes; // type path → channel id
+	for (const TPair<FStringView, TSharedPtr<FChaosVDDataProcessorBase>>& DataProcessor : RegisteredDataProcessors)
+	{
+		if (DataProcessor.Value)
+		{
+			DataProcessor.Value->GetPostLoadNativeTypesWithChannels(AllNativeTypes);
+		}
+	}
+
+	if (AllNativeTypes.IsEmpty())
+	{
+		return;
+	}
+
+	TSet<FString> FormattedTypes;
+	for (const TPair<FName, FName>& Pair : AllNativeTypes)
+	{
+		if (!Pair.Value.IsNone())
+		{
+			FormattedTypes.Add(FText::FormatOrdered(
+				LOCTEXT("NativeTypeWithChannel", "{0} (via {1} channel)"),
+				FText::FromName(Pair.Key),
+				FText::FromName(Pair.Value)).ToString());
+		}
+		else
+		{
+			FormattedTypes.Add(Pair.Key.ToString());
+		}
+	}
+
+	FText Message = FText::FormatOrdered(
+		LOCTEXT("NativeSerializationWarning",
+			"This recording contains particle extra data using native struct serialization "
+			"for the following types:\n\n{0}\n\n"
+			"Native serialization reads bytes at fixed offsets and displays exact recorded values "
+			"— but only when the struct layout matches the one at record time. If the layout has "
+			"changed, the result may be a crash or silently corrupt values with no warning.\n\n"
+			"SerializeBin entries (also present) are crash-safe but may show inaccurate values "
+			"if struct properties were added, removed, or renamed since recording.\n\n"
+			"Choose 'SerializeBin Only' to strip native entries and avoid crash risk. "
+			"Choose 'Load All' only if the recording and viewer binaries are known to match."),
+		FText::FromString(GenerateFormattedStringListFromSet(FormattedTypes)));
+
+	TWeakObjectPtr<UChaosVDGeneralSettings> WeakSettings(Settings);
+
+	FTSTicker::GetCoreTicker().AddTicker(FTickerDelegate::CreateLambda([Message, WeakSettings](float) -> bool
+	{
+		TSharedPtr<bool> bSuppressInFuture = MakeShared<bool>(false);
+
+		TSharedRef<SWindow> DialogWindow = SNew(SWindow)
+			.Title(LOCTEXT("NativeSerializationTitle", "Native-Serialized Extra Data Detected"))
+			.SizingRule(ESizingRule::Autosized)
+			.SupportsMaximize(false)
+			.SupportsMinimize(false);
+
+		DialogWindow->SetContent(
+			SNew(SVerticalBox)
+			+ SVerticalBox::Slot().AutoHeight().Padding(16.f, 16.f, 16.f, 8.f)
+			[
+				SNew(STextBlock).Text(Message).WrapTextAt(500.f)
+			]
+			+ SVerticalBox::Slot().AutoHeight().Padding(16.f, 4.f)
+			[
+				SNew(SCheckBox)
+				.IsChecked(ECheckBoxState::Checked)
+				.OnCheckStateChanged_Lambda([bSuppressInFuture](ECheckBoxState State)
+				{
+					*bSuppressInFuture = (State == ECheckBoxState::Unchecked);
+				})
+				[
+					SNew(STextBlock).Text(LOCTEXT("NativeSerializationShowWarning", "Show this warning in future"))
+				]
+			]
+			+ SVerticalBox::Slot().AutoHeight().HAlign(HAlign_Right).Padding(16.f, 8.f, 16.f, 16.f)
+			[
+				SNew(SHorizontalBox)
+				+ SHorizontalBox::Slot().AutoWidth().Padding(0.f, 0.f, 8.f, 0.f)
+				[
+					SNew(SButton)
+					.Text(LOCTEXT("NativeSerializationShowAll", "Load All"))
+					.ToolTipText(LOCTEXT("NativeSerializationShowAllTip",
+						"Display all extra data including native-serialized entries (may crash if struct layout changed)."))
+					.OnClicked_Lambda([DialogWindowWeak = TWeakPtr<SWindow>(DialogWindow), bSuppressInFuture, WeakSettings]() -> FReply
+					{
+						if (*bSuppressInFuture)
+						{
+							if (UChaosVDGeneralSettings* CVDSettings = WeakSettings.Get())
+							{
+								CVDSettings->bShowNativeSerializationWarning = false;
+								CVDSettings->SaveConfig();
+							}
+						}
+						if (TSharedPtr<SWindow> DialogWindowPinned = DialogWindowWeak.Pin())
+						{
+							DialogWindowPinned->RequestDestroyWindow();
+						}
+						return FReply::Handled();
+					})
+				]
+				+ SHorizontalBox::Slot().AutoWidth()
+				[
+					SNew(SButton)
+					.Text(LOCTEXT("NativeSerializationSkip", "SerializeBin Only"))
+					.ToolTipText(LOCTEXT("NativeSerializationSkipTip",
+						"Strip native-serialized entries; only property-tagged data is displayed (safe)."))
+					.OnClicked_Lambda([DialogWindowWeak = TWeakPtr<SWindow>(DialogWindow), bSuppressInFuture, WeakSettings]() -> FReply
+					{
+						if (UChaosVDGeneralSettings* CVDSettings = WeakSettings.Get())
+						{
+							CVDSettings->ParticleExtraDataLoadingMode = EChaosVDParticleExtraDataLoadingMode::SerializeBinOnly;
+							if (*bSuppressInFuture)
+							{
+								CVDSettings->bShowNativeSerializationWarning = false;
+							}
+							CVDSettings->SaveConfig();
+						}
+						if (TSharedPtr<SWindow> DialogWindowPinned = DialogWindowWeak.Pin())
+						{
+							DialogWindowPinned->RequestDestroyWindow();
+						}
+						return FReply::Handled();
+					})
+				]
+			]
+		);
+
+		FSlateApplication::Get().AddModalWindow(DialogWindow, FSlateApplication::Get().GetActiveTopLevelWindow());
+		return false;
+	}));
+}
+
+void FChaosVDTraceProvider::ShowSerializeBinOnlyStrippedDialog()
+{
+	UChaosVDGeneralSettings* Settings = FChaosVDSettingsManager::Get().GetSettingsObject<UChaosVDGeneralSettings>();
+	if (!Settings || Settings->ParticleExtraDataLoadingMode != EChaosVDParticleExtraDataLoadingMode::SerializeBinOnly
+		|| !Settings->bShowSafeModeSkippedWarning)
+	{
+		return;
+	}
+
+	TMap<FName, FName> AllNativeTypes;
+	for (const TPair<FStringView, TSharedPtr<FChaosVDDataProcessorBase>>& DataProcessor : RegisteredDataProcessors)
+	{
+		if (DataProcessor.Value)
+		{
+			DataProcessor.Value->GetPostLoadNativeTypesWithChannels(AllNativeTypes);
+		}
+	}
+
+	if (AllNativeTypes.IsEmpty())
+	{
+		return;
+	}
+
+	TSet<FString> FormattedTypes;
+	for (const TPair<FName, FName>& Pair : AllNativeTypes)
+	{
+		if (!Pair.Value.IsNone())
+		{
+			FormattedTypes.Add(FText::FormatOrdered(
+				LOCTEXT("NativeTypeWithChannel", "{0} (via {1} channel)"),
+				FText::FromName(Pair.Key),
+				FText::FromName(Pair.Value)).ToString());
+		}
+		else
+		{
+			FormattedTypes.Add(Pair.Key.ToString());
+		}
+	}
+
+	FText Message = FText::FormatOrdered(
+		LOCTEXT("SerializeBinOnlyStrippedMessage",
+			"The following particle extra data types were recorded with native struct serialization "
+			"and have been stripped because Particle Extra Data Loading Mode is set to "
+			"SerializeBin Only:\n\n{0}\n\n"
+			"These entries are not shown. Native serialization reads bytes at fixed offsets — "
+			"if the struct layout has changed since recording, loading them may crash or display "
+			"corrupt values.\n\n"
+			"To load them anyway, set Particle Extra Data Loading Mode to 'Load All' in CVD "
+			"General Settings. Only do this when the recording and viewer binaries are known to match."),
+		FText::FromString(GenerateFormattedStringListFromSet(FormattedTypes)));
+
+	TWeakObjectPtr<UChaosVDGeneralSettings> WeakSettings(Settings);
+
+	FTSTicker::GetCoreTicker().AddTicker(FTickerDelegate::CreateLambda([Message, WeakSettings](float) -> bool
+	{
+		TSharedPtr<bool> bSuppressInFuture = MakeShared<bool>(false);
+
+		TSharedRef<SWindow> DialogWindow = SNew(SWindow)
+			.Title(LOCTEXT("SerializeBinOnlyStrippedTitle", "Native-Serialized Extra Data Stripped"))
+			.SizingRule(ESizingRule::Autosized)
+			.SupportsMaximize(false)
+			.SupportsMinimize(false);
+
+		DialogWindow->SetContent(
+			SNew(SVerticalBox)
+			+ SVerticalBox::Slot().AutoHeight().Padding(16.f, 16.f, 16.f, 8.f)
+			[
+				SNew(STextBlock).Text(Message).WrapTextAt(500.f)
+			]
+			+ SVerticalBox::Slot().AutoHeight().Padding(16.f, 4.f)
+			[
+				SNew(SCheckBox)
+				.IsChecked(ECheckBoxState::Checked)
+				.OnCheckStateChanged_Lambda([bSuppressInFuture](ECheckBoxState NewState)
+				{
+					*bSuppressInFuture = (NewState == ECheckBoxState::Unchecked);
+				})
+				[
+					SNew(STextBlock).Text(LOCTEXT("SerializeBinOnlyStrippedShowWarning", "Show this warning in future"))
+				]
+			]
+			+ SVerticalBox::Slot().AutoHeight().HAlign(HAlign_Right).Padding(16.f, 8.f, 16.f, 16.f)
+			[
+				SNew(SHorizontalBox)
+				+ SHorizontalBox::Slot().AutoWidth().Padding(0.f, 0.f, 8.f, 0.f)
+				[
+					SNew(SButton)
+					.ButtonStyle(FAppStyle::Get(), "SimpleButton")
+					.ToolTipText(LOCTEXT("SerializeBinOnlyStrippedCopyTooltip", "Copy the text in this message to the clipboard (CTRL+C)"))
+					.ContentPadding(2.f)
+					.OnClicked_Lambda([Message]() -> FReply
+					{
+						FPlatformApplicationMisc::ClipboardCopy(*Message.ToString());
+						return FReply::Handled();
+					})
+					.Content()
+					[
+						SNew(SImage)
+						.Image(FAppStyle::Get().GetBrush("Icons.Clipboard"))
+						.ColorAndOpacity(FSlateColor::UseForeground())
+					]
+				]
+				+ SHorizontalBox::Slot().AutoWidth()
+				[
+					SNew(SButton)
+					.Text(LOCTEXT("SerializeBinOnlyStrippedOK", "OK"))
+					.OnClicked_Lambda([DialogWindowWeak = TWeakPtr<SWindow>(DialogWindow), bSuppressInFuture, WeakSettings]() -> FReply
+					{
+						if (*bSuppressInFuture)
+						{
+							if (UChaosVDGeneralSettings* CVDSettings = WeakSettings.Get())
+							{
+								CVDSettings->bShowSafeModeSkippedWarning = false;
+								CVDSettings->SaveConfig();
+							}
+						}
+						if (TSharedPtr<SWindow> DialogWindowPinned = DialogWindowWeak.Pin())
+						{
+							DialogWindowPinned->RequestDestroyWindow();
+						}
+						return FReply::Handled();
+					})
+				]
+			]
+		);
+
+		FSlateApplication::Get().AddModalWindow(DialogWindow, FSlateApplication::Get().GetActiveTopLevelWindow());
+		return false;
+	}));
+}
+
 FChaosVDFrameStageData* FChaosVDTraceProvider::GetCurrentSolverStageDataForCurrentFrame(int32 SolverID, EChaosVDSolverStageAccessorFlags Flags)
 {
 	auto CreateInBetweenSolverStage = [](FChaosVDSolverFrameData& InFrameData)
diff --git a/Engine/Plugins/ChaosVD/Source/ChaosVD/Public/ChaosVDScene.h b/Engine/Plugins/ChaosVD/Source/ChaosVD/Public/ChaosVDScene.h
--- a/Engine/Plugins/ChaosVD/Source/ChaosVD/Public/ChaosVDScene.h
+++ b/Engine/Plugins/ChaosVD/Source/ChaosVD/Public/ChaosVDScene.h
@@ -49,6 +49,18 @@ DECLARE_MULTICAST_DELEGATE_OneParam(FChaosVDSolverInfoActorCreatedDelegate, ACha
 
 DECLARE_MULTICAST_DELEGATE_TwoParams(FChaosVDSolverVisibilityChangedDelegate, int32 SolverID, bool bNewVisibility)
 
+class IDetailLayoutBuilder;
+
+/** Fired by FChaosVDSceneParticleCustomization when building details for a selected particle.
+ *  Subscribers (e.g. UChaosVDParticleExtraDataComponent) can append additional categories and
+ *  struct entries to the details panel via the provided IDetailLayoutBuilder. */
+DECLARE_MULTICAST_DELEGATE_ThreeParams(FChaosVDParticleDetailsExtraDataDelegate, int32 /*SolverID*/, int32 /*ParticleID*/, IDetailLayoutBuilder& /*DetailBuilder*/)
+
+/** Lightweight delegate broadcast on every frame change for an already-selected particle.
+ *  Subscribers update their cached struct scope memory in-place; they set bOutNeedsFullRebuild=true
+ *  if the data layout has changed and a full IDetailLayoutBuilder refresh is required. */
+DECLARE_MULTICAST_DELEGATE_ThreeParams(FChaosVDParticleDetailsExtraDataRefreshDelegate, int32 /*SolverID*/, int32 /*ParticleID*/, bool& /*bOutNeedsFullRebuild*/)
+
 UENUM()
 enum class EChaosVDSceneCleanUpOptions
 {
@@ -82,6 +94,20 @@ public:
 		return SceneUpdatedDelegate;
 	}
 
+	/** Fired when the particle details panel is being built for a selected particle.
+	 *  Subscribers can append additional categories/structs to the panel via IDetailLayoutBuilder. */
+	FChaosVDParticleDetailsExtraDataDelegate& OnParticleDetailsExtraDataRequested()
+	{
+		return ParticleDetailsExtraDataDelegate;
+	}
+
+	/** Fired on every frame change while a particle remains selected.
+	 *  Subscribers should update their cached struct scope memory in-place (no layout rebuild). */
+	FChaosVDParticleDetailsExtraDataRefreshDelegate& OnParticleDetailsExtraDataRefreshRequested()
+	{
+		return ParticleDetailsExtraDataRefreshDelegate;
+	}
+
 	/** Updates, Adds and Remove actors to match the provided Solver Stage Data */
 	void UpdateFromRecordedSolverStageData(const int32 SolverID, const FChaosVDFrameStageData& InRecordedStepData, const FChaosVDSolverFrameData& InFrameData);
 
@@ -297,6 +323,8 @@ private:
 	TObjectPtr<UWorld> PhysicsVDWorld = nullptr;
 
 	FChaosVDSceneUpdatedDelegate SceneUpdatedDelegate;
+	FChaosVDParticleDetailsExtraDataDelegate ParticleDetailsExtraDataDelegate;
+	FChaosVDParticleDetailsExtraDataRefreshDelegate ParticleDetailsExtraDataRefreshDelegate;
 
 	TSharedPtr<FChaosVDGeometryBuilder> GeometryGenerator;
 
diff --git a/Engine/Plugins/ChaosVD/Source/ChaosVD/Public/Components/ChaosVDParticleExtraDataComponent.h b/Engine/Plugins/ChaosVD/Source/ChaosVD/Public/Components/ChaosVDParticleExtraDataComponent.h
new file mode 100644
--- /dev/null
+++ b/Engine/Plugins/ChaosVD/Source/ChaosVD/Public/Components/ChaosVDParticleExtraDataComponent.h
@@ -0,0 +1,64 @@
+// Copyright Epic Games, Inc. All Rights Reserved.
+
+#pragma once
+
+#include "Components/ChaosVDSolverDataComponent.h"
+#include "DataWrappers/ChaosVDParticleExtraDataContainer.h"
+#include "ChaosVDParticleExtraDataComponent.generated.h"
+
+class IDetailLayoutBuilder;
+class FChaosVDScene;
+
+/** Holds a deserialized struct scope together with enough metadata to match it
+ *  against fresh frame data for in-place memory updates. */
+struct FChaosVDExtraDataScopeEntry
+{
+	FName CategoryName;
+	FName StructTypePath;
+	TSharedPtr<FStructOnScope> Scope;
+};
+
+/**
+ * Solver data component that reads FChaosVDParticleExtraDataContainer from the current
+ * solver frame and injects the recorded structs into the CVD particle details panel.
+ *
+ * Added to AChaosVDSolverInfoActor as a default subobject. It binds to
+ * FChaosVDScene::OnParticleDetailsExtraDataRequested and for each selected particle
+ * adds the traced categories to the IDetailLayoutBuilder, so the extra data appears
+ * alongside "Particle Data" in the details panel (and in "Show Data in New Panel").
+ *
+ * On subsequent frame changes (OnParticleDetailsExtraDataRefreshRequested), it updates
+ * the cached FStructOnScope memory in-place so Slate picks up the new values without
+ * requiring a full layout rebuild.
+ */
+UCLASS(MinimalAPI)
+class UChaosVDParticleExtraDataComponent : public UChaosVDSolverDataComponent
+{
+	GENERATED_BODY()
+
+public:
+
+	CHAOSVD_API virtual void ClearData() override;
+	CHAOSVD_API virtual void UpdateFromSolverFrameData(const FChaosVDSolverFrameData& InSolverFrameData) override;
+	CHAOSVD_API virtual void SetScene(const TWeakPtr<FChaosVDScene>& InSceneWeakPtr) override;
+
+private:
+
+	void HandleParticleDetailsExtraData(int32 InSolverID, int32 InParticleID, IDetailLayoutBuilder& DetailBuilder);
+	void HandleParticleDetailsRefresh(int32 InSolverID, int32 InParticleID, bool& bOutNeedsFullRebuild);
+	void UnbindFromScene(const TSharedPtr<FChaosVDScene>& Scene);
+
+	/** Extra data container for the current frame, looked up once per frame update. */
+	TSharedPtr<FChaosVDParticleExtraDataContainer> CurrentFrameContainer;
+
+	/** Stable FStructOnScope instances added to the details panel for the currently selected particle.
+	 *  Kept alive so we can update their memory in-place on frame changes without rebuilding the layout. */
+	TArray<FChaosVDExtraDataScopeEntry> CachedScopeEntries;
+
+	/** Solver/particle for which CachedScopeEntries were built. */
+	int32 CachedSolverID = INDEX_NONE;
+	int32 CachedParticleID = INDEX_NONE;
+
+	FDelegateHandle ExtraDataDelegateHandle;
+	FDelegateHandle RefreshDelegateHandle;
+};
diff --git a/Engine/Plugins/ChaosVD/Source/ChaosVD/Public/Trace/ChaosVDTraceProvider.h b/Engine/Plugins/ChaosVD/Source/ChaosVD/Public/Trace/ChaosVDTraceProvider.h
--- a/Engine/Plugins/ChaosVD/Source/ChaosVD/Public/Trace/ChaosVDTraceProvider.h
+++ b/Engine/Plugins/ChaosVD/Source/ChaosVD/Public/Trace/ChaosVDTraceProvider.h
@@ -111,7 +111,21 @@ private:
 	}
 	
 	void HandleAnalysisComplete();
-	
+
+	/** Shows a Slate dialog listing data types skipped because Safe Loading Mode was active. */
+	void ShowSafeModeSkippedDialog();
+
+	/** Shows a Slate dialog warning that non-backwards-compatible types were loaded with Safe Loading Mode disabled. */
+	void ShowUnsafeDataDialog();
+
+	/** Shows a Slate dialog warning that natively-serialized particle extra data entries are present,
+	 *  and lets the user choose whether to display them or skip them. */
+	void ShowNativeSerializationDialog();
+
+	/** Shows a Slate dialog informing the user that natively-serialized particle extra data entries
+	 *  were stripped because Particle Extra Data Loading Mode is set to SerializeBin Only. */
+	void ShowSerializeBinOnlyStrippedDialog();
+
 	FBinaryDataContainer& FindOrAddUnprocessedData(const int32 DataID);
 	void RemoveUnprocessedData(const int32 DataID);
 
@@ -175,6 +189,18 @@ public:
 		return TypesFailedToSerialize;
 	}
 
+	/** Types whose data was skipped entirely because Safe Loading Mode was enabled and they are not backwards compatible. */
+	const TSet<FString>& GetTypesSkippedInSafeMode() const
+	{
+		return TypesSkippedInSafeMode;
+	}
+
+	/** Types that were processed despite not being backwards compatible (Safe Loading Mode was off). */
+	const TSet<FString>& GetUnsafeTypesProcessed() const
+	{
+		return UnsafeTypesProcessed;
+	}
+
 	void AddParticleMetadata(uint64 MetadaId, const TSharedPtr<FChaosVDParticleMetadata>& InMetadata);
 	TSharedPtr<FChaosVDParticleMetadata> GetParticleMetadata(uint64 MetadataId);
 
@@ -232,6 +258,8 @@ private:
 
 	TSet<FString> MissingDataProcessors;
 	TSet<FString> TypesFailedToSerialize;
+	TSet<FString> TypesSkippedInSafeMode;
+	TSet<FString> UnsafeTypesProcessed;
 
 	int64 DataProcessedSoFarCounter = 0;
 
