Chaos Visual Debugger: Modifying the language around extra particle data, users won't know what SerailizeBin or SerializeItem are

#Synced-CL 52655663

diff --git a/Engine/Plugins/ChaosVD/Source/ChaosVD/Private/Trace/ChaosVDTraceProvider.cpp b/Engine/Plugins/ChaosVD/Source/ChaosVD/Private/Trace/ChaosVDTraceProvider.cpp
--- a/Engine/Plugins/ChaosVD/Source/ChaosVD/Private/Trace/ChaosVDTraceProvider.cpp
+++ b/Engine/Plugins/ChaosVD/Source/ChaosVD/Private/Trace/ChaosVDTraceProvider.cpp
@@ -759,11 +759,11 @@ void FChaosVDTraceProvider::ShowNativeSerializationDialog()
 			"This recording contains particle extra data using native struct serialization "
 			"for the following types:\n\n{0}\n\n"
 			"Native serialization reads bytes at fixed offsets and displays exact recorded values "
-			"— but only when the struct layout matches the one at record time. If the layout has "
+			"-- but only when the struct layout matches the one at record time. If the layout has "
 			"changed, the result may be a crash or silently corrupt values with no warning.\n\n"
-			"SerializeBin entries (also present) are crash-safe but may show inaccurate values "
+			"Reflection-based entries (also present) are crash-safe but may show inaccurate values "
 			"if struct properties were added, removed, or renamed since recording.\n\n"
-			"Choose 'SerializeBin Only' to strip native entries and avoid crash risk. "
+			"Choose 'Reflection-Based Only' to strip native entries and avoid crash risk. "
 			"Choose 'Load All' only if the recording and viewer binaries are known to match."),
 		FText::FromString(GenerateFormattedStringListFromSet(FormattedTypes)));
 
@@ -826,9 +826,9 @@ void FChaosVDTraceProvider::ShowNativeSerializationDialog()
 				+ SHorizontalBox::Slot().AutoWidth()
 				[
 					SNew(SButton)
-					.Text(LOCTEXT("NativeSerializationSkip", "SerializeBin Only"))
+					.Text(LOCTEXT("NativeSerializationSkip", "Reflection-Based Only"))
 					.ToolTipText(LOCTEXT("NativeSerializationSkipTip",
-						"Strip native-serialized entries; only property-tagged data is displayed (safe)."))
+						"Strip natively-serialized entries; only reflection-based (property-tagged) data is displayed (safe)."))
 					.OnClicked_Lambda([DialogWindowWeak = TWeakPtr<SWindow>(DialogWindow), bSuppressInFuture, WeakSettings]() -> FReply
 					{
 						if (UChaosVDGeneralSettings* CVDSettings = WeakSettings.Get())
@@ -898,8 +898,8 @@ void FChaosVDTraceProvider::ShowSerializeBinOnlyStrippedDialog()
 		LOCTEXT("SerializeBinOnlyStrippedMessage",
 			"The following particle extra data types were recorded with native struct serialization "
 			"and have been stripped because Particle Extra Data Loading Mode is set to "
-			"SerializeBin Only:\n\n{0}\n\n"
-			"These entries are not shown. Native serialization reads bytes at fixed offsets — "
+			"Reflection-Based Only:\n\n{0}\n\n"
+			"These entries are not shown. Native serialization reads bytes at fixed offsets -- "
 			"if the struct layout has changed since recording, loading them may crash or display "
 			"corrupt values.\n\n"
 			"To load them anyway, set Particle Extra Data Loading Mode to 'Load All' in CVD "
diff --git a/Engine/Plugins/ChaosVD/Source/ChaosVD/Public/Settings/ChaosVDGeneralSettings.h b/Engine/Plugins/ChaosVD/Source/ChaosVD/Public/Settings/ChaosVDGeneralSettings.h
--- a/Engine/Plugins/ChaosVD/Source/ChaosVD/Public/Settings/ChaosVDGeneralSettings.h
+++ b/Engine/Plugins/ChaosVD/Source/ChaosVD/Public/Settings/ChaosVDGeneralSettings.h
@@ -8,13 +8,13 @@
 
 /** Controls how FChaosVDParticleExtraDataProcessor handles particle extra data at load time.
  *
- *  SerializeBin entries use property-tagged serialization: properties are matched by name at load
- *  time against the *current* struct layout, so adds/removes/renames are handled gracefully
+ *  Reflection-based entries use property-tagged serialization: properties are matched by name at
+ *  load time against the *current* struct layout, so adds/removes/renames are handled gracefully
  *  (missing props show defaults; added props are lost). No crash risk.
  *
- *  NativeSerialization entries invoke the struct's native Serialize function, which reads bytes
+ *  Natively-serialized entries invoke the struct's native Serialize function, which reads bytes
  *  at fixed offsets. When the struct layout matches, values are exact. When it does not match,
- *  the result is a crash or silently corrupt values — there is no graceful degradation.
+ *  the result is a crash or silently corrupt values -- there is no graceful degradation.
  */
 UENUM()
 enum class EChaosVDParticleExtraDataLoadingMode : uint8
@@ -24,13 +24,13 @@ enum class EChaosVDParticleExtraDataLoadingMode : uint8
 	 *  have changed since recording. */
 	SkipAll          UMETA(DisplayName = "Skip All (Safe)"),
 
-	/** Load only SerializeBin entries; NativeSerialization entries are stripped in the processor
+	/** Load only reflection-based entries; natively-serialized entries are stripped in the processor
 	 *  before reaching memory. No crash risk, but displayed values may be inaccurate if struct
 	 *  definitions changed since recording. */
-	SerializeBinOnly UMETA(DisplayName = "SerializeBin Only"),
+	SerializeBinOnly UMETA(DisplayName = "Reflection-Based Only"),
 
-	/** Load all entries including NativeSerialization. Displays exact recorded values when the
-	 *  struct layout matches the recording. If the layout has changed, NativeSerialization entries
+	/** Load all entries including natively-serialized ones. Displays exact recorded values when the
+	 *  struct layout matches the recording. If the layout has changed, natively-serialized entries
 	 *  may crash or display corrupt values. Only use when the recording and viewer binaries are
 	 *  known to match. */
 	LoadAll          UMETA(DisplayName = "Load All (Unsafe)"),
@@ -125,14 +125,14 @@ PRAGMA_ENABLE_DEPRECATION_WARNINGS
 	 *   that guarantees neither crashes nor silently incorrect values, regardless of how much struct
 	 *   definitions have changed since recording.
 	 *
-	 * SerializeBinOnly (default): only property-tagged (SerializeBin) entries are loaded;
-	 *   NativeSerialization entries are filtered out in the processor and never stored in memory.
+	 * Reflection-Based Only (default): only reflection-based (property-tagged) entries are loaded;
+	 *   natively-serialized entries are filtered out in the processor and never stored in memory.
 	 *   No crash risk, but values may be inaccurate if struct properties were added, removed, or
 	 *   renamed since recording.
 	 *
-	 * LoadAll: all entries including NativeSerialization are loaded. Displays exact recorded values
-	 *   when the viewer binary matches the recording binary. If struct layouts have changed,
-	 *   NativeSerialization entries may crash or show corrupt values. Only use when recording and
+	 * LoadAll: all entries including natively-serialized ones are loaded. Displays exact recorded
+	 *   values when the viewer binary matches the recording binary. If struct layouts have changed,
+	 *   natively-serialized entries may crash or show corrupt values. Only use when recording and
 	 *   viewer are known to match.
 	 *
 	 * Note: bSafeLoadingMode applies to any other non-backwards-compatible processors; this enum
